motion matching - reclaiming assets memory via WITH_EDITORONLY_DATA

diff --git a/Engine/Plugins/Experimental/Animation/PoseSearch/Source/Editor/Private/PoseSearchDatabaseAssetListItem.cpp b/Engine/Plugins/Experimental/Animation/PoseSearch/Source/Editor/Private/PoseSearchDatabaseAssetListItem.cpp
--- a/Engine/Plugins/Experimental/Animation/PoseSearch/Source/Editor/Private/PoseSearchDatabaseAssetListItem.cpp
+++ b/Engine/Plugins/Experimental/Animation/PoseSearch/Source/Editor/Private/PoseSearchDatabaseAssetListItem.cpp
@@ -383,7 +383,7 @@ namespace UE::PoseSearch
 		TSharedPtr<FDatabaseAssetTreeNode> TreeNodePtr = WeakAssetTreeNode.Pin();
 		if (const FSearchIndexAsset* SelectedIndexAsset = ViewModelPtr->GetSelectedActorIndexAsset())
 		{
-			if (TreeNodePtr->SourceAssetIdx == SelectedIndexAsset->SourceAssetIdx)
+			if (TreeNodePtr->SourceAssetIdx == SelectedIndexAsset->GetSourceAssetIdx())
 			{
 				return EVisibility::Visible;
 			}
diff --git a/Engine/Plugins/Experimental/Animation/PoseSearch/Source/Editor/Private/PoseSearchDatabaseViewModel.cpp b/Engine/Plugins/Experimental/Animation/PoseSearch/Source/Editor/Private/PoseSearchDatabaseViewModel.cpp
--- a/Engine/Plugins/Experimental/Animation/PoseSearch/Source/Editor/Private/PoseSearchDatabaseViewModel.cpp
+++ b/Engine/Plugins/Experimental/Animation/PoseSearch/Source/Editor/Private/PoseSearchDatabaseViewModel.cpp
@@ -129,7 +129,7 @@ namespace UE::PoseSearch
 		const FSearchIndex& SearchIndex = PoseSearchDatabase->GetSearchIndex();
 		const FSearchIndexAsset& IndexAsset = SearchIndex.Assets[IndexAssetIndex];
 
-		const FPoseSearchDatabaseAnimationAssetBase* DatabaseAnimationAsset = PoseSearchDatabase->GetAnimationAssetBase(IndexAsset.SourceAssetIdx);
+		const FPoseSearchDatabaseAnimationAssetBase* DatabaseAnimationAsset = PoseSearchDatabase->GetAnimationAssetBase(IndexAsset.GetSourceAssetIdx());
 		UAnimationAsset* PreviewAsset = DatabaseAnimationAsset->GetAnimationAsset();
 		if (!PreviewAsset)
 		{
@@ -139,7 +139,7 @@ namespace UE::PoseSearch
 		const FInstancedStruct& DatabaseAsset = PoseSearchDatabase->GetAnimationAssetStruct(IndexAsset);
 		const FPoseSearchDatabaseAnimationAssetBase* DatabaseAnimationAssetBase = DatabaseAsset.GetPtr<FPoseSearchDatabaseAnimationAssetBase>();
 		check(DatabaseAnimationAssetBase);
-		PreviewActor.Sampler.Init(DatabaseAnimationAssetBase->GetAnimationAsset(), IndexAsset.BlendParameters);
+		PreviewActor.Sampler.Init(DatabaseAnimationAssetBase->GetAnimationAsset(), IndexAsset.GetBlendParameters());
 		PreviewActor.IndexAssetIndex = IndexAssetIndex;
 		PreviewActor.CurrentPoseIndex = INDEX_NONE;
 		PreviewActor.PlayTimeOffset = PoseIdxForTimeOffset < 0 ? 0.f : PoseSearchDatabase->GetRealAssetTime(PoseIdxForTimeOffset);
@@ -162,9 +162,9 @@ namespace UE::PoseSearch
 		Mesh->EnablePreview(true, PreviewAsset);
 		
 		AnimInstance->SetAnimationAsset(PreviewAsset, false, 0.0f);
-		AnimInstance->SetBlendSpacePosition(IndexAsset.BlendParameters);
+		AnimInstance->SetBlendSpacePosition(IndexAsset.GetBlendParameters());
 		
-		if (IndexAsset.bMirrored && PoseSearchDatabase->Schema)
+		if (IndexAsset.IsMirrored() && PoseSearchDatabase->Schema)
 		{
 			AnimInstance->SetMirrorDataTable(PoseSearchDatabase->Schema->MirrorDataTable);
 		}
@@ -243,7 +243,7 @@ namespace UE::PoseSearch
 			// SetPosition is in [0..1] range for blendspaces
 			AnimInstance->SetPosition(PreviewActor.Sampler.ToNormalizedTime(PreviewActor.CurrentTime));
 			AnimInstance->SetPlayRate(0.f);
-			AnimInstance->SetBlendSpacePosition(IndexAsset.BlendParameters);
+			AnimInstance->SetBlendSpacePosition(IndexAsset.GetBlendParameters());
 
 			PreviewActor.QuantizedTimeRootTransform = PreviewActor.Sampler.ExtractRootTransform(QuantizedTime);
 			FTransform RootMotion = PreviewActor.Sampler.ExtractRootTransform(PreviewActor.CurrentTime);
@@ -439,7 +439,7 @@ namespace UE::PoseSearch
 						MaxPreviewPlayLength = FMath::Max(MaxPreviewPlayLength, IndexAsset.GetLastSampleTime(PoseSearchDatabase->Schema->SampleRate) - PreviewActor.PlayTimeOffset);
 						MinPreviewPlayLength = FMath::Min(MinPreviewPlayLength, IndexAsset.GetFirstSampleTime(PoseSearchDatabase->Schema->SampleRate) - PreviewActor.PlayTimeOffset);
 						PreviewActors.Add(PreviewActor);
-						SelectedSourceAssetIdx = IndexAsset.SourceAssetIdx;
+						SelectedSourceAssetIdx = IndexAsset.GetSourceAssetIdx();
 					}
 				}
 			}
@@ -477,9 +477,9 @@ namespace UE::PoseSearch
 			for (int32 IndexAssetIndex = 0; IndexAssetIndex < SearchIndex.Assets.Num(); ++IndexAssetIndex)
 			{
 				const FSearchIndexAsset& IndexAsset = SearchIndex.Assets[IndexAssetIndex];
-				if (AnimationPreviewMode == EAnimationPreviewMode::OriginalAndMirrored || !IndexAsset.bMirrored)
+				if (AnimationPreviewMode == EAnimationPreviewMode::OriginalAndMirrored || !IndexAsset.IsMirrored())
 				{
-					if (const int32* SelectedNodesIndex = AssociatedAssetIndices.Find(IndexAsset.SourceAssetIdx))
+					if (const int32* SelectedNodesIndex = AssociatedAssetIndices.Find(IndexAsset.GetSourceAssetIdx()))
 					{
 						FDatabasePreviewActor PreviewActor = SpawnPreviewActor(IndexAssetIndex);
 						if (PreviewActor.IsValid())
@@ -570,10 +570,10 @@ namespace UE::PoseSearch
 				if (PreviewActor.IndexAssetIndex >= 0 && PreviewActor.IndexAssetIndex < SearchIndex.Assets.Num() && PreviewActor.Sampler.IsInitialized())
 				{
 					const FSearchIndexAsset& IndexAsset = SearchIndex.Assets[PreviewActor.IndexAssetIndex];
-					if (IndexAsset.SourceAssetIdx == SourceAssetIdx)
+					if (IndexAsset.GetSourceAssetIdx() == SourceAssetIdx)
 					{
 						CurrentPlayTime = PreviewActor.Sampler.ToNormalizedTime(PlayTime + IndexAsset.GetFirstSampleTime(PoseSearchDatabase->Schema->SampleRate) + PreviewActor.PlayTimeOffset);
-						BlendParameters = IndexAsset.BlendParameters;
+						BlendParameters = IndexAsset.GetBlendParameters();
 						return true;
 					}
 				}
@@ -581,10 +581,10 @@ namespace UE::PoseSearch
 
 			for (const FSearchIndexAsset& IndexAsset : SearchIndex.Assets)
 			{
-				if (IndexAsset.SourceAssetIdx == SourceAssetIdx)
+				if (IndexAsset.GetSourceAssetIdx() == SourceAssetIdx)
 				{
 					CurrentPlayTime = PlayTime + IndexAsset.GetFirstSampleTime(PoseSearchDatabase->Schema->SampleRate);
-					BlendParameters = IndexAsset.BlendParameters;
+					BlendParameters = IndexAsset.GetBlendParameters();
 
 					const bool bIsBlendSpace = PoseSearchDatabase->GetAnimationAssetStruct(IndexAsset).GetPtr<FPoseSearchDatabaseBlendSpace>() != nullptr;
 					if (bIsBlendSpace && !FMath::IsNearlyEqual(MaxPreviewPlayLength, MinPreviewPlayLength))
diff --git a/Engine/Plugins/Experimental/Animation/PoseSearch/Source/Editor/Private/PoseSearchDebuggerDatabaseView.cpp b/Engine/Plugins/Experimental/Animation/PoseSearch/Source/Editor/Private/PoseSearchDebuggerDatabaseView.cpp
--- a/Engine/Plugins/Experimental/Animation/PoseSearch/Source/Editor/Private/PoseSearchDebuggerDatabaseView.cpp
+++ b/Engine/Plugins/Experimental/Animation/PoseSearch/Source/Editor/Private/PoseSearchDebuggerDatabaseView.cpp
@@ -166,9 +166,9 @@ static void AddUnfilteredDatabaseRow(const UPoseSearchDatabase* Database,
 
 		Row->PoseIdx = DbPoseIdx;
 		Row->PoseCandidateFlags = PoseCandidateFlags;
-		Row->DbAssetIdx = SearchIndexAsset->SourceAssetIdx;
+		Row->DbAssetIdx = SearchIndexAsset->GetSourceAssetIdx();
 		Row->AssetTime = Time;
-		Row->bMirrored = SearchIndexAsset->bMirrored;
+		Row->bMirrored = SearchIndexAsset->IsMirrored();
 
 		Row->CostVector.SetNum(Database->Schema->SchemaCardinality);
 		const TArray<float> PoseValues = SearchIndex.GetPoseValuesSafe(DbPoseIdx);
@@ -217,7 +217,7 @@ static void AddUnfilteredDatabaseRow(const UPoseSearchDatabase* Database,
 			Row->AssetName = DatabaseAsset->GetName();
 			Row->AssetPath = DatabaseAsset->GetAnimationAsset() ? DatabaseAsset->GetAnimationAsset()->GetPathName() : "";
 			Row->bLooping = DatabaseAsset->IsLooping();
-			Row->BlendParameters = SearchIndexAsset->BlendParameters;
+			Row->BlendParameters = SearchIndexAsset->GetBlendParameters();
 			Row->AnimFrame = 0;
 			Row->AnimPercentage = 0.0f;
 
diff --git a/Engine/Plugins/Experimental/Animation/PoseSearch/Source/Editor/Private/PoseSearchDebuggerViewModel.cpp b/Engine/Plugins/Experimental/Animation/PoseSearch/Source/Editor/Private/PoseSearchDebuggerViewModel.cpp
--- a/Engine/Plugins/Experimental/Animation/PoseSearch/Source/Editor/Private/PoseSearchDebuggerViewModel.cpp
+++ b/Engine/Plugins/Experimental/Animation/PoseSearch/Source/Editor/Private/PoseSearchDebuggerViewModel.cpp
@@ -66,10 +66,10 @@ void FDebuggerViewModel::ShowSelectedSkeleton(const UPoseSearchDatabase* Databas
 		bSelecting = true;
 
 		Skeletons[SelectedPose].Time = Time;
-		Skeletons[SelectedPose].bMirrored = IndexAsset->bMirrored;
+		Skeletons[SelectedPose].bMirrored = IndexAsset->IsMirrored();
 		Skeletons[SelectedPose].SourceDatabase = Database;
-		Skeletons[SelectedPose].AssetIdx = IndexAsset->SourceAssetIdx;
-		Skeletons[SelectedPose].BlendParameters = IndexAsset->BlendParameters;
+		Skeletons[SelectedPose].AssetIdx = IndexAsset->GetSourceAssetIdx();
+		Skeletons[SelectedPose].BlendParameters = IndexAsset->GetBlendParameters();
 	}
 }
 
@@ -145,10 +145,10 @@ void FDebuggerViewModel::OnUpdateNodeSelection(int32 InNodeId)
 			int32 CurrentPoseIdx = ActiveMotionMatchingState->GetCurrentDatabasePoseIndex();
 			if (const FSearchIndexAsset* IndexAsset = CurrentSearchIndex.GetAssetForPoseSafe(CurrentPoseIdx))
 			{
-				Skeletons[Asset].bMirrored = IndexAsset->bMirrored;
+				Skeletons[Asset].bMirrored = IndexAsset->IsMirrored();
 				Skeletons[Asset].SourceDatabase = CurrentDatabase;
-				Skeletons[Asset].AssetIdx = IndexAsset->SourceAssetIdx;
-				Skeletons[Asset].BlendParameters = IndexAsset->BlendParameters;
+				Skeletons[Asset].AssetIdx = IndexAsset->GetSourceAssetIdx();
+				Skeletons[Asset].BlendParameters = IndexAsset->GetBlendParameters();
 			}
 		}
 	}
diff --git a/Engine/Plugins/Experimental/Animation/PoseSearch/Source/Runtime/Private/AnimNode_MotionMatching.cpp b/Engine/Plugins/Experimental/Animation/PoseSearch/Source/Runtime/Private/AnimNode_MotionMatching.cpp
--- a/Engine/Plugins/Experimental/Animation/PoseSearch/Source/Runtime/Private/AnimNode_MotionMatching.cpp
+++ b/Engine/Plugins/Experimental/Animation/PoseSearch/Source/Runtime/Private/AnimNode_MotionMatching.cpp
@@ -196,9 +196,9 @@ void FAnimNode_MotionMatching::UpdateAssetPlayer(const FAnimationUpdateContext&
 				// root bone blending needs to be immediate if MM node controls the offset between mesh component and root bone
 				const float RootBoneBlendTime = YawFromAnimationBlendRate < 0.f ? BlendTime : 0.f;
 				FAnimNode_BlendStack_Standalone::BlendTo(Context, DatabaseAsset->GetAnimationAsset(), MotionMatchingState.CurrentSearchResult.AssetTime,
-					DatabaseAsset->IsLooping(), SearchIndexAsset->bMirrored, CurrentResultDatabase->Schema->MirrorDataTable.Get(),
+					SearchIndexAsset->IsLooping(), SearchIndexAsset->IsMirrored(), CurrentResultDatabase->Schema->MirrorDataTable.Get(),
 					BlendTime, RootBoneBlendTime, MaxTimeBeforeFreezingInnerBlends,
-					BlendProfile, BlendOption, bUseInertialBlend, SearchIndexAsset->BlendParameters, MotionMatchingState.WantedPlayRate);
+					BlendProfile, BlendOption, bUseInertialBlend, SearchIndexAsset->GetBlendParameters(), MotionMatchingState.WantedPlayRate);
 			}
 		}
 	}
diff --git a/Engine/Plugins/Experimental/Animation/PoseSearch/Source/Runtime/Private/AnimNode_PoseSearchHistoryCollector.cpp b/Engine/Plugins/Experimental/Animation/PoseSearch/Source/Runtime/Private/AnimNode_PoseSearchHistoryCollector.cpp
--- a/Engine/Plugins/Experimental/Animation/PoseSearch/Source/Runtime/Private/AnimNode_PoseSearchHistoryCollector.cpp
+++ b/Engine/Plugins/Experimental/Animation/PoseSearch/Source/Runtime/Private/AnimNode_PoseSearchHistoryCollector.cpp
@@ -133,7 +133,14 @@ void FAnimNode_PoseSearchHistoryCollector::Evaluate_AnyThread(FPoseContext& Outp
 #if ENABLE_DRAW_DEBUG && ENABLE_ANIM_DEBUG
 	if (CVarAnimPoseHistoryDebugDraw.GetValueOnAnyThread())
 	{
-		PoseHistory.DebugDraw(*Output.AnimInstanceProxy, DebugColor.ToFColor(true), &DebugDrawTrajectory);
+		FColor Color;
+#if WITH_EDITORONLY_DATA
+		Color = DebugColor.ToFColor(true);
+#else // WITH_EDITORONLY_DATA
+		Color = FLinearColor::Red.ToFColor(true);
+#endif // WITH_EDITORONLY_DATA
+
+		PoseHistory.DebugDraw(*Output.AnimInstanceProxy, Color, &DebugDrawTrajectory);
 	}
 #endif // ENABLE_DRAW_DEBUG && ENABLE_ANIM_DEBUG
 }
@@ -185,7 +192,14 @@ void FAnimNode_PoseSearchComponentSpaceHistoryCollector::EvaluateComponentSpace_
 #if ENABLE_DRAW_DEBUG && ENABLE_ANIM_DEBUG
 	if (CVarAnimPoseHistoryDebugDraw.GetValueOnAnyThread())
 	{
-		PoseHistory.DebugDraw(*Output.AnimInstanceProxy, DebugColor.ToFColor(true), &DebugDrawTrajectory);
+		FColor Color;
+#if WITH_EDITORONLY_DATA
+		Color = DebugColor.ToFColor(true);
+#else // WITH_EDITORONLY_DATA
+		Color = FLinearColor::Red.ToFColor(true);
+#endif // WITH_EDITORONLY_DATA
+
+		PoseHistory.DebugDraw(*Output.AnimInstanceProxy, Color, &DebugDrawTrajectory);
 	}
 #endif // ENABLE_DRAW_DEBUG && ENABLE_ANIM_DEBUG
 }
diff --git a/Engine/Plugins/Experimental/Animation/PoseSearch/Source/Runtime/Private/MotionMatchingAnimNodeLibrary.cpp b/Engine/Plugins/Experimental/Animation/PoseSearch/Source/Runtime/Private/MotionMatchingAnimNodeLibrary.cpp
--- a/Engine/Plugins/Experimental/Animation/PoseSearch/Source/Runtime/Private/MotionMatchingAnimNodeLibrary.cpp
+++ b/Engine/Plugins/Experimental/Animation/PoseSearch/Source/Runtime/Private/MotionMatchingAnimNodeLibrary.cpp
@@ -26,9 +26,9 @@ void UMotionMatchingAnimNodeLibrary::GetMotionMatchingSearchResult(const FMotion
 				{
 					Result.SelectedAnimation = DatabaseAsset->GetAnimationAsset();
 					Result.SelectedTime = MotionMatchingState.CurrentSearchResult.AssetTime;
-					Result.bLoop = DatabaseAsset->IsLooping();
-					Result.bIsMirrored = SearchIndexAsset->bMirrored;
-					Result.BlendParameters = SearchIndexAsset->BlendParameters;
+					Result.bLoop = SearchIndexAsset->IsLooping();
+					Result.bIsMirrored = SearchIndexAsset->IsMirrored();
+					Result.BlendParameters = SearchIndexAsset->GetBlendParameters();
 					Result.SelectedDatabase = CurrentResultDatabase;
 					Result.WantedPlayRate = MotionMatchingState.WantedPlayRate;
 					bIsResultValid = true;
diff --git a/Engine/Plugins/Experimental/Animation/PoseSearch/Source/Runtime/Private/PoseSearchAssetIndexer.cpp b/Engine/Plugins/Experimental/Animation/PoseSearch/Source/Runtime/Private/PoseSearchAssetIndexer.cpp
--- a/Engine/Plugins/Experimental/Animation/PoseSearch/Source/Runtime/Private/PoseSearchAssetIndexer.cpp
+++ b/Engine/Plugins/Experimental/Animation/PoseSearch/Source/Runtime/Private/PoseSearchAssetIndexer.cpp
@@ -290,7 +290,7 @@ FAssetIndexer::FSampleInfo FAssetIndexer::GetSampleInfo(float SampleTime) const
 
 FTransform FAssetIndexer::MirrorTransform(const FTransform& Transform) const
 {
-	return SearchIndexAsset.bMirrored ? SamplingContext.MirrorTransform(Transform) : Transform;
+	return SearchIndexAsset.IsMirrored() ? SamplingContext.MirrorTransform(Transform) : Transform;
 }
 
 FAssetIndexer::CachedEntry& FAssetIndexer::GetEntry(float SampleTime)
@@ -364,7 +364,7 @@ FAssetIndexer::CachedEntry& FAssetIndexer::GetEntry(float SampleTime)
 
 		Pose[FCompactPoseBoneIndex(RootBoneIndexType)].SetIdentity();
 
-		if (SearchIndexAsset.bMirrored && Schema.MirrorDataTable)
+		if (SearchIndexAsset.IsMirrored() && Schema.MirrorDataTable)
 		{
 			FAnimationRuntime::MirrorPose(
 				Pose,
@@ -754,12 +754,14 @@ const UPoseSearchSchema* FAssetIndexer::GetSchema() const
 	return &Schema;
 }
 
+#if WITH_EDITOR
 float FAssetIndexer::CalculatePermutationTimeOffset() const
 {
 	check(Schema.PermutationsSampleRate > 0 && SearchIndexAsset.IsInitialized());
-	const float PermutationTimeOffset = Schema.PermutationsTimeOffset + SearchIndexAsset.PermutationIdx / float(Schema.PermutationsSampleRate);
+	const float PermutationTimeOffset = Schema.PermutationsTimeOffset + SearchIndexAsset.GetPermutationIdx() / float(Schema.PermutationsSampleRate);
 	return PermutationTimeOffset;
 }
+#endif // WITH_EDITOR
 
 #if ENABLE_ANIM_DEBUG
 void FAssetIndexer::CompareCachedEntries(const FAssetIndexer& Other) const
diff --git a/Engine/Plugins/Experimental/Animation/PoseSearch/Source/Runtime/Private/PoseSearchDatabase.cpp b/Engine/Plugins/Experimental/Animation/PoseSearch/Source/Runtime/Private/PoseSearchDatabase.cpp
--- a/Engine/Plugins/Experimental/Animation/PoseSearch/Source/Runtime/Private/PoseSearchDatabase.cpp
+++ b/Engine/Plugins/Experimental/Animation/PoseSearch/Source/Runtime/Private/PoseSearchDatabase.cpp
@@ -76,17 +76,15 @@ static void PopulateNonSelectableIdx(FNonSelectableIdx& NonSelectableIdx, FSearc
 	{
 		if (const FSearchIndexAsset* CurrentIndexAsset = SearchContext.GetCurrentResult().GetSearchIndexAsset(true))
 		{
-			const FPoseSearchDatabaseAnimationAssetBase* DatabaseAnimationAssetBase = Database->GetAnimationAssetStruct(CurrentIndexAsset->SourceAssetIdx).GetPtr<FPoseSearchDatabaseAnimationAssetBase>();
-			check(DatabaseAnimationAssetBase);
-			if (DatabaseAnimationAssetBase->bDisableReselection)
+			if (CurrentIndexAsset->IsDisableReselection())
 			{
-				// excluding all the poses from DatabaseAnimationAssetBase
+				// excluding all the poses with CurrentIndexAsset->GetSourceAssetIdx()
 				// @todo: optimize this code!
 				for (const FSearchIndexAsset& SearchIndexAsset : SearchIndex.Assets)
 				{
-					if (SearchIndexAsset.SourceAssetIdx == CurrentIndexAsset->SourceAssetIdx)
+					if (SearchIndexAsset.GetSourceAssetIdx() == CurrentIndexAsset->GetSourceAssetIdx())
 					{
-						const int32 FirstPoseIdx = SearchIndexAsset.FirstPoseIdx;
+						const int32 FirstPoseIdx = SearchIndexAsset.GetFirstPoseIdx();
 						const int32 LastPoseIdx = FirstPoseIdx + SearchIndexAsset.GetNumPoses();
 						for (int32 PoseIdx = FirstPoseIdx; PoseIdx < LastPoseIdx; ++PoseIdx)
 						{
@@ -107,11 +105,11 @@ static void PopulateNonSelectableIdx(FNonSelectableIdx& NonSelectableIdx, FSearc
 				const int32 CurrentResultPoseIdx = SearchContext.GetCurrentResult().PoseIdx;
 				const int32 UnboundMinPoseIdx = CurrentResultPoseIdx + FMath::FloorToInt(SearchContext.GetPoseJumpThresholdTime().Min * Database->Schema->SampleRate);
 				const int32 UnboundMaxPoseIdx = CurrentResultPoseIdx + FMath::CeilToInt(SearchContext.GetPoseJumpThresholdTime().Max * Database->Schema->SampleRate);
-				const int32 CurrentIndexAssetFirstPoseIdx = CurrentIndexAsset->FirstPoseIdx;
+				const int32 CurrentIndexAssetFirstPoseIdx = CurrentIndexAsset->GetFirstPoseIdx();
 				const int32 CurrentIndexAssetNumPoses = CurrentIndexAsset->GetNumPoses();
-				const bool IsLooping = Database->IsSourceAssetLooping(*CurrentIndexAsset);
+				const bool bIsLooping = CurrentIndexAsset->IsLooping();
 
-				if (IsLooping)
+				if (bIsLooping)
 				{
 					for (int32 UnboundPoseIdx = UnboundMinPoseIdx; UnboundPoseIdx < UnboundMaxPoseIdx; ++UnboundPoseIdx)
 					{
@@ -314,6 +312,7 @@ UAnimationAsset* FPoseSearchDatabaseSequence::GetAnimationAsset() const
 	return Sequence;
 }
 
+#if WITH_EDITORONLY_DATA
 UClass* FPoseSearchDatabaseSequence::GetAnimationAssetStaticClass() const
 {
 	return UAnimSequence::StaticClass();
@@ -336,6 +335,7 @@ bool FPoseSearchDatabaseSequence::IsRootMotionEnabled() const
 {
 	return Sequence ? Sequence->HasRootMotion() : false;
 }
+#endif // WITH_EDITORONLY_DATA
 
 //////////////////////////////////////////////////////////////////////////
 // FPoseSearchDatabaseBlendSpace
@@ -344,6 +344,7 @@ UAnimationAsset* FPoseSearchDatabaseBlendSpace::GetAnimationAsset() const
 	return BlendSpace.Get();
 }
 
+#if WITH_EDITORONLY_DATA
 UClass* FPoseSearchDatabaseBlendSpace::GetAnimationAssetStaticClass() const
 {
 	return UBlendSpace::StaticClass();
@@ -446,6 +447,8 @@ FVector FPoseSearchDatabaseBlendSpace::BlendParameterForSampleRanges(int32 Horiz
 		0.f);
 }
 
+#endif // WITH_EDITORONLY_DATA
+
 //////////////////////////////////////////////////////////////////////////
 // FPoseSearchDatabaseAnimComposite
 UAnimationAsset* FPoseSearchDatabaseAnimComposite::GetAnimationAsset() const
@@ -453,6 +456,7 @@ UAnimationAsset* FPoseSearchDatabaseAnimComposite::GetAnimationAsset() const
 	return AnimComposite;
 }
 
+#if WITH_EDITORONLY_DATA
 UClass* FPoseSearchDatabaseAnimComposite::GetAnimationAssetStaticClass() const
 {
 	return UAnimComposite::StaticClass();
@@ -475,6 +479,7 @@ bool FPoseSearchDatabaseAnimComposite::IsRootMotionEnabled() const
 {
 	return AnimComposite ? AnimComposite->HasRootMotion() : false;
 }
+#endif // WITH_EDITORONLY_DATA
 
 //////////////////////////////////////////////////////////////////////////
 // FPoseSearchDatabaseAnimMontage
@@ -483,6 +488,7 @@ UAnimationAsset* FPoseSearchDatabaseAnimMontage::GetAnimationAsset() const
 	return AnimMontage;
 }
 
+#if WITH_EDITORONLY_DATA
 UClass* FPoseSearchDatabaseAnimMontage::GetAnimationAssetStaticClass() const
 {
 	return UAnimMontage::StaticClass();
@@ -505,6 +511,7 @@ bool FPoseSearchDatabaseAnimMontage::IsRootMotionEnabled() const
 {
 	return AnimMontage ? AnimMontage->HasRootMotion() : false;
 }
+#endif // WITH_EDITORONLY_DATA
 
 //////////////////////////////////////////////////////////////////////////
 // UPoseSearchDatabase
@@ -521,14 +528,13 @@ void UPoseSearchDatabase::SetSearchIndex(const UE::PoseSearch::FSearchIndex& Sea
 const UE::PoseSearch::FSearchIndex& UPoseSearchDatabase::GetSearchIndex() const
 {
 	// making sure the search index is consistent. if it fails the calling code hasn't been protected by FAsyncPoseSearchDatabasesManagement::RequestAsyncBuildIndex
-	check(Schema && Schema->IsValid() && !SearchIndexPrivate.IsEmpty() && SearchIndexPrivate.WeightsSqrt.Num() == Schema->SchemaCardinality);
+	check(Schema && Schema->IsValid() && !SearchIndexPrivate.IsEmpty() && SearchIndexPrivate.GetNumDimensions() == Schema->SchemaCardinality);
 	return SearchIndexPrivate;
 }
 
 int32 UPoseSearchDatabase::GetPoseIndexFromTime(float Time, const UE::PoseSearch::FSearchIndexAsset& SearchIndexAsset) const
 {
-	const bool bIsLooping = IsSourceAssetLooping(SearchIndexAsset);
-	return SearchIndexAsset.GetPoseIndexFromTime(Time, bIsLooping, Schema->SampleRate);
+	return SearchIndexAsset.GetPoseIndexFromTime(Time, Schema->SampleRate);
 }
 
 const FInstancedStruct& UPoseSearchDatabase::GetAnimationAssetStruct(int32 AnimationAssetIndex) const
@@ -539,7 +545,7 @@ const FInstancedStruct& UPoseSearchDatabase::GetAnimationAssetStruct(int32 Anima
 
 const FInstancedStruct& UPoseSearchDatabase::GetAnimationAssetStruct(const UE::PoseSearch::FSearchIndexAsset& SearchIndexAsset) const
 {
-	return GetAnimationAssetStruct(SearchIndexAsset.SourceAssetIdx);
+	return GetAnimationAssetStruct(SearchIndexAsset.GetSourceAssetIdx());
 }
 
 FInstancedStruct& UPoseSearchDatabase::GetMutableAnimationAssetStruct(int32 AnimationAssetIndex)
@@ -550,7 +556,7 @@ FInstancedStruct& UPoseSearchDatabase::GetMutableAnimationAssetStruct(int32 Anim
 
 FInstancedStruct& UPoseSearchDatabase::GetMutableAnimationAssetStruct(const UE::PoseSearch::FSearchIndexAsset& SearchIndexAsset)
 {
-	return GetMutableAnimationAssetStruct(SearchIndexAsset.SourceAssetIdx);
+	return GetMutableAnimationAssetStruct(SearchIndexAsset.GetSourceAssetIdx());
 }
 
 const FPoseSearchDatabaseAnimationAssetBase* UPoseSearchDatabase::GetAnimationAssetBase(int32 AnimationAssetIndex) const
@@ -565,7 +571,7 @@ const FPoseSearchDatabaseAnimationAssetBase* UPoseSearchDatabase::GetAnimationAs
 
 const FPoseSearchDatabaseAnimationAssetBase* UPoseSearchDatabase::GetAnimationAssetBase(const UE::PoseSearch::FSearchIndexAsset& SearchIndexAsset) const
 {
-	return GetAnimationAssetBase(SearchIndexAsset.SourceAssetIdx);
+	return GetAnimationAssetBase(SearchIndexAsset.GetSourceAssetIdx());
 }
 
 FPoseSearchDatabaseAnimationAssetBase* UPoseSearchDatabase::GetMutableAnimationAssetBase(int32 AnimationAssetIndex)
@@ -580,23 +586,15 @@ FPoseSearchDatabaseAnimationAssetBase* UPoseSearchDatabase::GetMutableAnimationA
 
 FPoseSearchDatabaseAnimationAssetBase* UPoseSearchDatabase::GetMutableAnimationAssetBase(const UE::PoseSearch::FSearchIndexAsset& SearchIndexAsset)
 {
-	return GetMutableAnimationAssetBase(SearchIndexAsset.SourceAssetIdx);
-}
-
-const bool UPoseSearchDatabase::IsSourceAssetLooping(const UE::PoseSearch::FSearchIndexAsset& SearchIndexAsset) const
-{
-	return GetAnimationAssetBase(SearchIndexAsset.SourceAssetIdx)->IsLooping();
-}
-
-const FString UPoseSearchDatabase::GetSourceAssetName(const UE::PoseSearch::FSearchIndexAsset& SearchIndexAsset) const
-{
-	return GetAnimationAssetBase(SearchIndexAsset.SourceAssetIdx)->GetName();
+	return GetMutableAnimationAssetBase(SearchIndexAsset.GetSourceAssetIdx());
 }
 
+#if WITH_EDITOR
 int32 UPoseSearchDatabase::GetNumberOfPrincipalComponents() const
 {
 	return FMath::Min<int32>(NumberOfPrincipalComponents, Schema->SchemaCardinality);
 }
+#endif //WITH_EDITOR
 
 bool UPoseSearchDatabase::GetSkipSearchIfPossible() const
 {
@@ -818,7 +816,7 @@ float UPoseSearchDatabase::GetNormalizedAssetTime(int32 PoseIdx) const
 {
 	check(Schema);
 	const UE::PoseSearch::FSearchIndexAsset& Asset = GetSearchIndex().GetAssetForPose(PoseIdx);
-	const bool bIsBlendSpace = AnimationAssets[Asset.SourceAssetIdx].GetPtr<FPoseSearchDatabaseBlendSpace>() != nullptr;
+	const bool bIsBlendSpace = AnimationAssets[Asset.GetSourceAssetIdx()].GetPtr<FPoseSearchDatabaseBlendSpace>() != nullptr;
 
 	// sequences or anim composites
 	float AssetTime = Asset.GetTimeFromPoseIndex(PoseIdx, Schema->SampleRate);
@@ -957,7 +955,7 @@ FPoseSearchCost UPoseSearchDatabase::SearchContinuingPose(UE::PoseSearch::FSearc
 	const FSearchIndexAsset& SearchIndexAsset = SearchIndex.GetAssetForPose(PoseIdx);
 	const FPoseSearchDatabaseAnimationAssetBase* DatabaseAnimationAssetBase = GetAnimationAssetStruct(SearchIndexAsset).GetPtr<FPoseSearchDatabaseAnimationAssetBase>();
 	check(DatabaseAnimationAssetBase);
-	const FAnimationAssetSampler SequenceBaseSampler(DatabaseAnimationAssetBase->GetAnimationAsset(), SearchIndexAsset.BlendParameters);
+	const FAnimationAssetSampler SequenceBaseSampler(DatabaseAnimationAssetBase->GetAnimationAsset(), SearchIndexAsset.GetBlendParameters());
 	const float SampleTime = GetNormalizedAssetTime(PoseIdx);
 
 	float UpdatedContinuingPoseCostBias = ContinuingPoseCostBias;
@@ -1016,7 +1014,7 @@ UE::PoseSearch::FSearchResult UPoseSearchDatabase::SearchPCAKDTree(UE::PoseSearc
 	// there's no point in performing the search if CurrentBestTotalCost is already better than that
 	if (!GetSkipSearchIfPossible() || SearchContext.GetCurrentBestTotalCost() > SearchIndex.MinCostAddend)
 	{
-		const uint32 ClampedNumberOfPrincipalComponents = GetNumberOfPrincipalComponents();
+		const uint32 ClampedNumberOfPrincipalComponents = SearchIndex.GetNumberOfPrincipalComponents();
 		const uint32 ClampedKDTreeQueryNumNeighbors = FMath::Clamp<uint32>(KDTreeQueryNumNeighbors, 1, SearchIndex.GetNumPoses());
 		const bool bArePCAValuesPruned = SearchIndex.PCAValuesVectorToPoseIndexes.Num() > 0;
 
diff --git a/Engine/Plugins/Experimental/Animation/PoseSearch/Source/Runtime/Private/PoseSearchDerivedData.cpp b/Engine/Plugins/Experimental/Animation/PoseSearch/Source/Runtime/Private/PoseSearchDerivedData.cpp
--- a/Engine/Plugins/Experimental/Animation/PoseSearch/Source/Runtime/Private/PoseSearchDerivedData.cpp
+++ b/Engine/Plugins/Experimental/Animation/PoseSearch/Source/Runtime/Private/PoseSearchDerivedData.cpp
@@ -342,6 +342,8 @@ static void InitSearchIndexAssets(FSearchIndexBase& SearchIndex, const UPoseSear
 
 			const bool bAddUnmirrored = DatabaseAsset->GetMirrorOption() == EPoseSearchMirrorOption::UnmirroredOnly || DatabaseAsset->GetMirrorOption() == EPoseSearchMirrorOption::UnmirroredAndMirrored;
 			const bool bAddMirrored = DatabaseAsset->GetMirrorOption() == EPoseSearchMirrorOption::MirroredOnly || DatabaseAsset->GetMirrorOption() == EPoseSearchMirrorOption::UnmirroredAndMirrored;
+			const bool bIsLooping = DatabaseAsset->IsLooping();
+			const bool bDisableReselection = DatabaseAsset->IsDisableReselection();
 			const int32 SchemaSampleRate = Database->Schema->SampleRate;
 
 			if (const FPoseSearchDatabaseBlendSpace* DatabaseBlendSpace = DatabaseAssetStruct.GetPtr<FPoseSearchDatabaseBlendSpace>())
@@ -366,7 +368,8 @@ static void InitSearchIndexAssets(FSearchIndexBase& SearchIndex, const UPoseSear
 						{
 							if (bAddUnmirrored)
 							{
-								const FSearchIndexAsset PoseSearchIndexAsset(AnimationAssetIndex, TotalPoses, false, FFloatInterval(0.f, PlayLength), SchemaSampleRate, PermutationIdx, BlendParameters);
+								const FSearchIndexAsset PoseSearchIndexAsset(AnimationAssetIndex, TotalPoses, false, bIsLooping, 
+									bDisableReselection, FFloatInterval(0.f, PlayLength), SchemaSampleRate, PermutationIdx, BlendParameters);
 								if (PoseSearchIndexAsset.GetNumPoses() > 0)
 								{
 									SearchIndex.Assets.Add(PoseSearchIndexAsset);
@@ -376,7 +379,8 @@ static void InitSearchIndexAssets(FSearchIndexBase& SearchIndex, const UPoseSear
 
 							if (bAddMirrored)
 							{
-								const FSearchIndexAsset PoseSearchIndexAsset(AnimationAssetIndex, TotalPoses, true, FFloatInterval(0.f, PlayLength), SchemaSampleRate, PermutationIdx, BlendParameters);
+								const FSearchIndexAsset PoseSearchIndexAsset(AnimationAssetIndex, TotalPoses, true, bIsLooping,
+									bDisableReselection, FFloatInterval(0.f, PlayLength), SchemaSampleRate, PermutationIdx, BlendParameters);
 								if (PoseSearchIndexAsset.GetNumPoses() > 0)
 								{
 									SearchIndex.Assets.Add(PoseSearchIndexAsset);
@@ -392,14 +396,15 @@ static void InitSearchIndexAssets(FSearchIndexBase& SearchIndex, const UPoseSear
 			{
 				ValidRanges.Reset();
 
-				FindValidSequenceIntervals(SequenceBase, DatabaseAsset->GetSamplingRange(), DatabaseAsset->IsLooping(), Database->ExcludeFromDatabaseParameters, ValidRanges);
+				FindValidSequenceIntervals(SequenceBase, DatabaseAsset->GetSamplingRange(), bIsLooping, Database->ExcludeFromDatabaseParameters, ValidRanges);
 				for (const FFloatRange& Range : ValidRanges)
 				{
 					for (int32 PermutationIdx = 0; PermutationIdx < Database->Schema->NumberOfPermutations; ++PermutationIdx)
 					{
 						if (bAddUnmirrored)
 						{
-							const FSearchIndexAsset PoseSearchIndexAsset(AnimationAssetIndex, TotalPoses, false, FFloatInterval(Range.GetLowerBoundValue(), Range.GetUpperBoundValue()), SchemaSampleRate, PermutationIdx);
+							const FSearchIndexAsset PoseSearchIndexAsset(AnimationAssetIndex, TotalPoses, false, bIsLooping,
+								bDisableReselection, FFloatInterval(Range.GetLowerBoundValue(), Range.GetUpperBoundValue()), SchemaSampleRate, PermutationIdx);
 							if (PoseSearchIndexAsset.GetNumPoses() > 0)
 							{
 								SearchIndex.Assets.Add(PoseSearchIndexAsset);
@@ -409,7 +414,8 @@ static void InitSearchIndexAssets(FSearchIndexBase& SearchIndex, const UPoseSear
 
 						if (bAddMirrored)
 						{
-							const FSearchIndexAsset PoseSearchIndexAsset(AnimationAssetIndex, TotalPoses, true, FFloatInterval(Range.GetLowerBoundValue(), Range.GetUpperBoundValue()), SchemaSampleRate, PermutationIdx);
+							const FSearchIndexAsset PoseSearchIndexAsset(AnimationAssetIndex, TotalPoses, true, bIsLooping,
+								bDisableReselection, FFloatInterval(Range.GetLowerBoundValue(), Range.GetUpperBoundValue()), SchemaSampleRate, PermutationIdx);
 							if (PoseSearchIndexAsset.GetNumPoses() > 0)
 							{
 								SearchIndex.Assets.Add(PoseSearchIndexAsset);
@@ -844,11 +850,11 @@ static bool IndexDatabase(FSearchIndexBase& SearchIndexBase, const UPoseSearchDa
 	for (int32 AssetIdx = 0; AssetIdx != SearchIndexBase.Assets.Num(); ++AssetIdx)
 	{
 		FSearchIndexAsset& SearchIndexAsset = SearchIndexBase.Assets[AssetIdx];
-		check(SearchIndexAsset.FirstPoseIdx == TotalPoses);
+		check(SearchIndexAsset.GetFirstPoseIdx() == TotalPoses);
 
 		const FPoseSearchDatabaseAnimationAssetBase* DatabaseAnimationAssetBase = Database.GetAnimationAssetStruct(SearchIndexAsset).GetPtr<FPoseSearchDatabaseAnimationAssetBase>();
 		check(DatabaseAnimationAssetBase && DatabaseAnimationAssetBase->GetAnimationAsset());
-		const FAnimationAssetSampler& AssetSampler = Samplers[SamplerMap[{ DatabaseAnimationAssetBase->GetAnimationAsset(), SearchIndexAsset.BlendParameters }]];
+		const FAnimationAssetSampler& AssetSampler = Samplers[SamplerMap[{ DatabaseAnimationAssetBase->GetAnimationAsset(), SearchIndexAsset.GetBlendParameters() }]];
 
 		Indexers.Emplace(BoneContainer, SearchIndexAsset, SamplingContext, *Schema, AssetSampler);
 		TotalPoses += SearchIndexAsset.GetNumPoses();
@@ -1287,7 +1293,7 @@ void FPoseSearchDatabaseAsyncCacheTask::Wait(FCriticalSection& OuterMutex)
 	{
 		Database->SetSearchIndex(SearchIndex); // @todo: implement FSearchIndex move ctor and assignment operator and use a MoveTemp(SearchIndex) here
 
-		check(Database->Schema && Database->Schema->IsValid() && !SearchIndex.IsEmpty() && SearchIndex.WeightsSqrt.Num() == Database->Schema->SchemaCardinality);
+		check(Database->Schema && Database->Schema->IsValid() && !SearchIndex.IsEmpty() && SearchIndex.GetNumDimensions() == Database->Schema->SchemaCardinality);
 
 		SetState(EState::Ended);
 		bBroadcastOnDerivedDataRebuild = true;
@@ -1343,7 +1349,7 @@ void FPoseSearchDatabaseAsyncCacheTask::OnGetComplete(UE::DerivedData::FCacheGet
 		// cache can be corrupted in case the version of the derived data cache has not being updated while 
 		// developing channels that changes their cardinality without impacting any asset properties
 		// so to account for this, we just reindex the database and update the associated DDC 
-		if (SearchIndex.WeightsSqrt.Num() == Database->Schema->SchemaCardinality)
+		if (SearchIndex.GetNumDimensions() == Database->Schema->SchemaCardinality)
 		{
 			UE_LOG(LogPoseSearch, Log, TEXT("%s - %s BuildIndex From Cache"), *LexToString(FullIndexKey.Hash), *Database->GetName());
 		}
diff --git a/Engine/Plugins/Experimental/Animation/PoseSearch/Source/Runtime/Private/PoseSearchDerivedDataKey.cpp b/Engine/Plugins/Experimental/Animation/PoseSearch/Source/Runtime/Private/PoseSearchDerivedDataKey.cpp
--- a/Engine/Plugins/Experimental/Animation/PoseSearch/Source/Runtime/Private/PoseSearchDerivedDataKey.cpp
+++ b/Engine/Plugins/Experimental/Animation/PoseSearch/Source/Runtime/Private/PoseSearchDerivedDataKey.cpp
@@ -40,7 +40,7 @@ FKeyBuilder::FKeyBuilder(const UObject* Object, bool bUseDataVer, bool bPerformC
 	if (bUseDataVer)
 	{
 		// used to invalidate the key without having to change POSESEARCHDB_DERIVEDDATA_VER all the times
-		int32 POSESEARCHDB_DERIVEDDATA_VER_SMALL = 219;
+		int32 POSESEARCHDB_DERIVEDDATA_VER_SMALL = 227;
 		FGuid VersionGuid = FDevSystemGuids::GetSystemGuid(FDevSystemGuids::Get().POSESEARCHDB_DERIVEDDATA_VER);
 
 		*this << VersionGuid;
diff --git a/Engine/Plugins/Experimental/Animation/PoseSearch/Source/Runtime/Private/PoseSearchFeatureChannel_Heading.cpp b/Engine/Plugins/Experimental/Animation/PoseSearch/Source/Runtime/Private/PoseSearchFeatureChannel_Heading.cpp
--- a/Engine/Plugins/Experimental/Animation/PoseSearch/Source/Runtime/Private/PoseSearchFeatureChannel_Heading.cpp
+++ b/Engine/Plugins/Experimental/Animation/PoseSearch/Source/Runtime/Private/PoseSearchFeatureChannel_Heading.cpp
@@ -30,11 +30,13 @@ void UPoseSearchFeatureChannel_Heading::FindOrAddToSchema(UPoseSearchSchema* Sch
 	{
 		UPoseSearchFeatureChannel_Heading* Heading = NewObject<UPoseSearchFeatureChannel_Heading>(Schema, NAME_None, RF_Transient);
 		Heading->Bone.BoneName = BoneName;
+#if WITH_EDITORONLY_DATA
 		Heading->Weight = 0.f;
-		Heading->SampleTimeOffset = SampleTimeOffset;
-		Heading->HeadingAxis = HeadingAxis;
 		// @todo: perhaps add a tunable color for injected channels
 		Heading->DebugColor = FLinearColor::Gray;
+#endif // WITH_EDITORONLY_DATA
+		Heading->SampleTimeOffset = SampleTimeOffset;
+		Heading->HeadingAxis = HeadingAxis;
 		Heading->PermutationTimeType = PermutationTimeType;
 		Schema->AddTemporaryChannel(Heading);
 	}
@@ -125,7 +127,13 @@ void UPoseSearchFeatureChannel_Heading::DebugDraw(const UE::PoseSearch::FDebugDr
 {
 	using namespace UE::PoseSearch;
 
-	const FColor Color = DebugColor.ToFColor(true);
+	FColor Color;
+#if WITH_EDITORONLY_DATA
+	Color = DebugColor.ToFColor(true);
+#else // WITH_EDITORONLY_DATA
+	Color = FLinearColor::White.ToFColor(true);
+#endif // WITH_EDITORONLY_DATA
+
 	const FVector BoneHeading = DrawParams.ExtractRotation(PoseVector, OriginTimeOffset).RotateVector(FFeatureVectorHelper::DecodeVector(PoseVector, ChannelDataOffset, ComponentStripping));
 	const FVector BonePos = DrawParams.ExtractPosition(PoseVector, SampleTimeOffset, SchemaBoneIdx, PermutationTimeType, SamplingAttributeId);
 
diff --git a/Engine/Plugins/Experimental/Animation/PoseSearch/Source/Runtime/Private/PoseSearchFeatureChannel_Heading.h b/Engine/Plugins/Experimental/Animation/PoseSearch/Source/Runtime/Private/PoseSearchFeatureChannel_Heading.h
--- a/Engine/Plugins/Experimental/Animation/PoseSearch/Source/Runtime/Private/PoseSearchFeatureChannel_Heading.h
+++ b/Engine/Plugins/Experimental/Animation/PoseSearch/Source/Runtime/Private/PoseSearchFeatureChannel_Heading.h
@@ -29,9 +29,10 @@ public:
 	UPROPERTY(EditAnywhere, Category = "Settings")
 	FBoneReference OriginBone;
 
+#if WITH_EDITORONLY_DATA
 	UPROPERTY(EditAnywhere, Category = "Settings")
 	float Weight = 1.f;
-
+#endif // WITH_EDITORONLY_DATA
 
 	// if SamplingAttributeId >= 0, ALL the animations contained in the pose search database referencing the schema containing this channel are expected to have 
 	// UAnimNotifyState_PoseSearchSamplingAttribute notify state with a matching SamplingAttributeId, and the UAnimNotifyState_PoseSearchSamplingAttribute properties
@@ -60,8 +61,10 @@ public:
 	UPROPERTY(Transient)
 	int8 SchemaOriginBoneIdx = 0;
 
+#if WITH_EDITORONLY_DATA
 	UPROPERTY(EditAnywhere, Category = "Settings", meta = (ExcludeFromHash, DisplayPriority = 0))
 	FLinearColor DebugColor = FLinearColor::White;
+#endif // WITH_EDITORONLY_DATA
 
 	UPROPERTY(EditAnywhere, Category = "Settings")
 	EInputQueryPose InputQueryPose = EInputQueryPose::UseContinuingPose;
diff --git a/Engine/Plugins/Experimental/Animation/PoseSearch/Source/Runtime/Private/PoseSearchFeatureChannel_PermutationTime.cpp b/Engine/Plugins/Experimental/Animation/PoseSearch/Source/Runtime/Private/PoseSearchFeatureChannel_PermutationTime.cpp
--- a/Engine/Plugins/Experimental/Animation/PoseSearch/Source/Runtime/Private/PoseSearchFeatureChannel_PermutationTime.cpp
+++ b/Engine/Plugins/Experimental/Animation/PoseSearch/Source/Runtime/Private/PoseSearchFeatureChannel_PermutationTime.cpp
@@ -10,7 +10,9 @@ void UPoseSearchFeatureChannel_PermutationTime::FindOrAddToSchema(UPoseSearchSch
 	if (Schema->FindFirstChannelOfType<UPoseSearchFeatureChannel_PermutationTime>() == nullptr)
 	{
 		UPoseSearchFeatureChannel_PermutationTime* PermutationTime = NewObject<UPoseSearchFeatureChannel_PermutationTime>(Schema, NAME_None, RF_Transient);
+#if WITH_EDITORONLY_DATA
 		PermutationTime->Weight = 0.f;
+#endif // WITH_EDITORONLY_DATA
 		Schema->AddTemporaryChannel(PermutationTime);
 	}
 }
diff --git a/Engine/Plugins/Experimental/Animation/PoseSearch/Source/Runtime/Private/PoseSearchFeatureChannel_PermutationTime.h b/Engine/Plugins/Experimental/Animation/PoseSearch/Source/Runtime/Private/PoseSearchFeatureChannel_PermutationTime.h
--- a/Engine/Plugins/Experimental/Animation/PoseSearch/Source/Runtime/Private/PoseSearchFeatureChannel_PermutationTime.h
+++ b/Engine/Plugins/Experimental/Animation/PoseSearch/Source/Runtime/Private/PoseSearchFeatureChannel_PermutationTime.h
@@ -11,8 +11,10 @@ class POSESEARCH_API UPoseSearchFeatureChannel_PermutationTime : public UPoseSea
 	GENERATED_BODY()
 
 public:
+#if WITH_EDITORONLY_DATA
 	UPROPERTY(EditAnywhere, Category = "Settings")
 	float Weight = 1.f;
+#endif // WITH_EDITORONLY_DATA
 
 	// UPoseSearchFeatureChannel interface
 	virtual void Finalize(UPoseSearchSchema* Schema) override;
diff --git a/Engine/Plugins/Experimental/Animation/PoseSearch/Source/Runtime/Private/PoseSearchFeatureChannel_Phase.cpp b/Engine/Plugins/Experimental/Animation/PoseSearch/Source/Runtime/Private/PoseSearchFeatureChannel_Phase.cpp
--- a/Engine/Plugins/Experimental/Animation/PoseSearch/Source/Runtime/Private/PoseSearchFeatureChannel_Phase.cpp
+++ b/Engine/Plugins/Experimental/Animation/PoseSearch/Source/Runtime/Private/PoseSearchFeatureChannel_Phase.cpp
@@ -332,8 +332,12 @@ void UPoseSearchFeatureChannel_Phase::DebugDraw(const UE::PoseSearch::FDebugDraw
 
 	static float ScaleFactor = 1.f;
 
-
-	const FColor Color = DebugColor.ToFColor(true);
+	FColor Color;
+#if WITH_EDITORONLY_DATA
+	Color = DebugColor.ToFColor(true);
+#else // WITH_EDITORONLY_DATA
+	Color = FLinearColor::Yellow.ToFColor(true);
+#endif // WITH_EDITORONLY_DATA
 
 	const FVector2D Phase = FFeatureVectorHelper::DecodeVector2D(PoseVector, ChannelDataOffset);
 	const FVector BonePos = DrawParams.ExtractPosition(PoseVector, 0.f, SchemaBoneIdx);
diff --git a/Engine/Plugins/Experimental/Animation/PoseSearch/Source/Runtime/Private/PoseSearchFeatureChannel_Phase.h b/Engine/Plugins/Experimental/Animation/PoseSearch/Source/Runtime/Private/PoseSearchFeatureChannel_Phase.h
--- a/Engine/Plugins/Experimental/Animation/PoseSearch/Source/Runtime/Private/PoseSearchFeatureChannel_Phase.h
+++ b/Engine/Plugins/Experimental/Animation/PoseSearch/Source/Runtime/Private/PoseSearchFeatureChannel_Phase.h
@@ -15,15 +15,19 @@ public:
 	UPROPERTY(EditAnywhere, Category = "Settings")
 	FBoneReference Bone;
 
+#if WITH_EDITORONLY_DATA
 	UPROPERTY(EditAnywhere, Category = "Settings")
 	float Weight = 1.f;
+#endif // WITH_EDITORONLY_DATA
 
 	// index referencing the associated bone in UPoseSearchSchema::BoneReferences
 	UPROPERTY(Transient)
 	int8 SchemaBoneIdx = 0;
 
+#if WITH_EDITORONLY_DATA
 	UPROPERTY(EditAnywhere, Category = "Settings", meta = (ExcludeFromHash, DisplayPriority = 0))
 	FLinearColor DebugColor = FLinearColor::Yellow;
+#endif // WITH_EDITORONLY_DATA
 
 	UPROPERTY(EditAnywhere, Category = "Settings")
 	EInputQueryPose InputQueryPose = EInputQueryPose::UseContinuingPose;
diff --git a/Engine/Plugins/Experimental/Animation/PoseSearch/Source/Runtime/Private/PoseSearchFeatureChannel_Pose.cpp b/Engine/Plugins/Experimental/Animation/PoseSearch/Source/Runtime/Private/PoseSearchFeatureChannel_Pose.cpp
--- a/Engine/Plugins/Experimental/Animation/PoseSearch/Source/Runtime/Private/PoseSearchFeatureChannel_Pose.cpp
+++ b/Engine/Plugins/Experimental/Animation/PoseSearch/Source/Runtime/Private/PoseSearchFeatureChannel_Pose.cpp
@@ -14,9 +14,21 @@
 UPoseSearchFeatureChannel_Pose::UPoseSearchFeatureChannel_Pose()
 {
 	// defaulting UPoseSearchFeatureChannel_Pose for a meaningful locomotion setup
+#if WITH_EDITORONLY_DATA
 	Weight = 1.f;
-	SampledBones.Add(FPoseSearchBone({ {"foot_l"}, int32(EPoseSearchBoneFlags::Position | EPoseSearchBoneFlags::Velocity), 1.f, FLinearColor::Green }));
-	SampledBones.Add(FPoseSearchBone({ {"foot_r"}, int32(EPoseSearchBoneFlags::Position | EPoseSearchBoneFlags::Velocity), 1.f, FLinearColor::Green}));
+#endif // WITH_EDITORONLY_DATA
+
+	SampledBones.Add(FPoseSearchBone({ {"foot_l"}, int32(EPoseSearchBoneFlags::Position | EPoseSearchBoneFlags::Velocity)
+#if WITH_EDITORONLY_DATA
+		, 1.f, FLinearColor::Green
+#endif // WITH_EDITORONLY_DATA
+		}));
+
+	SampledBones.Add(FPoseSearchBone({ {"foot_r"}, int32(EPoseSearchBoneFlags::Position | EPoseSearchBoneFlags::Velocity)
+#if WITH_EDITORONLY_DATA
+		, 1.f, FLinearColor::Green
+#endif // WITH_EDITORONLY_DATA
+		}));
 }
 
 void UPoseSearchFeatureChannel_Pose::Finalize(UPoseSearchSchema* Schema)
@@ -30,9 +42,11 @@ void UPoseSearchFeatureChannel_Pose::Finalize(UPoseSearchSchema* Schema)
 		{
 			UPoseSearchFeatureChannel_Position* Position = NewObject<UPoseSearchFeatureChannel_Position>(this, NAME_None, RF_Transient);
 			Position->Bone = SampledBone.Reference;
+#if WITH_EDITORONLY_DATA
 			Position->Weight = SampledBone.Weight * Weight;
-			Position->SampleTimeOffset = 0.f;
 			Position->DebugColor = SampledBone.DebugColor;
+#endif // WITH_EDITORONLY_DATA
+			Position->SampleTimeOffset = 0.f;
 			Position->InputQueryPose = InputQueryPose;
 			SubChannels.Add(Position);
 		}
@@ -41,19 +55,23 @@ void UPoseSearchFeatureChannel_Pose::Finalize(UPoseSearchSchema* Schema)
 		{
 			UPoseSearchFeatureChannel_Heading* HeadingX = NewObject<UPoseSearchFeatureChannel_Heading>(this, NAME_None, RF_Transient);
 			HeadingX->Bone = SampledBone.Reference;
+#if WITH_EDITORONLY_DATA
 			HeadingX->Weight = SampledBone.Weight * Weight;
+			HeadingX->DebugColor = SampledBone.DebugColor;
+#endif // WITH_EDITORONLY_DATA
 			HeadingX->SampleTimeOffset = 0.f;
 			HeadingX->HeadingAxis = EHeadingAxis::X;
-			HeadingX->DebugColor = SampledBone.DebugColor;
 			HeadingX->InputQueryPose = InputQueryPose;
 			SubChannels.Add(HeadingX);
 
 			UPoseSearchFeatureChannel_Heading* HeadingY = NewObject<UPoseSearchFeatureChannel_Heading>(this, NAME_None, RF_Transient);
 			HeadingY->Bone = SampledBone.Reference;
+#if WITH_EDITORONLY_DATA
 			HeadingY->Weight = SampledBone.Weight * Weight;
+			HeadingY->DebugColor = SampledBone.DebugColor;
+#endif // WITH_EDITORONLY_DATA
 			HeadingY->SampleTimeOffset = 0.f;
 			HeadingY->HeadingAxis = EHeadingAxis::Y;
-			HeadingY->DebugColor = SampledBone.DebugColor;
 			HeadingY->InputQueryPose = InputQueryPose;
 			SubChannels.Add(HeadingY);
 		}
@@ -62,9 +80,11 @@ void UPoseSearchFeatureChannel_Pose::Finalize(UPoseSearchSchema* Schema)
 		{
 			UPoseSearchFeatureChannel_Velocity* Velocity = NewObject<UPoseSearchFeatureChannel_Velocity>(this, NAME_None, RF_Transient);
 			Velocity->Bone = SampledBone.Reference;
+#if WITH_EDITORONLY_DATA
 			Velocity->Weight = SampledBone.Weight * Weight;
-			Velocity->SampleTimeOffset = 0.f;
 			Velocity->DebugColor = SampledBone.DebugColor;
+#endif // WITH_EDITORONLY_DATA
+			Velocity->SampleTimeOffset = 0.f;
 			Velocity->InputQueryPose = InputQueryPose;
 			Velocity->bUseCharacterSpaceVelocities = bUseCharacterSpaceVelocities;
 			SubChannels.Add(Velocity);
@@ -74,8 +94,10 @@ void UPoseSearchFeatureChannel_Pose::Finalize(UPoseSearchSchema* Schema)
 		{
 			UPoseSearchFeatureChannel_Phase* Phase = NewObject<UPoseSearchFeatureChannel_Phase>(this, NAME_None, RF_Transient);
 			Phase->Bone = SampledBone.Reference;
+#if WITH_EDITORONLY_DATA
 			Phase->Weight = SampledBone.Weight * Weight;
 			Phase->DebugColor = SampledBone.DebugColor;
+#endif // WITH_EDITORONLY_DATA
 			Phase->InputQueryPose = InputQueryPose;
 			SubChannels.Add(Phase);
 		}
diff --git a/Engine/Plugins/Experimental/Animation/PoseSearch/Source/Runtime/Private/PoseSearchFeatureChannel_Pose.h b/Engine/Plugins/Experimental/Animation/PoseSearch/Source/Runtime/Private/PoseSearchFeatureChannel_Pose.h
--- a/Engine/Plugins/Experimental/Animation/PoseSearch/Source/Runtime/Private/PoseSearchFeatureChannel_Pose.h
+++ b/Engine/Plugins/Experimental/Animation/PoseSearch/Source/Runtime/Private/PoseSearchFeatureChannel_Pose.h
@@ -29,11 +29,13 @@ struct POSESEARCH_API FPoseSearchBone
 	UPROPERTY(EditAnywhere, meta = (Bitmask, BitmaskEnum = "/Script/PoseSearch.EPoseSearchBoneFlags"), Category = Config)
 	int32 Flags = int32(EPoseSearchBoneFlags::Position);
 
+#if WITH_EDITORONLY_DATA
 	UPROPERTY(EditAnywhere, Category = Config)
 	float Weight = 1.f;
 
 	UPROPERTY(EditAnywhere, Category = Config, meta=(ExcludeFromHash, DisplayPriority = 0))
 	FLinearColor DebugColor = FLinearColor::Green;
+#endif // WITH_EDITORONLY_DATA
 };
 
 // UPoseSearchFeatureChannel_Pose
@@ -43,8 +45,10 @@ class POSESEARCH_API UPoseSearchFeatureChannel_Pose : public UPoseSearchFeatureC
 	GENERATED_BODY()
 
 public:
+#if WITH_EDITORONLY_DATA
 	UPROPERTY(EditAnywhere, Category = "Settings")
 	float Weight = 1.f;
+#endif // WITH_EDITORONLY_DATA
 
 	// List of skeletal joints and associated Flags (Velocity, Position, etc) to sample.
 	UPROPERTY(EditAnywhere, Category = "Settings")
diff --git a/Engine/Plugins/Experimental/Animation/PoseSearch/Source/Runtime/Private/PoseSearchFeatureChannel_Position.cpp b/Engine/Plugins/Experimental/Animation/PoseSearch/Source/Runtime/Private/PoseSearchFeatureChannel_Position.cpp
--- a/Engine/Plugins/Experimental/Animation/PoseSearch/Source/Runtime/Private/PoseSearchFeatureChannel_Position.cpp
+++ b/Engine/Plugins/Experimental/Animation/PoseSearch/Source/Runtime/Private/PoseSearchFeatureChannel_Position.cpp
@@ -29,10 +29,12 @@ void UPoseSearchFeatureChannel_Position::FindOrAddToSchema(UPoseSearchSchema* Sc
 	{
 		UPoseSearchFeatureChannel_Position* Position = NewObject<UPoseSearchFeatureChannel_Position>(Schema, NAME_None, RF_Transient);
 		Position->Bone.BoneName = BoneName;
+#if WITH_EDITORONLY_DATA
 		Position->Weight = 0.f;
-		Position->SampleTimeOffset = SampleTimeOffset;
 		// @todo: perhaps add a tunable color for injected channels
 		Position->DebugColor = FLinearColor::Gray;
+#endif // WITH_EDITORONLY_DATA
+		Position->SampleTimeOffset = SampleTimeOffset;
 		Position->PermutationTimeType = PermutationTimeType;
 		Schema->AddTemporaryChannel(Position);
 	}
@@ -104,7 +106,12 @@ void UPoseSearchFeatureChannel_Position::DebugDraw(const UE::PoseSearch::FDebugD
 {
 	using namespace UE::PoseSearch;
 
-	const FColor Color = DebugColor.ToFColor(true);
+	FColor Color;
+#if WITH_EDITORONLY_DATA
+	Color = DebugColor.ToFColor(true);
+#else // WITH_EDITORONLY_DATA
+	Color = FLinearColor::Blue.ToFColor(true);
+#endif // WITH_EDITORONLY_DATA
 
 	const FVector FeaturesVector = FFeatureVectorHelper::DecodeVector(PoseVector, ChannelDataOffset, ComponentStripping);
 	if (SchemaOriginBoneIdx == RootSchemaBoneIdx)
diff --git a/Engine/Plugins/Experimental/Animation/PoseSearch/Source/Runtime/Private/PoseSearchFeatureChannel_Position.h b/Engine/Plugins/Experimental/Animation/PoseSearch/Source/Runtime/Private/PoseSearchFeatureChannel_Position.h
--- a/Engine/Plugins/Experimental/Animation/PoseSearch/Source/Runtime/Private/PoseSearchFeatureChannel_Position.h
+++ b/Engine/Plugins/Experimental/Animation/PoseSearch/Source/Runtime/Private/PoseSearchFeatureChannel_Position.h
@@ -18,9 +18,10 @@ public:
 	UPROPERTY(EditAnywhere, Category = "Settings")
 	FBoneReference OriginBone;
 
+#if WITH_EDITORONLY_DATA
 	UPROPERTY(EditAnywhere, Category = "Settings")
 	float Weight = 1.f;
-
+#endif // WITH_EDITORONLY_DATA
 
 	// if SamplingAttributeId >= 0, ALL the animations contained in the pose search database referencing the schema containing this channel are expected to have 
 	// UAnimNotifyState_PoseSearchSamplingAttribute notify state with a matching SamplingAttributeId, and the UAnimNotifyState_PoseSearchSamplingAttribute properties
@@ -46,8 +47,10 @@ public:
 	UPROPERTY(Transient)
 	int8 SchemaOriginBoneIdx = 0;
 
+#if WITH_EDITORONLY_DATA
 	UPROPERTY(EditAnywhere, Category = "Settings", meta = (ExcludeFromHash, DisplayPriority = 0))
 	FLinearColor DebugColor = FLinearColor::Blue;
+#endif // WITH_EDITORONLY_DATA
 
 	UPROPERTY(EditAnywhere, Category = "Settings")
 	EInputQueryPose InputQueryPose = EInputQueryPose::UseContinuingPose;
diff --git a/Engine/Plugins/Experimental/Animation/PoseSearch/Source/Runtime/Private/PoseSearchFeatureChannel_Trajectory.cpp b/Engine/Plugins/Experimental/Animation/PoseSearch/Source/Runtime/Private/PoseSearchFeatureChannel_Trajectory.cpp
--- a/Engine/Plugins/Experimental/Animation/PoseSearch/Source/Runtime/Private/PoseSearchFeatureChannel_Trajectory.cpp
+++ b/Engine/Plugins/Experimental/Animation/PoseSearch/Source/Runtime/Private/PoseSearchFeatureChannel_Trajectory.cpp
@@ -15,11 +15,33 @@
 UPoseSearchFeatureChannel_Trajectory::UPoseSearchFeatureChannel_Trajectory()
 {
 	// defaulting UPoseSearchFeatureChannel_Trajectory for a meaningful locomotion setup
+#if WITH_EDITORONLY_DATA
 	Weight = 7.f;
-	Samples.Add(FPoseSearchTrajectorySample({ -0.4f, int32(EPoseSearchTrajectoryFlags::PositionXY), 0.4f, FLinearColor::Red }));
-	Samples.Add(FPoseSearchTrajectorySample({ 0.f, int32(EPoseSearchTrajectoryFlags::VelocityXY | EPoseSearchTrajectoryFlags::FacingDirectionXY), 2.f, FLinearColor::Blue }));
-	Samples.Add(FPoseSearchTrajectorySample({ 0.35f, int32(EPoseSearchTrajectoryFlags::PositionXY | EPoseSearchTrajectoryFlags::FacingDirectionXY), 0.7f, FLinearColor::Blue }));
-	Samples.Add(FPoseSearchTrajectorySample({ 0.7f, int32(EPoseSearchTrajectoryFlags::VelocityXY | EPoseSearchTrajectoryFlags::PositionXY | EPoseSearchTrajectoryFlags::FacingDirectionXY), 0.5f, FLinearColor::Blue }));
+#endif // WITH_EDITORONLY_DATA
+
+	Samples.Add(FPoseSearchTrajectorySample({ -0.4f, int32(EPoseSearchTrajectoryFlags::PositionXY)
+#if WITH_EDITORONLY_DATA
+		, 0.4f, FLinearColor::Red
+#endif // WITH_EDITORONLY_DATA
+		}));
+
+	Samples.Add(FPoseSearchTrajectorySample({ 0.f, int32(EPoseSearchTrajectoryFlags::VelocityXY | EPoseSearchTrajectoryFlags::FacingDirectionXY)
+#if WITH_EDITORONLY_DATA
+		, 2.f, FLinearColor::Blue
+#endif // WITH_EDITORONLY_DATA
+		}));
+
+	Samples.Add(FPoseSearchTrajectorySample({ 0.35f, int32(EPoseSearchTrajectoryFlags::PositionXY | EPoseSearchTrajectoryFlags::FacingDirectionXY)
+#if WITH_EDITORONLY_DATA
+		, 0.7f, FLinearColor::Blue
+#endif // WITH_EDITORONLY_DATA
+		}));
+
+	Samples.Add(FPoseSearchTrajectorySample({ 0.7f, int32(EPoseSearchTrajectoryFlags::VelocityXY | EPoseSearchTrajectoryFlags::PositionXY | EPoseSearchTrajectoryFlags::FacingDirectionXY)
+#if WITH_EDITORONLY_DATA
+		, 0.5f, FLinearColor::Blue
+#endif // WITH_EDITORONLY_DATA
+		}));
 }
 
 void UPoseSearchFeatureChannel_Trajectory::Finalize(UPoseSearchSchema* Schema)
@@ -31,9 +53,11 @@ void UPoseSearchFeatureChannel_Trajectory::Finalize(UPoseSearchSchema* Schema)
 		if (EnumHasAnyFlags(Sample.Flags, EPoseSearchTrajectoryFlags::Position | EPoseSearchTrajectoryFlags::PositionXY))
 		{
 			UPoseSearchFeatureChannel_Position* Position = NewObject<UPoseSearchFeatureChannel_Position>(this, NAME_None, RF_Transient);
+#if WITH_EDITORONLY_DATA
 			Position->Weight = Sample.Weight * Weight;
-			Position->SampleTimeOffset = Sample.Offset;
 			Position->DebugColor = Sample.DebugColor;
+#endif // WITH_EDITORONLY_DATA
+			Position->SampleTimeOffset = Sample.Offset;
 			Position->InputQueryPose = EInputQueryPose::UseCharacterPose;
 			if (EnumHasAnyFlags(Sample.Flags, EPoseSearchTrajectoryFlags::PositionXY))
 			{
@@ -45,9 +69,11 @@ void UPoseSearchFeatureChannel_Trajectory::Finalize(UPoseSearchSchema* Schema)
 		if (EnumHasAnyFlags(Sample.Flags, EPoseSearchTrajectoryFlags::Velocity | EPoseSearchTrajectoryFlags::VelocityXY))
 		{
 			UPoseSearchFeatureChannel_Velocity* Velocity = NewObject<UPoseSearchFeatureChannel_Velocity>(this, NAME_None, RF_Transient);
+#if WITH_EDITORONLY_DATA
 			Velocity->Weight = Sample.Weight * Weight;
-			Velocity->SampleTimeOffset = Sample.Offset;
 			Velocity->DebugColor = Sample.DebugColor;
+#endif // WITH_EDITORONLY_DATA
+			Velocity->SampleTimeOffset = Sample.Offset;
 			Velocity->InputQueryPose = EInputQueryPose::UseCharacterPose;
 			Velocity->bUseCharacterSpaceVelocities = false;
 			if (EnumHasAnyFlags(Sample.Flags, EPoseSearchTrajectoryFlags::VelocityXY))
@@ -60,9 +86,11 @@ void UPoseSearchFeatureChannel_Trajectory::Finalize(UPoseSearchSchema* Schema)
 		if (EnumHasAnyFlags(Sample.Flags, EPoseSearchTrajectoryFlags::VelocityDirection | EPoseSearchTrajectoryFlags::VelocityDirectionXY))
 		{
 			UPoseSearchFeatureChannel_Velocity* Velocity = NewObject<UPoseSearchFeatureChannel_Velocity>(this, NAME_None, RF_Transient);
+#if WITH_EDITORONLY_DATA
 			Velocity->Weight = Sample.Weight * Weight;
-			Velocity->SampleTimeOffset = Sample.Offset;
 			Velocity->DebugColor = Sample.DebugColor;
+#endif // WITH_EDITORONLY_DATA
+			Velocity->SampleTimeOffset = Sample.Offset;
 			Velocity->InputQueryPose = EInputQueryPose::UseCharacterPose;
 			Velocity->bUseCharacterSpaceVelocities = false;
 			Velocity->bNormalize = true;
@@ -76,9 +104,11 @@ void UPoseSearchFeatureChannel_Trajectory::Finalize(UPoseSearchSchema* Schema)
 		if (EnumHasAnyFlags(Sample.Flags, EPoseSearchTrajectoryFlags::FacingDirection | EPoseSearchTrajectoryFlags::FacingDirectionXY))
 		{
 			UPoseSearchFeatureChannel_Heading* Heading = NewObject<UPoseSearchFeatureChannel_Heading>(this, NAME_None, RF_Transient);
+#if WITH_EDITORONLY_DATA
 			Heading->Weight = Sample.Weight * Weight;
-			Heading->SampleTimeOffset = Sample.Offset;
 			Heading->DebugColor = Sample.DebugColor;
+#endif // WITH_EDITORONLY_DATA
+			Heading->SampleTimeOffset = Sample.Offset;
 			Heading->InputQueryPose = EInputQueryPose::UseCharacterPose;
 			if (EnumHasAnyFlags(Sample.Flags, EPoseSearchTrajectoryFlags::FacingDirectionXY))
 			{
@@ -119,7 +149,12 @@ void UPoseSearchFeatureChannel_Trajectory::DebugDraw(const UE::PoseSearch::FDebu
 		for (int32 i = 0; i < Positions.Num(); ++i)
 		{
 			const float CurrTimeOffset = Positions[i]->SampleTimeOffset;
-			const FColor Color = Positions[i]->DebugColor.ToFColor(true);;
+			FColor Color;
+#if WITH_EDITORONLY_DATA
+			Color = Positions[i]->DebugColor.ToFColor(true);
+#else // WITH_EDITORONLY_DATA
+			Color = FLinearColor::Blue.ToFColor(true);
+#endif // WITH_EDITORONLY_DATA
 
 			if (PrevTimeOffset * CurrTimeOffset < UE_KINDA_SMALL_NUMBER)
 			{
diff --git a/Engine/Plugins/Experimental/Animation/PoseSearch/Source/Runtime/Private/PoseSearchFeatureChannel_Trajectory.h b/Engine/Plugins/Experimental/Animation/PoseSearch/Source/Runtime/Private/PoseSearchFeatureChannel_Trajectory.h
--- a/Engine/Plugins/Experimental/Animation/PoseSearch/Source/Runtime/Private/PoseSearchFeatureChannel_Trajectory.h
+++ b/Engine/Plugins/Experimental/Animation/PoseSearch/Source/Runtime/Private/PoseSearchFeatureChannel_Trajectory.h
@@ -36,20 +36,25 @@ struct POSESEARCH_API FPoseSearchTrajectorySample
 	UPROPERTY(EditAnywhere, meta = (Bitmask, BitmaskEnum = "/Script/PoseSearch.EPoseSearchTrajectoryFlags"), Category = Config)
 	int32 Flags = int32(EPoseSearchTrajectoryFlags::Position);
 
+#if WITH_EDITORONLY_DATA
 	UPROPERTY(EditAnywhere, Category = Config)
 	float Weight = 1.f;
 
 	UPROPERTY(EditAnywhere, Category = Config, meta = (ExcludeFromHash, DisplayPriority = 0))
 	FLinearColor DebugColor = FLinearColor::Blue;
+#endif // WITH_EDITORONLY_DATA
 };
+
 UCLASS(BlueprintType, EditInlineNew, meta = (DisplayName = "Trajectory Channel"), CollapseCategories)
 class POSESEARCH_API UPoseSearchFeatureChannel_Trajectory : public UPoseSearchFeatureChannel_GroupBase
 {
 	GENERATED_BODY()
 
 public:
+#if WITH_EDITORONLY_DATA
 	UPROPERTY(EditAnywhere, Category = "Settings")
 	float Weight = 1.f;
+#endif // WITH_EDITORONLY_DATA
 
 	UPROPERTY(EditAnywhere, Category = "Settings")
 	TArray<FPoseSearchTrajectorySample> Samples;
diff --git a/Engine/Plugins/Experimental/Animation/PoseSearch/Source/Runtime/Private/PoseSearchFeatureChannel_Velocity.cpp b/Engine/Plugins/Experimental/Animation/PoseSearch/Source/Runtime/Private/PoseSearchFeatureChannel_Velocity.cpp
--- a/Engine/Plugins/Experimental/Animation/PoseSearch/Source/Runtime/Private/PoseSearchFeatureChannel_Velocity.cpp
+++ b/Engine/Plugins/Experimental/Animation/PoseSearch/Source/Runtime/Private/PoseSearchFeatureChannel_Velocity.cpp
@@ -85,7 +85,13 @@ void UPoseSearchFeatureChannel_Velocity::DebugDraw(const UE::PoseSearch::FDebugD
 {
 	using namespace UE::PoseSearch;
 
-	const FColor Color = DebugColor.ToFColor(true);
+	FColor Color;
+#if WITH_EDITORONLY_DATA
+	Color = DebugColor.ToFColor(true);
+#else // WITH_EDITORONLY_DATA
+	Color = FLinearColor::Green.ToFColor(true);
+#endif // WITH_EDITORONLY_DATA
+
 	const float LinearVelocityScale = bNormalize ? 15.f : 0.08f;
 
 	const FVector LinearVelocity = DrawParams.GetRootTransform().TransformVector(FFeatureVectorHelper::DecodeVector(PoseVector, ChannelDataOffset, ComponentStripping));
diff --git a/Engine/Plugins/Experimental/Animation/PoseSearch/Source/Runtime/Private/PoseSearchFeatureChannel_Velocity.h b/Engine/Plugins/Experimental/Animation/PoseSearch/Source/Runtime/Private/PoseSearchFeatureChannel_Velocity.h
--- a/Engine/Plugins/Experimental/Animation/PoseSearch/Source/Runtime/Private/PoseSearchFeatureChannel_Velocity.h
+++ b/Engine/Plugins/Experimental/Animation/PoseSearch/Source/Runtime/Private/PoseSearchFeatureChannel_Velocity.h
@@ -18,9 +18,10 @@ public:
 	UPROPERTY(EditAnywhere, Category = "Settings")
 	FBoneReference OriginBone;
 
+#if WITH_EDITORONLY_DATA
 	UPROPERTY(EditAnywhere, Category = "Settings")
 	float Weight = 1.f;
-
+#endif // WITH_EDITORONLY_DATA
 
 	// if SamplingAttributeId >= 0, ALL the animations contained in the pose search database referencing the schema containing this channel are expected to have 
 	// UAnimNotifyState_PoseSearchSamplingAttribute notify state with a matching SamplingAttributeId, and the UAnimNotifyState_PoseSearchSamplingAttribute properties
@@ -46,8 +47,10 @@ public:
 	UPROPERTY(Transient)
 	int8 SchemaOriginBoneIdx = 0;
 
+#if WITH_EDITORONLY_DATA
 	UPROPERTY(EditAnywhere, Category = "Settings", meta = (ExcludeFromHash, DisplayPriority = 0))
 	FLinearColor DebugColor = FLinearColor::Green;
+#endif // WITH_EDITORONLY_DATA
 
 	UPROPERTY(EditAnywhere, Category = "Settings")
 	EInputQueryPose InputQueryPose = EInputQueryPose::UseContinuingPose;
diff --git a/Engine/Plugins/Experimental/Animation/PoseSearch/Source/Runtime/Private/PoseSearchIndex.cpp b/Engine/Plugins/Experimental/Animation/PoseSearch/Source/Runtime/Private/PoseSearchIndex.cpp
--- a/Engine/Plugins/Experimental/Animation/PoseSearch/Source/Runtime/Private/PoseSearchIndex.cpp
+++ b/Engine/Plugins/Experimental/Animation/PoseSearch/Source/Runtime/Private/PoseSearchIndex.cpp
@@ -302,6 +302,7 @@ FArchive& operator<<(FArchive& Ar, FPoseMetadata& Metadata)
 	bool bInBlockTransition = Metadata.IsBlockTransition();
 	FFloat16 CostAddend = Metadata.CostAddend;
 	
+	// @todo: optimize the archived size of FPoseMetadata, since most members are bitfields
 	Ar << ValueOffset;
 	Ar << AssetIndex;
 	Ar << bInBlockTransition;
@@ -318,8 +319,11 @@ bool FSearchIndexAsset::operator==(const FSearchIndexAsset& Other) const
 	return
 		SourceAssetIdx == Other.SourceAssetIdx &&
 		bMirrored == Other.bMirrored &&
+		bLooping == Other.bLooping &&
+		bDisableReselection == Other.bDisableReselection &&
 		PermutationIdx == Other.PermutationIdx &&
-		BlendParameters == Other.BlendParameters &&
+		BlendParameterX == Other.BlendParameterX &&
+		BlendParameterY == Other.BlendParameterY &&
 		FirstPoseIdx == Other.FirstPoseIdx &&
 		FirstSampleIdx == Other.FirstSampleIdx &&
 		LastSampleIdx == Other.LastSampleIdx;
@@ -327,13 +331,32 @@ bool FSearchIndexAsset::operator==(const FSearchIndexAsset& Other) const
 
 FArchive& operator<<(FArchive& Ar, FSearchIndexAsset& IndexAsset)
 {
-	Ar << IndexAsset.SourceAssetIdx;
-	Ar << IndexAsset.bMirrored;
-	Ar << IndexAsset.PermutationIdx;
-	Ar << IndexAsset.BlendParameters;
-	Ar << IndexAsset.FirstPoseIdx;
-	Ar << IndexAsset.FirstSampleIdx;
-	Ar << IndexAsset.LastSampleIdx;
+	int32 SourceAssetIdx = IndexAsset.SourceAssetIdx;
+	bool bMirrored = IndexAsset.bMirrored; 
+	bool bLooping = IndexAsset.bLooping;
+	bool bDisableReselection = IndexAsset.bDisableReselection;
+	int32 PermutationIdx = IndexAsset.PermutationIdx;
+	float BlendParameterX = IndexAsset.BlendParameterX;
+	float BlendParameterY = IndexAsset.BlendParameterY;
+	int32 FirstPoseIdx = IndexAsset.FirstPoseIdx;
+	int32 FirstSampleIdx = IndexAsset.FirstSampleIdx;
+	int32 LastSampleIdx = IndexAsset.LastSampleIdx;
+
+	// @todo: optimize the archived size of FSearchIndexAsset, since most members are bitfields
+	Ar << SourceAssetIdx;
+	Ar << bMirrored;
+	Ar << bLooping;
+	Ar << bDisableReselection;
+	Ar << PermutationIdx;
+	Ar << BlendParameterX;
+	Ar << BlendParameterY;
+	Ar << FirstPoseIdx;
+	Ar << FirstSampleIdx;
+	Ar << LastSampleIdx;
+
+	new(&IndexAsset) FSearchIndexAsset(SourceAssetIdx, bMirrored, bLooping, bDisableReselection,
+		PermutationIdx, FVector(BlendParameterX, BlendParameterY, 0.f), FirstPoseIdx, FirstSampleIdx, LastSampleIdx);
+
 	return Ar;
 }
 
@@ -516,7 +539,7 @@ void FSearchIndex::Reset()
 
 TConstArrayView<float> FSearchIndex::GetPoseValues(int32 PoseIdx) const
 {
-	return GetPoseValuesBase(PoseIdx, WeightsSqrt.Num());
+	return GetPoseValuesBase(PoseIdx, GetNumDimensions());
 }
 
 TConstArrayView<float> FSearchIndex::GetReconstructedPoseValues(int32 PoseIdx, TArrayView<float> BufferUsedForReconstruction) const
@@ -526,7 +549,7 @@ TConstArrayView<float> FSearchIndex::GetReconstructedPoseValues(int32 PoseIdx, T
 	// @todo: reconstruction is not yet supported with pruned PCAValues
 	check(PCAValuesVectorToPoseIndexes.Num() == 0);
 
-	const int32 NumDimensions = WeightsSqrt.Num();
+	const int32 NumDimensions = GetNumDimensions();
 	const int32 NumPoses = GetNumPoses();
 	check(PoseIdx >= 0 && PoseIdx < NumPoses && NumDimensions > 0);
 	check(BufferUsedForReconstruction.Num() == NumDimensions);
@@ -552,15 +575,27 @@ TConstArrayView<float> FSearchIndex::GetReconstructedPoseValues(int32 PoseIdx, T
 	return BufferUsedForReconstruction;
 }
 
-TConstArrayView<float> FSearchIndex::PCAProject(TConstArrayView<float> PoseValues, TArrayView<float> BufferUsedForProjection) const
+int32 FSearchIndex::GetNumDimensions() const
 {
-	QUICK_SCOPE_CYCLE_COUNTER(STAT_PoseSearch_PCAProject);
+	return WeightsSqrt.Num();
+}
+
+int32 FSearchIndex::GetNumberOfPrincipalComponents() const
+{
+	const int32 NumDimensions = GetNumDimensions();
+	check(NumDimensions > 0 && PCAProjectionMatrix.Num() > 0 && PCAProjectionMatrix.Num() % NumDimensions == 0);
 
-	const int32 NumDimensions = WeightsSqrt.Num();
 	const int32 NumberOfPrincipalComponents = PCAProjectionMatrix.Num() / NumDimensions;
+	return NumberOfPrincipalComponents;
+}
 
+TConstArrayView<float> FSearchIndex::PCAProject(TConstArrayView<float> PoseValues, TArrayView<float> BufferUsedForProjection) const
+{
+	QUICK_SCOPE_CYCLE_COUNTER(STAT_PoseSearch_PCAProject);
+
+	const int32 NumDimensions = GetNumDimensions();
+	const int32 NumberOfPrincipalComponents = GetNumberOfPrincipalComponents();
 	check(PoseValues.Num() == NumDimensions);
-	check(PCAProjectionMatrix.Num() > 0 && PCAProjectionMatrix.Num() % NumDimensions == 0);
 	check(BufferUsedForProjection.Num() == NumberOfPrincipalComponents);
 
 	const RowMajorVectorMapConst WeightsSqrtMap(WeightsSqrt.GetData(), 1, NumDimensions);
@@ -657,7 +692,7 @@ TArray<float> FSearchIndex::GetPoseValuesSafe(int32 PoseIdx) const
 	{
 		if (IsValuesEmpty())
 		{
-			const int32 NumDimensions = WeightsSqrt.Num();
+			const int32 NumDimensions = GetNumDimensions();
 			PoseValues.SetNumUninitialized(NumDimensions);
 			GetReconstructedPoseValues(PoseIdx, PoseValues);
 		}
@@ -676,11 +711,10 @@ TConstArrayView<float> FSearchIndex::GetPCAPoseValues(int32 PCAValuesVectorIdx)
 		return TConstArrayView<float>();
 	}
 
-	const int32 NumDimensions = WeightsSqrt.Num();
-	const int32 NumberOfPrincipalComponents = PCAProjectionMatrix.Num() / NumDimensions;
+	const int32 NumDimensions = GetNumDimensions();
+	const int32 NumberOfPrincipalComponents = GetNumberOfPrincipalComponents();
 
 #if DO_CHECK
-	check(NumDimensions > 0 && PCAProjectionMatrix.Num() > 0 && PCAProjectionMatrix.Num() % NumDimensions == 0);
 	check(PCAValues.Num() % NumberOfPrincipalComponents == 0);
 	const int32 NumPCAValuesVectors = PCAValues.Num() / NumberOfPrincipalComponents;
 	check(PCAValuesVectorIdx >= 0 && PCAValuesVectorIdx < NumPCAValuesVectors );
diff --git a/Engine/Plugins/Experimental/Animation/PoseSearch/Source/Runtime/Private/PoseSearchLibrary.cpp b/Engine/Plugins/Experimental/Animation/PoseSearch/Source/Runtime/Private/PoseSearchLibrary.cpp
--- a/Engine/Plugins/Experimental/Animation/PoseSearch/Source/Runtime/Private/PoseSearchLibrary.cpp
+++ b/Engine/Plugins/Experimental/Animation/PoseSearch/Source/Runtime/Private/PoseSearchLibrary.cpp
@@ -478,9 +478,9 @@ void UPoseSearchLibrary::MotionMatch(
 				if (const FAnimNode_PoseSearchHistoryCollector_Base* PoseHistoryNode = TagSubsystem->FindNodeByTag<FAnimNode_PoseSearchHistoryCollector_Base>(PoseHistoryName, AnimInstance))
 				{
 					ExtendedPoseHistory.Init(&PoseHistoryNode->GetPoseHistory());
-#if ENABLE_DRAW_DEBUG && ENABLE_ANIM_DEBUG
+#if ENABLE_DRAW_DEBUG && ENABLE_ANIM_DEBUG && WITH_EDITORONLY_DATA
 					HistoryCollectorColor = PoseHistoryNode->DebugColor.ToFColor(true);
-#endif // ENABLE_DRAW_DEBUG && ENABLE_ANIM_DEBUG
+#endif // ENABLE_DRAW_DEBUG && ENABLE_ANIM_DEBUG && WITH_EDITORONLY_DATA
 				}
 			}
 		}
@@ -556,9 +556,9 @@ void UPoseSearchLibrary::MotionMatch(
 			{
 				Result.SelectedAnimation = DatabaseAsset->GetAnimationAsset();
 				Result.SelectedTime = SearchResult.AssetTime;
-				Result.bLoop = DatabaseAsset->IsLooping();
-				Result.bIsMirrored = SearchIndexAsset->bMirrored;
-				Result.BlendParameters = SearchIndexAsset->BlendParameters;
+				Result.bLoop = SearchIndexAsset->IsLooping();
+				Result.bIsMirrored = SearchIndexAsset->IsMirrored();
+				Result.BlendParameters = SearchIndexAsset->GetBlendParameters();
 				Result.SelectedDatabase = Database;
 				Result.SearchCost = SearchResult.PoseCost.GetTotalCost();
 			}
diff --git a/Engine/Plugins/Experimental/Animation/PoseSearch/Source/Runtime/Private/PoseSearchResult.cpp b/Engine/Plugins/Experimental/Animation/PoseSearch/Source/Runtime/Private/PoseSearchResult.cpp
--- a/Engine/Plugins/Experimental/Animation/PoseSearch/Source/Runtime/Private/PoseSearchResult.cpp
+++ b/Engine/Plugins/Experimental/Animation/PoseSearch/Source/Runtime/Private/PoseSearchResult.cpp
@@ -39,7 +39,7 @@ void FSearchResult::Update(float NewAssetTime)
 		{
 			TArray<FBlendSampleData> BlendSamples;
 			int32 TriangulationIndex = 0;
-			DatabaseBlendSpace->BlendSpace->GetSamplesFromBlendInput(SearchIndexAsset.BlendParameters, BlendSamples, TriangulationIndex, true);
+			DatabaseBlendSpace->BlendSpace->GetSamplesFromBlendInput(SearchIndexAsset.GetBlendParameters(), BlendSamples, TriangulationIndex, true);
 
 			const float PlayLength = DatabaseBlendSpace->BlendSpace->GetAnimationLengthFromSampleData(BlendSamples);
 
@@ -96,18 +96,18 @@ bool FSearchResult::CanAdvance(float DeltaTime) const
 		{
 			TArray<FBlendSampleData> BlendSamples;
 			int32 TriangulationIndex = 0;
-			DatabaseBlendSpace->BlendSpace->GetSamplesFromBlendInput(SearchIndexAsset->BlendParameters, BlendSamples, TriangulationIndex, true);
+			DatabaseBlendSpace->BlendSpace->GetSamplesFromBlendInput(SearchIndexAsset->GetBlendParameters(), BlendSamples, TriangulationIndex, true);
 
 			const float PlayLength = DatabaseBlendSpace->BlendSpace->GetAnimationLengthFromSampleData(BlendSamples);
 
 			// Asset player time for blend spaces is normalized [0, 1] so we need to convert it back to real time before we advance it
 			SteppedTime = AssetTime * PlayLength;
-			bCanAdvance = ETAA_Finished != FAnimationRuntime::AdvanceTime(DatabaseBlendSpace->IsLooping(), DeltaTime, SteppedTime, PlayLength);
+			bCanAdvance = ETAA_Finished != FAnimationRuntime::AdvanceTime(SearchIndexAsset->IsLooping(), DeltaTime, SteppedTime, PlayLength);
 		}
 		else if (const FPoseSearchDatabaseAnimationAssetBase* DatabaseAnimationAssetBase = DatabaseAsset.GetPtr<FPoseSearchDatabaseAnimationAssetBase>())
 		{
 			const float AssetLength = DatabaseAnimationAssetBase->GetAnimationAsset()->GetPlayLength();
-			bCanAdvance = ETAA_Finished != FAnimationRuntime::AdvanceTime(DatabaseAnimationAssetBase->IsLooping(), DeltaTime, SteppedTime, AssetLength);
+			bCanAdvance = ETAA_Finished != FAnimationRuntime::AdvanceTime(SearchIndexAsset->IsLooping(), DeltaTime, SteppedTime, AssetLength);
 		}
 	}
 	return bCanAdvance;
diff --git a/Engine/Plugins/Experimental/Animation/PoseSearch/Source/Runtime/Public/PoseSearch/AnimNode_PoseSearchHistoryCollector.h b/Engine/Plugins/Experimental/Animation/PoseSearch/Source/Runtime/Public/PoseSearch/AnimNode_PoseSearchHistoryCollector.h
--- a/Engine/Plugins/Experimental/Animation/PoseSearch/Source/Runtime/Public/PoseSearch/AnimNode_PoseSearchHistoryCollector.h
+++ b/Engine/Plugins/Experimental/Animation/PoseSearch/Source/Runtime/Public/PoseSearch/AnimNode_PoseSearchHistoryCollector.h
@@ -36,8 +36,10 @@ public:
 	// Update Counter for detecting being relevant
 	FGraphTraversalCounter UpdateCounter;
 
+#if WITH_EDITORONLY_DATA
 	UPROPERTY(EditAnywhere, Category = Debug)
 	FLinearColor DebugColor = FLinearColor::Red;
+#endif // WITH_EDITORONLY_DATA
 
 	// Optional Trajectory for debug drawing purposes only: if not provided, all the collected pose transforms will be drawn using the Graph context
 	UPROPERTY(EditAnywhere, Category = Debug, meta = (PinHiddenByDefault))
diff --git a/Engine/Plugins/Experimental/Animation/PoseSearch/Source/Runtime/Public/PoseSearch/PoseSearchAssetIndexer.h b/Engine/Plugins/Experimental/Animation/PoseSearch/Source/Runtime/Public/PoseSearch/PoseSearchAssetIndexer.h
--- a/Engine/Plugins/Experimental/Animation/PoseSearch/Source/Runtime/Public/PoseSearch/PoseSearchAssetIndexer.h
+++ b/Engine/Plugins/Experimental/Animation/PoseSearch/Source/Runtime/Public/PoseSearch/PoseSearchAssetIndexer.h
@@ -78,11 +78,13 @@ public:
 	
 	TArrayView<float> GetPoseVector(int32 SampleIdx) const;
 	const UPoseSearchSchema* GetSchema() const;
-	float CalculatePermutationTimeOffset() const;
 	float CalculateSampleTime(int32 SampleIdx) const;
-
 	bool IsProcessFailed() const { return bProcessFailed; }
 
+#if WITH_EDITOR
+	float CalculatePermutationTimeOffset() const;
+#endif //WITH_EDITOR
+
 #if ENABLE_ANIM_DEBUG
 	void CompareCachedEntries(const FAssetIndexer& Other) const;
 #endif // ENABLE_ANIM_DEBUG
diff --git a/Engine/Plugins/Experimental/Animation/PoseSearch/Source/Runtime/Public/PoseSearch/PoseSearchDatabase.h b/Engine/Plugins/Experimental/Animation/PoseSearch/Source/Runtime/Public/PoseSearch/PoseSearchDatabase.h
--- a/Engine/Plugins/Experimental/Animation/PoseSearch/Source/Runtime/Public/PoseSearch/PoseSearchDatabase.h
+++ b/Engine/Plugins/Experimental/Animation/PoseSearch/Source/Runtime/Public/PoseSearch/PoseSearchDatabase.h
@@ -56,6 +56,7 @@ struct FPoseSearchExcludeFromDatabaseParameters
 {
 	GENERATED_BODY()
 
+#if WITH_EDITORONLY_DATA
 	// Determines how much of the start of an animation segment is preserved for blending in seconds.
 	// Excluding the beginning of animation segments can help ensure an exact past trajectory is used when building the channels.
 	UPROPERTY(EditAnywhere, Category = "Settings", meta = (DisplayName = "Anim Start Interval"))
@@ -66,6 +67,7 @@ struct FPoseSearchExcludeFromDatabaseParameters
 	// and also prevents the selection of a anim segment which will end too soon to be worth selecting.
 	UPROPERTY(EditAnywhere, Category = "Settings", meta = (DisplayName = "Anim End Interval"))
 	float SequenceEndInterval = 0.3f;
+#endif // WITH_EDITORONLY_DATA
 };
 
 USTRUCT()
@@ -73,17 +75,19 @@ struct POSESEARCH_API FPoseSearchDatabaseAnimationAssetBase
 {
 	GENERATED_BODY()
 	virtual ~FPoseSearchDatabaseAnimationAssetBase() = default;
-
 	virtual UAnimationAsset* GetAnimationAsset() const { return nullptr; }
+
+#if WITH_EDITORONLY_DATA
+	virtual bool IsDisableReselection() const { return bDisableReselection; }
+	virtual void SetDisableReselection(bool bValue) { bDisableReselection = bValue; }
 	virtual UClass* GetAnimationAssetStaticClass() const { return nullptr; }
 	virtual bool IsLooping() const { return false; }
 	virtual const FString GetName() const { return FString(); }
-	virtual bool IsDisableReselection() const { return bDisableReselection; }
-	virtual void SetDisableReselection(bool bValue) { bDisableReselection = bValue; }
 	virtual bool IsEnabled() const { return bEnabled; }
 	virtual void SetIsEnabled(bool bValue) { bEnabled = bValue; }
 	virtual bool IsRootMotionEnabled() const { return false; }
 	virtual EPoseSearchMirrorOption GetMirrorOption() const { return MirrorOption; }
+
 	// [0, 0] represents the entire frame range of the original animation.
 	virtual FFloatInterval GetSamplingRange() const { return FFloatInterval(0.f, 0.f); }
 
@@ -100,7 +104,6 @@ struct POSESEARCH_API FPoseSearchDatabaseAnimationAssetBase
 	UPROPERTY(EditAnywhere, Category = "Settings", meta = (DisplayPriority = 3))
 	EPoseSearchMirrorOption MirrorOption = EPoseSearchMirrorOption::UnmirroredOnly;
 
-#if WITH_EDITORONLY_DATA
 	// SynchronizeWithExternalDependency is true when this asset has been added via SynchronizeWithExternalDependencies.
 	// To delete it, remove the PoseSearchBranchIn notify state
 	UPROPERTY(VisibleAnywhere, Category = "Settings", meta = (DisplayPriority = 20))
@@ -118,18 +121,21 @@ struct POSESEARCH_API FPoseSearchDatabaseSequence : public FPoseSearchDatabaseAn
 	UPROPERTY(EditAnywhere, Category="Settings", meta = (DisplayPriority = 0))
 	TObjectPtr<UAnimSequence> Sequence;
 
+#if WITH_EDITORONLY_DATA
 	// It allows users to set a time range to an individual animation sequence in the database. 
 	// This is effectively trimming the beginning and end of the animation in the database (not in the original sequence).
 	// If set to [0, 0] it will be the entire frame range of the original sequence.
 	UPROPERTY(EditAnywhere, Category="Settings", meta = (DisplayPriority = 2))
 	FFloatInterval SamplingRange = FFloatInterval(0.f, 0.f);
 
-	UAnimationAsset* GetAnimationAsset() const override;
 	UClass* GetAnimationAssetStaticClass() const override;
 	bool IsLooping() const override;
 	const FString GetName() const override;
 	bool IsRootMotionEnabled() const override;
 	FFloatInterval GetSamplingRange() const override { return SamplingRange; }
+#endif // WITH_EDITORONLY_DATA
+	
+	UAnimationAsset* GetAnimationAsset() const override;
 };
 
 /** An blend space entry in a UPoseSearchDatabase. */
@@ -142,6 +148,8 @@ struct POSESEARCH_API FPoseSearchDatabaseBlendSpace : public FPoseSearchDatabase
 	UPROPERTY(EditAnywhere, Category = "Settings", meta = (DisplayPriority = 0))
 	TObjectPtr<UBlendSpace> BlendSpace;
 
+#if WITH_EDITORONLY_DATA
+
 	// If true this BlendSpace will output a single segment in the database.
 	UPROPERTY(EditAnywhere, Category = "Settings", meta = (DisplayPriority = 4))
 	bool bUseSingleSample = false;
@@ -166,7 +174,6 @@ struct POSESEARCH_API FPoseSearchDatabaseBlendSpace : public FPoseSearchDatabase
 	UPROPERTY(EditAnywhere, Category = "Settings", meta = (EditCondition = "bUseSingleSample", EditConditionHides, DisplayPriority = 9))
 	float BlendParamY = 0.f;
 
-	UAnimationAsset* GetAnimationAsset() const override;
 	UClass* GetAnimationAssetStaticClass() const override;
 	bool IsLooping() const override;
 	const FString GetName() const override;
@@ -174,6 +181,9 @@ struct POSESEARCH_API FPoseSearchDatabaseBlendSpace : public FPoseSearchDatabase
 
 	void GetBlendSpaceParameterSampleRanges(int32& HorizontalBlendNum, int32& VerticalBlendNum) const;
 	FVector BlendParameterForSampleRanges(int32 HorizontalBlendIndex, int32 VerticalBlendIndex) const;
+#endif // WITH_EDITORONLY_DATA
+
+	UAnimationAsset* GetAnimationAsset() const override;
 };
 
 /** An entry in a UPoseSearchDatabase. */
@@ -186,18 +196,21 @@ struct POSESEARCH_API FPoseSearchDatabaseAnimComposite : public FPoseSearchDatab
 	UPROPERTY(EditAnywhere, Category = "Settings", meta = (DisplayPriority = 0))
 	TObjectPtr<UAnimComposite> AnimComposite;
 
+#if WITH_EDITORONLY_DATA
 	// It allows users to set a time range to an individual animation sequence in the database. 
 	// This is effectively trimming the beginning and end of the animation in the database (not in the original sequence).
 	// If set to [0, 0] it will be the entire frame range of the original sequence.
 	UPROPERTY(EditAnywhere, Category = "Settings", meta = (DisplayPriority = 3))
 	FFloatInterval SamplingRange = FFloatInterval(0.f, 0.f);
 
-	UAnimationAsset* GetAnimationAsset() const override;
 	UClass* GetAnimationAssetStaticClass() const override;
 	bool IsLooping() const override;
 	const FString GetName() const override;
 	bool IsRootMotionEnabled() const override;
 	FFloatInterval GetSamplingRange() const override { return SamplingRange; }
+#endif // WITH_EDITORONLY_DATA
+
+	UAnimationAsset* GetAnimationAsset() const override;
 };
 
 /** An anim montage entry in a UPoseSearchDatabase. */
@@ -210,18 +223,21 @@ struct POSESEARCH_API FPoseSearchDatabaseAnimMontage : public FPoseSearchDatabas
 	UPROPERTY(EditAnywhere, Category="Settings", meta = (DisplayPriority = 0))
 	TObjectPtr<UAnimMontage> AnimMontage;
 
+#if WITH_EDITORONLY_DATA
 	// It allows users to set a time range to an individual animation sequence in the database. 
 	// This is effectively trimming the beginning and end of the animation in the database (not in the original sequence).
 	// If set to [0, 0] it will be the entire frame range of the original sequence.
 	UPROPERTY(EditAnywhere, Category="Settings", meta = (DisplayPriority = 2))
 	FFloatInterval SamplingRange = FFloatInterval(0.f, 0.f);
 
-	UAnimationAsset* GetAnimationAsset() const override;
 	UClass* GetAnimationAssetStaticClass() const override;
 	bool IsLooping() const override;
 	const FString GetName() const override;
 	bool IsRootMotionEnabled() const override;
 	FFloatInterval GetSamplingRange() const override { return SamplingRange; }
+#endif // WITH_EDITORONLY_DATA
+
+	UAnimationAsset* GetAnimationAsset() const override;
 };
 
 /** A data asset for indexing a collection of animation sequences. */
@@ -252,9 +268,11 @@ public:
 	UPROPERTY(EditAnywhere, Category = "Database")
 	float LoopingCostBias = -0.005f;
 
+#if WITH_EDITORONLY_DATA
 	// These settings allow users to trim the start and end of animations in the database to preserve start/end frames for blending, and prevent the system from selecting the very last frames before it blends out.
 	UPROPERTY(EditAnywhere, Category = "Database")
 	FPoseSearchExcludeFromDatabaseParameters ExcludeFromDatabaseParameters;
+#endif // WITH_EDITORONLY_DATA
 
 	UPROPERTY(EditAnywhere, Category="Database")
 	TArray<FInstancedStruct> AnimationAssets;
@@ -282,18 +300,21 @@ public:
 	UPROPERTY(EditAnywhere, Category = "Performance")
 	EPoseSearchMode PoseSearchMode = EPoseSearchMode::PCAKDTree;
 
+#if WITH_EDITORONLY_DATA
 	// Number of dimensions used to create the kdtree. More dimensions allows a better explanation of the variance of the dataset that usually translates in better search results, but will imply more memory usage and worse performances.
 	UPROPERTY(EditAnywhere, Category = "Performance", meta = (EditCondition = "PoseSearchMode == EPoseSearchMode::PCAKDTree", EditConditionHides, ClampMin = "1", ClampMax = "64", UIMin = "1", UIMax = "64"))
 	int32 NumberOfPrincipalComponents = 4;
 
 	UPROPERTY(EditAnywhere, Category = "Performance", meta = (EditCondition = "PoseSearchMode == EPoseSearchMode::PCAKDTree", EditConditionHides, ClampMin = "1", ClampMax = "256", UIMin = "1", UIMax = "256"))
 	int32 KDTreeMaxLeafSize = 16;
+#endif // WITH_EDITORONLY_DATA
 	
 	// @todo: rename to KNNQueryNumNeighbors to be usable with the VPTree as well
 	// Out of a kdtree search, results will have only an approximate cost, so the database search will select the best “KDTree Query Num Neighbors” poses to perform the full cost analysis, and be able to elect the best pose.
 	UPROPERTY(EditAnywhere, Category = "Performance", meta = (DisplayName = "KNNQueryNumNeighbors", EditCondition = "PoseSearchMode == EPoseSearchMode::PCAKDTree || PoseSearchMode == VPTree", EditConditionHides, ClampMin = "1", ClampMax = "600", UIMin = "1"))
 	int32 KDTreeQueryNumNeighbors = 200;
 
+#if WITH_EDITORONLY_DATA
 	// if two poses values (multi dimensional point with the schema cardinality) are closer than PosePruningSimilarityThreshold,
 	// only one will be saved into the database FSearchIndexBase (to save memory) and accessed by the two different pose indexes
 	UPROPERTY(EditAnywhere, Category = "Performance")
@@ -303,6 +324,7 @@ public:
 	// only one will be saved into the database FSearchIndex (to save memory).
 	UPROPERTY(EditAnywhere, Category = "Performance", meta = (EditCondition = "PoseSearchMode == EPoseSearchMode::PCAKDTree", EditConditionHides))
 	float PCAValuesPruningSimilarityThreshold = 0.f;
+#endif // WITH_EDITORONLY_DATA
 
 	// @todo: rename to KNNQueryNumNeighborsWithDuplicates to be usable with the VPTree as well
 	// if PCAValuesPruningSimilarityThreshold > 0 the kdtree will remove duplicates, every result out of the KDTreeQueryNumNeighbors could potentially references multiple poses.
@@ -340,9 +362,6 @@ public:
 	const FPoseSearchDatabaseAnimationAssetBase* GetAnimationAssetBase(const UE::PoseSearch::FSearchIndexAsset& SearchIndexAsset) const;
 	FPoseSearchDatabaseAnimationAssetBase* GetMutableAnimationAssetBase(int32 AnimationAssetIndex);
 	FPoseSearchDatabaseAnimationAssetBase* GetMutableAnimationAssetBase(const UE::PoseSearch::FSearchIndexAsset& SearchIndexAsset);
-	const bool IsSourceAssetLooping(const UE::PoseSearch::FSearchIndexAsset& SearchIndexAsset) const;
-	const FString GetSourceAssetName(const UE::PoseSearch::FSearchIndexAsset& SearchIndexAsset) const;
-	int32 GetNumberOfPrincipalComponents() const;
 	float GetRealAssetTime(int32 PoseIdx) const;
 	float GetNormalizedAssetTime(int32 PoseIdx) const;
 
@@ -357,6 +376,8 @@ public:
 	FPoseSearchCost SearchContinuingPose(UE::PoseSearch::FSearchContext& SearchContext) const;
 
 #if WITH_EDITOR
+	int32 GetNumberOfPrincipalComponents() const;
+
 	virtual void BeginCacheForCookedPlatformData(const ITargetPlatform* TargetPlatform) override;
 	virtual bool IsCachedCookedPlatformDataLoaded(const ITargetPlatform* TargetPlatform) override;
 
diff --git a/Engine/Plugins/Experimental/Animation/PoseSearch/Source/Runtime/Public/PoseSearch/PoseSearchIndex.h b/Engine/Plugins/Experimental/Animation/PoseSearch/Source/Runtime/Public/PoseSearch/PoseSearchIndex.h
--- a/Engine/Plugins/Experimental/Animation/PoseSearch/Source/Runtime/Public/PoseSearch/PoseSearchIndex.h
+++ b/Engine/Plugins/Experimental/Animation/PoseSearch/Source/Runtime/Public/PoseSearch/PoseSearchIndex.h
@@ -87,31 +87,59 @@ struct FSearchIndexAsset
 
 	FSearchIndexAsset(
 		int32 InSourceAssetIdx,
-		int32 InFirstPoseIdx,
-		bool bInMirrored, 
-		const FFloatInterval& InSamplingInterval,
-		int32 SchemaSampleRate,
+		bool bInMirrored,
+		bool bInLooping,
+		bool bInDisableReselection,
 		int32 InPermutationIdx,
-		FVector InBlendParameters = FVector::Zero())
+		const FVector& InBlendParameters,
+		int32 InFirstPoseIdx,
+		int32 InFirstSampleIdx,
+		int32 InLastSampleIdx)
 		: SourceAssetIdx(InSourceAssetIdx)
 		, bMirrored(bInMirrored)
+		, bLooping(bInLooping)
+		, bDisableReselection(bInDisableReselection)
 		, PermutationIdx(InPermutationIdx)
-		, BlendParameters(InBlendParameters)
+		, BlendParameterX(InBlendParameters.X)
+		, BlendParameterY(InBlendParameters.Y)
 		, FirstPoseIdx(InFirstPoseIdx)
-		, FirstSampleIdx(FMath::CeilToInt(InSamplingInterval.Min * SchemaSampleRate))
-		, LastSampleIdx(FMath::FloorToInt(InSamplingInterval.Max * SchemaSampleRate))
+		, FirstSampleIdx(InFirstSampleIdx)
+		, LastSampleIdx(InLastSampleIdx)
 	{
-		check(SchemaSampleRate > 0);
+		check(FMath::IsNearlyZero(InBlendParameters.Z));
 	}
 
-	// Index of the source asset in search index's container (i.e. UPoseSearchDatabase)
-	int32 SourceAssetIdx = INDEX_NONE;
-	bool bMirrored = false;
-	int32 PermutationIdx = INDEX_NONE;
-	FVector BlendParameters = FVector::Zero();
-	int32 FirstPoseIdx = INDEX_NONE;
-	int32 FirstSampleIdx = INDEX_NONE;
-	int32 LastSampleIdx = INDEX_NONE;
+	FSearchIndexAsset(
+		int32 InSourceAssetIdx,
+		int32 InFirstPoseIdx,
+		bool bInMirrored,
+		bool bInLooping,
+		bool bInDisableReselection,
+		const FFloatInterval& InSamplingInterval,
+		int32 SchemaSampleRate,
+		int32 InPermutationIdx,
+		FVector InBlendParameters = FVector::Zero())
+		: FSearchIndexAsset(
+			InSourceAssetIdx,
+			bInMirrored,
+			bInLooping,
+			bInDisableReselection,
+			InPermutationIdx,
+			InBlendParameters,
+			InFirstPoseIdx,
+			FMath::CeilToInt(InSamplingInterval.Min * SchemaSampleRate),
+			FMath::FloorToInt(InSamplingInterval.Max * SchemaSampleRate))
+	{
+		check(SchemaSampleRate > 0);
+	}
+	
+	int32 GetSourceAssetIdx() const { return SourceAssetIdx; }
+	bool IsMirrored() const { return bMirrored; }
+	bool IsLooping() const { return bLooping; }
+	bool IsDisableReselection() const { return bDisableReselection; }
+	int32 GetPermutationIdx() const { return PermutationIdx; }
+	FVector GetBlendParameters() const { return FVector(BlendParameterX, BlendParameterY, 0.f); }
+	int32 GetFirstPoseIdx() const { return FirstPoseIdx; }
 
 	bool IsPoseInRange(int32 PoseIdx) const { return (PoseIdx >= FirstPoseIdx) && (PoseIdx < FirstPoseIdx + GetNumPoses()); }
 	bool operator==(const FSearchIndexAsset& Other) const;
@@ -135,13 +163,13 @@ struct FSearchIndexAsset
 	float GetFirstSampleTime(int32 SchemaSampleRate) const { check(SchemaSampleRate > 0); return FirstSampleIdx / float(SchemaSampleRate); }
 	float GetLastSampleTime(int32 SchemaSampleRate) const { check(SchemaSampleRate > 0); return LastSampleIdx / float(SchemaSampleRate); }
 
-	int32 GetPoseIndexFromTime(float Time, bool bIsLooping, int32 SchemaSampleRate) const
+	int32 GetPoseIndexFromTime(float Time, int32 SchemaSampleRate) const
 	{
 		check(IsInitialized());
 
 		const int32 NumPoses = GetNumPoses();
 		int32 PoseOffset = FMath::RoundToInt(SchemaSampleRate * Time) - FirstSampleIdx;
-		if (bIsLooping)
+		if (bLooping)
 		{
 			if (PoseOffset < 0)
 			{
@@ -174,6 +202,19 @@ struct FSearchIndexAsset
 	}
 
 	friend FArchive& operator<<(FArchive& Ar, FSearchIndexAsset& IndexAsset);
+
+private:
+	// Index of the source asset in search index's container (i.e. UPoseSearchDatabase)
+	const int32 SourceAssetIdx = INDEX_NONE;
+	const bool bMirrored : 1 = false; 
+	const bool bLooping : 1 = false;
+	const bool bDisableReselection : 1 = false;
+	const int32 PermutationIdx = INDEX_NONE;
+	const float BlendParameterX = 0.f;
+	const float BlendParameterY = 0.f;
+	const int32 FirstPoseIdx = INDEX_NONE;
+	const int32 FirstSampleIdx = INDEX_NONE;
+	const int32 LastSampleIdx = INDEX_NONE;
 };
 
 struct FSearchStats
@@ -395,6 +436,8 @@ struct FSearchIndex : public FSearchIndexBase
 	void Reset();
 	TConstArrayView<float> GetPoseValues(int32 PoseIdx) const;
 	TConstArrayView<float> GetReconstructedPoseValues(int32 PoseIdx, TArrayView<float> BufferUsedForReconstruction) const;
+	int32 GetNumDimensions() const;
+	int32 GetNumberOfPrincipalComponents() const;
 	POSESEARCH_API TConstArrayView<float> PCAProject(TConstArrayView<float> PoseValues, TArrayView<float> BufferUsedForProjection) const;
 
 	POSESEARCH_API TArray<float> GetPoseValuesSafe(int32 PoseIdx) const;
@@ -432,13 +475,13 @@ struct FVPTreeDataSource
 
     const TConstArrayView<float> operator[](int32 Index) const
     {
-		const int32 DataCardinality = SearchIndex.WeightsSqrt.Num();
+		const int32 DataCardinality = SearchIndex.GetNumDimensions();
         return SearchIndex.GetValuesVector(Index, DataCardinality);
     }
 
     int32 Num() const
     {
-		const int32 DataCardinality = SearchIndex.WeightsSqrt.Num();
+		const int32 DataCardinality = SearchIndex.GetNumDimensions();
         return SearchIndex.GetNumValuesVectors(DataCardinality);
     }
 
diff --git a/Engine/Plugins/Experimental/Animation/PoseSearch/Source/Runtime/Public/PoseSearch/PoseSearchSchema.h b/Engine/Plugins/Experimental/Animation/PoseSearch/Source/Runtime/Public/PoseSearch/PoseSearchSchema.h
--- a/Engine/Plugins/Experimental/Animation/PoseSearch/Source/Runtime/Public/PoseSearch/PoseSearchSchema.h
+++ b/Engine/Plugins/Experimental/Animation/PoseSearch/Source/Runtime/Public/PoseSearch/PoseSearchSchema.h
@@ -58,9 +58,11 @@ public:
 	UPROPERTY(EditAnywhere, Category = "Schema", meta = (DisplayPriority = 1))
 	TObjectPtr<UMirrorDataTable> MirrorDataTable;
 
+#if WITH_EDITORONLY_DATA
 	// Type of operation performed to the full pose features dataset
 	UPROPERTY(EditAnywhere, Category = "Schema", meta = (DisplayPriority = 2))
 	EPoseSearchDataPreprocessor DataPreprocessor = EPoseSearchDataPreprocessor::Normalize;
+#endif //WITH_EDITORONLY_DATA
 
 	UPROPERTY(Transient)
 	int32 SchemaCardinality = 0;
@@ -71,6 +73,7 @@ public:
 	UPROPERTY(Transient)
 	TArray<uint16> BoneIndicesWithParents;
 
+#if WITH_EDITORONLY_DATA
 	// How many times the animation assets of the database using this schema will be indexed.
 	UPROPERTY(EditAnywhere, Category = "Permutations", meta = (ClampMin = "1"))
 	int32 NumberOfPermutations = 1;
@@ -83,6 +86,7 @@ public:
 	// subsequent permutations will have PermutationTime = SamplingTime + PermutationsTimeOffset + PermutationIndex / PermutationsSampleRate.
 	UPROPERTY(EditAnywhere, Category = "Permutations")
 	float PermutationsTimeOffset = 0.f;
+#endif // WITH_EDITORONLY_DATA
 
 	// if true a padding channel will be added to make sure the data is 16 bytes (aligned) and padded, to facilitate performance improvements at cost of eventual additional memory
 	UPROPERTY(EditAnywhere, Category = "Performance")
