motion matching
- KDTree performance improvements when no filtering is needed (bRunNonSelectableIdxPostKDTree || NonSelectableIdx.IsEmpty())
- VPTree performance improvements
- QOL improvements

diff --git a/Engine/Plugins/Experimental/Animation/PoseSearch/Source/Runtime/Private/KDTree.cpp b/Engine/Plugins/Experimental/Animation/PoseSearch/Source/Runtime/Private/KDTree.cpp
--- a/Engine/Plugins/Experimental/Animation/PoseSearch/Source/Runtime/Private/KDTree.cpp
+++ b/Engine/Plugins/Experimental/Animation/PoseSearch/Source/Runtime/Private/KDTree.cpp
@@ -294,7 +294,8 @@ void FKDTree::Construct(AccessorType Count, AccessorType Dim, const float* Data,
 	new(this)FKDTree(Count, Dim, Data, MaxLeafSize);
 }
 
-bool FKDTree::FindNeighbors(FKNNResultSet& Result, TConstArrayView<float> Query) const
+template <typename RESULTSET>
+inline int32 FindNeighborsInternal(FKDTreeImplementation* Impl, RESULTSET& Result, TConstArrayView<float> Query)
 {
 #if UE_POSE_SEARCH_USE_NANOFLANN
 
@@ -306,36 +307,30 @@ bool FKDTree::FindNeighbors(FKNNResultSet& Result, TConstArrayView<float> Query)
 		32,			// Ignored parameter (Kept for compatibility with the FLANN interface).
 		0.f,		// search for eps-approximate neighbours (default: 0)
 		false);		// only for radius search, require neighbours sorted by
-	return Impl->findNeighbors(Result, Query.GetData(), SearchParams);
+	Impl->findNeighbors(Result, Query.GetData(), SearchParams);
+	return Result.Num();
 
 #else // UE_POSE_SEARCH_USE_NANOFLANN
 
 	checkNoEntry(); // unimplemented
-	return false;
+	return 0;
 
 #endif // UE_POSE_SEARCH_USE_NANOFLANN
 }
 
-bool FKDTree::FindNeighbors(FRadiusResultSet& Result, TConstArrayView<float> Query) const
+int32 FKDTree::FindNeighbors(FKNNResultSet& Result, TConstArrayView<float> Query) const
 {
-#if UE_POSE_SEARCH_USE_NANOFLANN
-
-	QUICK_SCOPE_CYCLE_COUNTER(STAT_FKDTree_FindNeighbors);
-
-	check(Query.GetData() && Query.Num() == Impl->dim && Impl->root_node);
-
-	const nanoflann::SearchParams SearchParams(
-		32,			// Ignored parameter (Kept for compatibility with the FLANN interface).
-		0.f,		// search for eps-approximate neighbours (default: 0)
-		false);		// only for radius search, require neighbours sorted by
-	return Impl->findNeighbors(Result, Query.GetData(), SearchParams);
-
-#else // UE_POSE_SEARCH_USE_NANOFLANN
+	return FindNeighborsInternal(Impl, Result, Query);
+}
 
-	checkNoEntry(); // unimplemented
-	return false;
+int32 FKDTree::FindNeighbors(FFilteredKNNResultSet& Result, TConstArrayView<float> Query) const
+{
+	return FindNeighborsInternal(Impl, Result, Query);
+}
 
-#endif // UE_POSE_SEARCH_USE_NANOFLANN
+int32 FKDTree::FindNeighbors(FRadiusResultSet& Result, TConstArrayView<float> Query) const
+{
+	return FindNeighborsInternal(Impl, Result, Query);
 }
 
 SIZE_T FKDTree::GetAllocatedSize() const
diff --git a/Engine/Plugins/Experimental/Animation/PoseSearch/Source/Runtime/Private/PoseSearchAssetSampler.cpp b/Engine/Plugins/Experimental/Animation/PoseSearch/Source/Runtime/Private/PoseSearchAssetSampler.cpp
--- a/Engine/Plugins/Experimental/Animation/PoseSearch/Source/Runtime/Private/PoseSearchAssetSampler.cpp
+++ b/Engine/Plugins/Experimental/Animation/PoseSearch/Source/Runtime/Private/PoseSearchAssetSampler.cpp
@@ -313,8 +313,14 @@ float FAnimationAssetSampler::ToNormalizedTime(float RealTime) const
 	if (CachedPlayLength > UE_KINDA_SMALL_NUMBER && Cast<UBlendSpace>(AnimationAsset.Get()))
 	{
 		const float NormalizedTime = RealTime / CachedPlayLength;
-		check(NormalizedTime >= 0.f && NormalizedTime <= 1.f);
-		return NormalizedTime;
+
+		if (NormalizedTime >= 0.f && NormalizedTime <= 1.f)
+		{
+			return NormalizedTime;
+		}
+
+		UE_LOG(LogPoseSearch, Error, TEXT("FAnimationAssetSampler::ToNormalizedTime: requested RealTime %f is greater than CachedPlayLength %f for UBlendSpace %s!"), RealTime, CachedPlayLength, *AnimationAsset->GetName());
+		return FMath::Clamp(NormalizedTime, 0.f, 1.f);
 	}
 
 	return RealTime;
diff --git a/Engine/Plugins/Experimental/Animation/PoseSearch/Source/Runtime/Private/PoseSearchDatabase.cpp b/Engine/Plugins/Experimental/Animation/PoseSearch/Source/Runtime/Private/PoseSearchDatabase.cpp
--- a/Engine/Plugins/Experimental/Animation/PoseSearch/Source/Runtime/Private/PoseSearchDatabase.cpp
+++ b/Engine/Plugins/Experimental/Animation/PoseSearch/Source/Runtime/Private/PoseSearchDatabase.cpp
@@ -1041,16 +1041,22 @@ UE::PoseSearch::FSearchResult UPoseSearchDatabase::SearchPCAKDTree(UE::PoseSearc
 		bRunNonSelectableIdxPostKDTree |= bValidateKNNSearch;
 #endif // WITH_EDITOR && ENABLE_ANIM_DEBUG
 
-		// if bRunNonSelectableIdxPostKDTree we filter out the NonSelectableIdx after kdtree search, otherwise during kdtree search
-		FKDTree::FKNNResultSet ResultSet(ClampedKDTreeQueryNumNeighbors, ResultIndexes, ResultDistanceSqr, bRunNonSelectableIdxPostKDTree ? TConstArrayView<int32>() : NonSelectableIdx);
-
 		check(QueryValues.Num() == NumDimensions);
 		// projecting QueryValues into the PCA space 
 		TConstArrayView<float> PCAQueryValues = SearchIndex.PCAProject(QueryValues, ProjectedQueryValues);
 		check(PCAQueryValues.Num() == ClampedNumberOfPrincipalComponents);
 
-		// Querying the KDTree with PCAQueryValues, the projected in PCA space QueryValues
-		SearchIndex.KDTree.FindNeighbors(ResultSet, PCAQueryValues);
+		int32 NumResults = 0;
+		if (bRunNonSelectableIdxPostKDTree || NonSelectableIdx.IsEmpty())
+		{
+			FKDTree::FKNNResultSet ResultSet(ClampedKDTreeQueryNumNeighbors, ResultIndexes, ResultDistanceSqr);
+			NumResults = SearchIndex.KDTree.FindNeighbors(ResultSet, PCAQueryValues);
+		}
+		else
+		{
+			FKDTree::FFilteredKNNResultSet ResultSet(ClampedKDTreeQueryNumNeighbors, ResultIndexes, ResultDistanceSqr, NonSelectableIdx);
+			NumResults = SearchIndex.KDTree.FindNeighbors(ResultSet, PCAQueryValues);
+		}
 
 #if WITH_EDITOR && ENABLE_ANIM_DEBUG
 		if (bValidateKNNSearch)
@@ -1072,7 +1078,7 @@ UE::PoseSearch::FSearchResult UPoseSearchDatabase::SearchPCAKDTree(UE::PoseSearc
 					return A.Value < B.Value;
 				});
 
-			for (int32 ResultIndex = 0; ResultIndex < ResultSet.Num(); ++ResultIndex)
+			for (int32 ResultIndex = 0; ResultIndex < NumResults; ++ResultIndex)
 			{
 				if (PCAValueIndexCost[ResultIndex].Key != ResultIndexes[ResultIndex])
 				{
@@ -1107,7 +1113,7 @@ UE::PoseSearch::FSearchResult UPoseSearchDatabase::SearchPCAKDTree(UE::PoseSearc
 			if (NumDimensions % 4 == 0)
 			{
 				int32 NumEvaluatePoseKernelCalls = 0;
-				for (int32 ResultIndex = 0; ResultIndex < ResultSet.Num(); ++ResultIndex)
+				for (int32 ResultIndex = 0; ResultIndex < NumResults; ++ResultIndex)
 				{
 					const TConstArrayView<int32> PoseIndexes = SearchIndex.PCAValuesVectorToPoseIndexes[ResultIndexes[ResultIndex]];
 					for (int32 Index = 0; Index < PoseIndexes.Num() && NumEvaluatePoseKernelCalls < MaxNumEvaluatePoseKernelCalls; ++Index, ++NumEvaluatePoseKernelCalls)
@@ -1119,7 +1125,7 @@ UE::PoseSearch::FSearchResult UPoseSearchDatabase::SearchPCAKDTree(UE::PoseSearc
 			else
 			{
 				int32 NumEvaluatePoseKernelCalls = 0;
-				for (int32 ResultIndex = 0; ResultIndex < ResultSet.Num(); ++ResultIndex)
+				for (int32 ResultIndex = 0; ResultIndex < NumResults; ++ResultIndex)
 				{
 					const TConstArrayView<int32> PoseIndexes = SearchIndex.PCAValuesVectorToPoseIndexes[ResultIndexes[ResultIndex]];
 					for (int32 Index = 0; Index < PoseIndexes.Num() && NumEvaluatePoseKernelCalls < MaxNumEvaluatePoseKernelCalls; ++Index, ++NumEvaluatePoseKernelCalls)
@@ -1135,7 +1141,7 @@ UE::PoseSearch::FSearchResult UPoseSearchDatabase::SearchPCAKDTree(UE::PoseSearc
 			// FMemory_Alloca is forced 16 bytes aligned
 			TArrayView<float> ReconstructedPoseValuesBuffer((float*)FMemory_Alloca(NumDimensions * sizeof(float)), NumDimensions);
 			check(IsAligned(ReconstructedPoseValuesBuffer.GetData(), alignof(VectorRegister4Float)));
-			for (int32 ResultIndex = 0; ResultIndex < ResultSet.Num(); ++ResultIndex)
+			for (int32 ResultIndex = 0; ResultIndex < NumResults; ++ResultIndex)
 			{
 				EvaluatePoseKernel<true, false>(Result, SearchIndex, QueryValues, ReconstructedPoseValuesBuffer, ResultIndexes[ResultIndex], SearchFilters, SearchContext, this, true, ResultIndex);
 			}
@@ -1143,7 +1149,7 @@ UE::PoseSearch::FSearchResult UPoseSearchDatabase::SearchPCAKDTree(UE::PoseSearc
 		// is the data padded at 16 bytes (and 16 bytes aligned by construction)?
 		else if (NumDimensions % 4 == 0)
 		{
-			for (int32 ResultIndex = 0; ResultIndex < ResultSet.Num(); ++ResultIndex)
+			for (int32 ResultIndex = 0; ResultIndex < NumResults; ++ResultIndex)
 			{
 				EvaluatePoseKernel<false, true>(Result, SearchIndex, QueryValues, TArrayView<float>(), ResultIndexes[ResultIndex], SearchFilters, SearchContext, this, true, ResultIndex);
 			}
@@ -1151,7 +1157,7 @@ UE::PoseSearch::FSearchResult UPoseSearchDatabase::SearchPCAKDTree(UE::PoseSearc
 		// no reconstruction, but data is not 16 bytes padded
 		else
 		{
-			for (int32 ResultIndex = 0; ResultIndex < ResultSet.Num(); ++ResultIndex)
+			for (int32 ResultIndex = 0; ResultIndex < NumResults; ++ResultIndex)
 			{
 				EvaluatePoseKernel<false, false>(Result, SearchIndex, QueryValues, TArrayView<float>(), ResultIndexes[ResultIndex], SearchFilters, SearchContext, this, true, ResultIndex);
 			}
@@ -1215,7 +1221,7 @@ UE::PoseSearch::FSearchResult UPoseSearchDatabase::SearchVPTree(UE::PoseSearch::
 		SearchIndex.VPTree.FindNeighbors(QueryValues, ResultSet, DataSource);
 		
 		int32 NumEvaluatePoseKernelCalls = 0;
-		const TArray<FIndexDistance>& UnsortedResults = ResultSet.GetUnsortedResults();
+		const TConstArrayView<FIndexDistance> UnsortedResults = ResultSet.GetUnsortedResults();
 
 		const bool bAreValuesPruned = SearchIndex.ValuesVectorToPoseIndexes.Num() > 0;
 		if (bAreValuesPruned)
diff --git a/Engine/Plugins/Experimental/Animation/PoseSearch/Source/Runtime/Private/PoseSearchLibrary.cpp b/Engine/Plugins/Experimental/Animation/PoseSearch/Source/Runtime/Private/PoseSearchLibrary.cpp
--- a/Engine/Plugins/Experimental/Animation/PoseSearch/Source/Runtime/Private/PoseSearchLibrary.cpp
+++ b/Engine/Plugins/Experimental/Animation/PoseSearch/Source/Runtime/Private/PoseSearchLibrary.cpp
@@ -618,6 +618,10 @@ UE::PoseSearch::FSearchResult UPoseSearchLibrary::MotionMatch(const FAnimationBa
 					{
 						PerDbAnimMap.FindOrAdd(PoseSearchBranchIn->Database).AddUnique(SequenceBase);
 					}
+					else
+					{
+						UE_LOG(LogPoseSearch, Error, TEXT("improperly setup UAnimNotifyState_PoseSearchBranchIn with null Database in %s"), *SequenceBase->GetName());
+					}
 				}
 			}
 		}
@@ -639,6 +643,8 @@ UE::PoseSearch::FSearchResult UPoseSearchLibrary::MotionMatch(const FAnimationBa
 
 		for (const FPerDbAnimPair& PerDbAnimPair : PerDbAnimMap)
 		{
+			check(PerDbAnimPair.Key);
+			
 			SearchContext.SetAnimationsToConsider(PerDbAnimPair.Value);
 
 			const FSearchResult NewSearchResult = PerDbAnimPair.Key->Search(SearchContext);
diff --git a/Engine/Plugins/Experimental/Animation/PoseSearch/Source/Runtime/Public/PoseSearch/KDTree.h b/Engine/Plugins/Experimental/Animation/PoseSearch/Source/Runtime/Public/PoseSearch/KDTree.h
--- a/Engine/Plugins/Experimental/Animation/PoseSearch/Source/Runtime/Public/PoseSearch/KDTree.h
+++ b/Engine/Plugins/Experimental/Animation/PoseSearch/Source/Runtime/Public/PoseSearch/KDTree.h
@@ -82,20 +82,18 @@ struct FKDTree
 
 	struct FKNNResultSet
 	{
-		inline FKNNResultSet(AccessorType InNumNeighbors, TArrayView<AccessorType> InIndexes, TArrayView<float> InDistances, TConstArrayView<AccessorType> InExcludeFromSearchIndexes = TConstArrayView<AccessorType>())
+		inline FKNNResultSet(AccessorType InNumNeighbors, TArrayView<AccessorType> InIndexes, TArrayView<float> InDistances, float InitValue = UE_BIG_NUMBER)
 		: Indexes(InIndexes)
 		, Distances(InDistances)
 		, NumNeighbors(InNumNeighbors)
 		, Count(0)
-		, ExcludeFromSearchIndexes(InExcludeFromSearchIndexes)
 		{
 			// by having IndexesView and DistancesView cardinality bigger than NumNeighbors, we can skip some if statements in the addPoint method
 			check(NumNeighbors > 0);
 			check(InIndexes.Num() > NumNeighbors);
 			check(InDistances.Num() > NumNeighbors);
-			check(Algo::IsSorted(InExcludeFromSearchIndexes));
 
-			Distances[NumNeighbors - 1] = UE_BIG_NUMBER;
+			Distances[NumNeighbors - 1] = InitValue;
 		}
 
 		inline AccessorType Num() const
@@ -110,11 +108,6 @@ struct FKDTree
 
 		inline bool addPoint(float dist, AccessorType index)
 		{
-			if (Algo::BinarySearch(ExcludeFromSearchIndexes, index) != INDEX_NONE)
-			{
-				return true;
-			}
-
 			// shifting Distances[i] and Indexes[i] to make space for "dist" and "index" at the right "i"th slot
 			AccessorType i;
 			for (i = Count; (i > 0) && (Distances[i - 1] > dist); --i)
@@ -140,75 +133,50 @@ struct FKDTree
 
 		inline float worstDist() const { return Distances[NumNeighbors - 1]; }
 
-	private:
+	protected:
 		TArrayView<AccessorType> Indexes;
 		TArrayView<float> Distances;
 		AccessorType NumNeighbors;
 		AccessorType Count;
-		TConstArrayView<AccessorType> ExcludeFromSearchIndexes; // sorted array view
 	};
 
-	struct FRadiusResultSet
+	struct FFilteredKNNResultSet : public FKNNResultSet
 	{
-		inline FRadiusResultSet(float Radius, AccessorType MaxNumNeighbors, TArrayView<AccessorType> InIndexes, TArrayView<float> InDistances)
-		: Indexes(InIndexes)
-		, Distances(InDistances)
-		, NumNeighbors(MaxNumNeighbors)
-		, Count(0)
+		inline FFilteredKNNResultSet(AccessorType InNumNeighbors, TArrayView<AccessorType> InIndexes, TArrayView<float> InDistances, TConstArrayView<AccessorType> InExcludeFromSearchIndexes = TConstArrayView<AccessorType>())
+		: FKNNResultSet(InNumNeighbors, InIndexes, InDistances)
+		, ExcludeFromSearchIndexes(InExcludeFromSearchIndexes)
 		{
-			// by having IndexesView and DistancesView cardinality bigger than NumNeighbors, we can skip some if statements in the addPoint method
-			check(NumNeighbors > 0);
-			check(InIndexes.Num() > NumNeighbors);
-			check(InDistances.Num() > NumNeighbors);
-
-			Distances[NumNeighbors - 1] = Radius;
+			check(Algo::IsSorted(InExcludeFromSearchIndexes));
 		}
 
-		inline AccessorType Num() const
+		inline bool addPoint(float dist, AccessorType index)
 		{
-			return Count;
+			if (Algo::BinarySearch(ExcludeFromSearchIndexes, index) == INDEX_NONE)
+			{
+				FKNNResultSet::addPoint(dist, index);
+			}
+			return true;
 		}
 
-		inline bool full() const
+	protected:
+		TConstArrayView<AccessorType> ExcludeFromSearchIndexes; // sorted array view
+	};
+
+	struct FRadiusResultSet : public FKNNResultSet
+	{
+		inline FRadiusResultSet(float Radius, AccessorType InNumNeighbors, TArrayView<AccessorType> InIndexes, TArrayView<float> InDistances)
+		: FKNNResultSet(InNumNeighbors, InIndexes, InDistances, Radius)
 		{
-			return Count == NumNeighbors;
 		}
 
 		inline bool addPoint(float dist, AccessorType index)
 		{
 			if (dist < worstDist())
 			{
-				// shifting Distances[i] and Indexes[i] to make space for "dist" and "index" at the right "i"th slot
-				AccessorType i;
-				for (i = Count; (i > 0) && (Distances[i - 1] > dist); --i)
-				{
-					// no need to check "if (i < capacity)" since dists and indices can contains more items than capacity_ 
-					Distances[i] = Distances[i - 1];
-					Indexes[i] = Indexes[i - 1];
-				}
-
-				// inserting "dist" and "index" in a sorted manner
-				// no need to check "if (i < capacity)" since dists and indices can contains more items than capacity_ 
-				Distances[i] = dist;
-				Indexes[i] = index;
-
-				if (Count < NumNeighbors)
-				{
-					Count++;
-				}
+				FKNNResultSet::addPoint(dist, index);
 			}
-			
-			// tell caller that the search shall continue
 			return true;
 		}
-
-		inline float worstDist() const { return Distances[NumNeighbors - 1]; }
-
-	private:
-		TArrayView<AccessorType> Indexes;
-		TArrayView<float> Distances;
-		AccessorType NumNeighbors;
-		AccessorType Count;
 	};
 	
 
@@ -225,8 +193,11 @@ struct FKDTree
 	
 	void Reset();
 	void Construct(AccessorType Count, AccessorType Dim, const float* Data, AccessorType MaxLeafSize = 16);
-	bool FindNeighbors(FKNNResultSet& Result, TConstArrayView<float> Query) const;
-	bool FindNeighbors(FRadiusResultSet& Result, TConstArrayView<float> Query) const;
+
+	int32 FindNeighbors(FKNNResultSet& Result, TConstArrayView<float> Query) const;
+	int32 FindNeighbors(FFilteredKNNResultSet& Result, TConstArrayView<float> Query) const;
+	int32 FindNeighbors(FRadiusResultSet& Result, TConstArrayView<float> Query) const;
+
 	POSESEARCH_API SIZE_T GetAllocatedSize() const;
 
 	FDataSource DataSource;
diff --git a/Engine/Plugins/Experimental/Animation/PoseSearch/Source/Runtime/Public/PoseSearch/VPTree.h b/Engine/Plugins/Experimental/Animation/PoseSearch/Source/Runtime/Public/PoseSearch/VPTree.h
--- a/Engine/Plugins/Experimental/Animation/PoseSearch/Source/Runtime/Public/PoseSearch/VPTree.h
+++ b/Engine/Plugins/Experimental/Animation/PoseSearch/Source/Runtime/Public/PoseSearch/VPTree.h
@@ -79,6 +79,12 @@ struct FVPTreeResultSet
     FVPTreeResultSet(int32 InNumNeighbors)
         : NumNeighbors(InNumNeighbors)
     {
+#if !NO_LOGGING
+		if (InNumNeighbors > HeapDefaultMaxSize)
+		{
+			UE_LOG(LogPoseSearch, Warning, TEXT("FVPTreeResultSet - preallocated Heap data size 'HeapDefaultMaxSize' of %d is less than requested %d num neighbors. Performances will be negatively impacted"), HeapDefaultMaxSize, InNumNeighbors);
+		}
+#endif // !NO_LOGGING
     }
 
     VPTreeScalar GetWorstDistance() const
@@ -140,7 +146,7 @@ struct FVPTreeResultSet
 		return Results;
     }
 
-	const TArray<FIndexDistance>& GetUnsortedResults() const
+	TConstArrayView<FIndexDistance> GetUnsortedResults() const
 	{
 		return Heap;
 	}
@@ -161,7 +167,8 @@ private:
 	};
 
     const int32 NumNeighbors;
-    TArray<FIndexDistance> Heap;
+	enum { HeapDefaultMaxSize = 256 };
+    TArray<FIndexDistance, TInlineAllocator<HeapDefaultMaxSize>> Heap;
 };
 
 struct FVPTree
@@ -227,8 +234,8 @@ struct FVPTree
 		bool bBestNodeIndexFound = false;
 #endif // VALIDATE_FINDNEIGHBORS
 
-		// @todo: avoid allocations
-		TArray<int32> NodesToSearch;
+		enum { NodesToSearchDefaultMaxSize = 512 };
+		TArray<int32, TInlineAllocator<NodesToSearchDefaultMaxSize>> NodesToSearch;
 		NodesToSearch.Emplace(0);
 
         VPTreeScalar Tau = VPTreeScalar(UE_BIG_NUMBER);
@@ -304,6 +311,12 @@ struct FVPTree
 			if (Node.LeftIndex != INDEX_NONE && QueryDistance < Node.Distance + Tau)
 			{
 				NodesToSearch.Emplace(Node.LeftIndex);
+#if !NO_LOGGING
+				if (NodesToSearch.Num() > NodesToSearchDefaultMaxSize)
+				{
+					UE_LOG(LogPoseSearch, Warning, TEXT("FVPTree::FindNeighbors - requested more than preallocated NodesToSearch data size 'NodesToSearchDefaultMaxSize' (%d / %d)"), NodesToSearch.Num(), NodesToSearchDefaultMaxSize);
+				}
+#endif // !NO_LOGGING
 
 #if VALIDATE_FINDNEIGHBORS
 				if (bNodeIndexContainsBestNodeIndex)
@@ -319,6 +332,12 @@ struct FVPTree
 			if (Node.RightIndex != INDEX_NONE && QueryDistance >= Node.Distance - Tau)
 			{
 				NodesToSearch.Emplace(Node.RightIndex);
+#if !NO_LOGGING
+				if (NodesToSearch.Num() > NodesToSearchDefaultMaxSize)
+				{
+					UE_LOG(LogPoseSearch, Warning, TEXT("FVPTree::FindNeighbors - requested more than preallocated NodesToSearch data size 'NodesToSearchDefaultMaxSize' (%d / %d)"), NodesToSearch.Num(), NodesToSearchDefaultMaxSize);
+				}
+#endif // !NO_LOGGING
 
 #if VALIDATE_FINDNEIGHBORS
 				if (bNodeIndexContainsBestNodeIndex)
