motion matching
- fix pose search debugger "show all candidates" for database skipped by the search
- improved best pose candidates memory and performance requirements

diff --git a/Engine/Plugins/Experimental/Animation/PoseSearch/Source/Runtime/Private/PoseSearchDatabase.cpp b/Engine/Plugins/Experimental/Animation/PoseSearch/Source/Runtime/Private/PoseSearchDatabase.cpp
--- a/Engine/Plugins/Experimental/Animation/PoseSearch/Source/Runtime/Private/PoseSearchDatabase.cpp
+++ b/Engine/Plugins/Experimental/Animation/PoseSearch/Source/Runtime/Private/PoseSearchDatabase.cpp
@@ -94,7 +94,7 @@ static void PopulateNonSelectableIdx(FNonSelectableIdx& NonSelectableIdx, FSearc
 #if UE_POSE_SEARCH_TRACE_ENABLED
 							const TArray<float> PoseValues = SearchIndex.GetPoseValuesSafe(PoseIdx);
 							const FPoseSearchCost PoseCost = SearchIndex.ComparePoses(PoseIdx, 0.f, PoseValues, QueryValues);
-							SearchContext.BestCandidates.Add(PoseCost, PoseIdx, Database, EPoseCandidateFlags::DiscardedBy_PoseJumpThresholdTime);
+							SearchContext.Track(Database, PoseIdx, EPoseCandidateFlags::DiscardedBy_PoseJumpThresholdTime, PoseCost);
 #endif // UE_POSE_SEARCH_TRACE_ENABLED
 						}
 					}
@@ -122,7 +122,7 @@ static void PopulateNonSelectableIdx(FNonSelectableIdx& NonSelectableIdx, FSearc
 #if UE_POSE_SEARCH_TRACE_ENABLED
 						const TArray<float> PoseValues = SearchIndex.GetPoseValuesSafe(PoseIdx);
 						const FPoseSearchCost PoseCost = SearchIndex.ComparePoses(PoseIdx, 0.f, PoseValues, QueryValues);
-						SearchContext.BestCandidates.Add(PoseCost, PoseIdx, Database, EPoseCandidateFlags::DiscardedBy_PoseJumpThresholdTime);
+						SearchContext.Track(Database, PoseIdx, EPoseCandidateFlags::DiscardedBy_PoseJumpThresholdTime, PoseCost);
 #endif // UE_POSE_SEARCH_TRACE_ENABLED
 					}
 				}
@@ -138,7 +138,7 @@ static void PopulateNonSelectableIdx(FNonSelectableIdx& NonSelectableIdx, FSearc
 #if UE_POSE_SEARCH_TRACE_ENABLED
 						const TArray<float> PoseValues = SearchIndex.GetPoseValuesSafe(PoseIdx);
 						const FPoseSearchCost PoseCost = SearchIndex.ComparePoses(PoseIdx, 0.f, PoseValues, QueryValues);
-						SearchContext.BestCandidates.Add(PoseCost, PoseIdx, Database, EPoseCandidateFlags::DiscardedBy_PoseJumpThresholdTime);
+						SearchContext.Track(Database, PoseIdx, EPoseCandidateFlags::DiscardedBy_PoseJumpThresholdTime, PoseCost);
 #endif // UE_POSE_SEARCH_TRACE_ENABLED
 					}
 				}
@@ -163,7 +163,7 @@ static void PopulateNonSelectableIdx(FNonSelectableIdx& NonSelectableIdx, FSearc
 				if (HistoricalPoseIndex.PoseIndex < SearchIndex.GetNumPoses())
 				{
 					const FPoseSearchCost PoseCost = SearchIndex.ComparePoses(HistoricalPoseIndex.PoseIndex, 0.f, SearchIndex.GetPoseValuesSafe(HistoricalPoseIndex.PoseIndex), QueryValues);
-					SearchContext.BestCandidates.Add(PoseCost, HistoricalPoseIndex.PoseIndex, Database, EPoseCandidateFlags::DiscardedBy_PoseReselectHistory);
+					SearchContext.Track(Database, HistoricalPoseIndex.PoseIndex, EPoseCandidateFlags::DiscardedBy_PoseReselectHistory, PoseCost);
 				}
 #endif // UE_POSE_SEARCH_TRACE_ENABLED
 			}
@@ -219,17 +219,17 @@ struct FSearchFilters
 				else if (Filter == &SelectableAssetIdxFilter)
 				{
 					const FPoseSearchCost PoseCost = SearchIndex.ComparePoses(PoseIdx, 0.f, PoseValues, QueryValues);
-					SearchContext.BestCandidates.Add(PoseCost, PoseIdx, Database, EPoseCandidateFlags::DiscardedBy_AssetIdxFilter);
+					SearchContext.Track(Database, PoseIdx, EPoseCandidateFlags::DiscardedBy_AssetIdxFilter, PoseCost);
 				}
 				else if (Filter == &BlockTransitionFilter)
 				{
 					const FPoseSearchCost PoseCost = SearchIndex.ComparePoses(PoseIdx, 0.f, PoseValues, QueryValues);
-					SearchContext.BestCandidates.Add(PoseCost, PoseIdx, Database, EPoseCandidateFlags::DiscardedBy_BlockTransition);
+					SearchContext.Track(Database, PoseIdx, EPoseCandidateFlags::DiscardedBy_BlockTransition, PoseCost);
 				}
 				else
 				{
 					const FPoseSearchCost PoseCost = SearchIndex.ComparePoses(PoseIdx, 0.f, PoseValues, QueryValues);
-					SearchContext.BestCandidates.Add(PoseCost, PoseIdx, Database, EPoseCandidateFlags::DiscardedBy_PoseFilter);
+					SearchContext.Track(Database, PoseIdx, EPoseCandidateFlags::DiscardedBy_PoseFilter, PoseCost);
 				}
 #endif // UE_POSE_SEARCH_TRACE_ENABLED
 				return false;
@@ -892,6 +892,11 @@ UE::PoseSearch::FSearchResult UPoseSearchDatabase::Search(UE::PoseSearch::FSearc
 #endif // WITH_EDITOR && ENABLE_ANIM_DEBUG
 	}
 
+#if UE_POSE_SEARCH_TRACE_ENABLED
+	// in case we skipped the search, or we didn't find any candidates we still have to track we requested to evaluate this database, so we keep track of this
+	SearchContext.Track(this);
+#endif // UE_POSE_SEARCH_TRACE_ENABLED
+
 	return Result;
 }
 
@@ -926,7 +931,7 @@ static inline void EvaluatePoseKernel(UE::PoseSearch::FSearchResult& Result, con
 #if UE_POSE_SEARCH_TRACE_ENABLED
 		if (bUpdateBestCandidates)
 		{
-			SearchContext.BestCandidates.Add(PoseCost, PoseIdx, Database, EPoseCandidateFlags::Valid_Pose);
+			SearchContext.Track(Database, PoseIdx, EPoseCandidateFlags::Valid_Pose, PoseCost);
 		}
 #endif // UE_POSE_SEARCH_TRACE_ENABLED
 	}
@@ -992,7 +997,7 @@ FPoseSearchCost UPoseSearchDatabase::SearchContinuingPose(UE::PoseSearch::FSearc
 		}
 
 #if UE_POSE_SEARCH_TRACE_ENABLED
-		SearchContext.BestCandidates.Add(ContinuingPoseCost, ContinuingPoseIdx, this, EPoseCandidateFlags::Valid_ContinuingPose);
+		SearchContext.Track(this, ContinuingPoseIdx, EPoseCandidateFlags::Valid_ContinuingPose, ContinuingPoseCost);
 #endif // UE_POSE_SEARCH_TRACE_ENABLED
 	}
 
diff --git a/Engine/Plugins/Experimental/Animation/PoseSearch/Source/Runtime/Private/PoseSearchLibrary.cpp b/Engine/Plugins/Experimental/Animation/PoseSearch/Source/Runtime/Private/PoseSearchLibrary.cpp
--- a/Engine/Plugins/Experimental/Animation/PoseSearch/Source/Runtime/Private/PoseSearchLibrary.cpp
+++ b/Engine/Plugins/Experimental/Animation/PoseSearch/Source/Runtime/Private/PoseSearchLibrary.cpp
@@ -148,57 +148,46 @@ void UPoseSearchLibrary::TraceMotionMatchingState(
 {
 	using namespace UE::PoseSearch;
 	
-	auto AddUniqueDatabase = [](TArray<FTraceMotionMatchingStateDatabaseEntry>& DatabaseEntries, const UPoseSearchDatabase* Database, UE::PoseSearch::FSearchContext& SearchContext) -> int32
-	{
-		const uint64 DatabaseId = FTraceMotionMatchingState::GetIdFromObject(Database);
-
-		int32 DbEntryIdx = INDEX_NONE;
-		for (int32 i = 0; i < DatabaseEntries.Num(); ++i)
-		{
-			if (DatabaseEntries[i].DatabaseId == DatabaseId)
-			{
-				DbEntryIdx = i;
-				break;
-			}
-		}
-		if (DbEntryIdx == INDEX_NONE)
-		{
-			DbEntryIdx = DatabaseEntries.Add({ DatabaseId });
-
-			// if throttling is on, the continuing pose can be valid, but no actual search occurred, so the query will not be cached, and we need to build it
-			DatabaseEntries[DbEntryIdx].QueryVector = SearchContext.GetOrBuildQuery(Database->Schema).GetValues();
-		}
-
-		return DbEntryIdx;
-	};
-
 	const int32 CurrentPoseIdx = bSearch && CurrentResult.PoseCost.IsValid() ? CurrentResult.PoseIdx : INDEX_NONE;
 	FTraceMotionMatchingState TraceState;
-	while (!SearchContext.BestCandidates.IsEmpty()) 
+	TraceState.DatabaseEntries.SetNum(SearchContext.GetBestPoseCandidatesMap().Num());
+	
+	int32 DbEntryIdx = 0;
+	for (TPair<const UPoseSearchDatabase*, FSearchContext::FBestPoseCandidates> DatabaseBestPoseCandidates : SearchContext.GetBestPoseCandidatesMap())
 	{
-		FSearchContext::FPoseCandidate PoseCandidate;
-		SearchContext.BestCandidates.Pop(PoseCandidate);
+		const UPoseSearchDatabase* Database = DatabaseBestPoseCandidates.Key;
+		check(Database);
 
-		const int32 DbEntryIdx = AddUniqueDatabase(TraceState.DatabaseEntries, PoseCandidate.Database, SearchContext);
 		FTraceMotionMatchingStateDatabaseEntry& DbEntry = TraceState.DatabaseEntries[DbEntryIdx];
 
-		FTraceMotionMatchingStatePoseEntry PoseEntry;
-		PoseEntry.DbPoseIdx = PoseCandidate.PoseIdx;
-		PoseEntry.Cost = PoseCandidate.Cost;
-		PoseEntry.PoseCandidateFlags = PoseCandidate.PoseCandidateFlags;
-		if (CurrentPoseIdx == PoseCandidate.PoseIdx && CurrentResult.Database.Get() == PoseCandidate.Database)
-		{
-			check(EnumHasAnyFlags(PoseEntry.PoseCandidateFlags, EPoseCandidateFlags::Valid_Pose | EPoseCandidateFlags::Valid_ContinuingPose));
+		// if throttling is on, the continuing pose can be valid, but no actual search occurred, so the query will not be cached, and we need to build it
+		DbEntry.QueryVector = SearchContext.GetOrBuildQuery(Database->Schema).GetValues();
+		DbEntry.DatabaseId = FTraceMotionMatchingState::GetIdFromObject(Database);
 
-			EnumAddFlags(PoseEntry.PoseCandidateFlags, EPoseCandidateFlags::Valid_CurrentPose);
-			
-			TraceState.CurrentDbEntryIdx = DbEntryIdx;
-			TraceState.CurrentPoseEntryIdx = DbEntry.PoseEntries.Add(PoseEntry);
-		}
-		else
+		for (int32 CandidateIdx = 0; CandidateIdx < DatabaseBestPoseCandidates.Value.Num(); ++CandidateIdx)
 		{
-			DbEntry.PoseEntries.Add(PoseEntry);
+			const FSearchContext::FPoseCandidate PoseCandidate = DatabaseBestPoseCandidates.Value.GetUnsortedCandidate(CandidateIdx);
+
+			FTraceMotionMatchingStatePoseEntry PoseEntry;
+			PoseEntry.DbPoseIdx = PoseCandidate.PoseIdx;
+			PoseEntry.Cost = PoseCandidate.Cost;
+			PoseEntry.PoseCandidateFlags = PoseCandidate.PoseCandidateFlags;
+			if (CurrentPoseIdx == PoseCandidate.PoseIdx && CurrentResult.Database.Get() == Database)
+			{
+				check(EnumHasAnyFlags(PoseEntry.PoseCandidateFlags, EPoseCandidateFlags::Valid_Pose | EPoseCandidateFlags::Valid_ContinuingPose));
+
+				EnumAddFlags(PoseEntry.PoseCandidateFlags, EPoseCandidateFlags::Valid_CurrentPose);
+
+				TraceState.CurrentDbEntryIdx = DbEntryIdx;
+				TraceState.CurrentPoseEntryIdx = DbEntry.PoseEntries.Add(PoseEntry);
+			}
+			else
+			{
+				DbEntry.PoseEntries.Add(PoseEntry);
+			}
 		}
+
+		++DbEntryIdx;
 	}
 
 	if (DeltaTime > SMALL_NUMBER)
diff --git a/Engine/Plugins/Experimental/Animation/PoseSearch/Source/Runtime/Public/PoseSearch/PoseSearchContext.h b/Engine/Plugins/Experimental/Animation/PoseSearch/Source/Runtime/Public/PoseSearch/PoseSearchContext.h
--- a/Engine/Plugins/Experimental/Animation/PoseSearch/Source/Runtime/Public/PoseSearch/PoseSearchContext.h
+++ b/Engine/Plugins/Experimental/Animation/PoseSearch/Source/Runtime/Public/PoseSearch/PoseSearchContext.h
@@ -224,24 +224,9 @@ private:
 	
 #if UE_POSE_SEARCH_TRACE_ENABLED
 
-	struct FPoseCandidateId
+	struct FPoseCandidateIdCost
 	{
 		int32 PoseIdx = 0;
-		const UPoseSearchDatabase* Database = nullptr;
-
-		bool operator==(const FPoseCandidateId& Other) const
-		{
-			return PoseIdx == Other.PoseIdx && Database == Other.Database;
-		}
-		
-		friend uint32 GetTypeHash(const FPoseCandidateId& PoseCandidateId)
-		{
-			return HashCombineFast(GetTypeHash(PoseCandidateId.PoseIdx), GetTypeHash(PoseCandidateId.Database));
-		}
-	};
-
-	struct FPoseCandidateIdCost : public FPoseCandidateId
-	{
 		FPoseSearchCost Cost;
 		bool operator<(const FPoseCandidateIdCost& Other) const { return Other.Cost < Cost; } // Reverse compare because BestCandidates is a max heap
 	};
@@ -253,30 +238,39 @@ public:
 		EPoseCandidateFlags PoseCandidateFlags = EPoseCandidateFlags::None;
 	};
 
+	void Track(const UPoseSearchDatabase* Database, int32 PoseIdx = INDEX_NONE, EPoseCandidateFlags PoseCandidateFlags = EPoseCandidateFlags::None, const FPoseSearchCost& Cost = FPoseSearchCost())
+	{
+		check(Database);
+
+		FBestPoseCandidates& BestPoseCandidates = BestPoseCandidatesMap.FindOrAdd(Database);
+		if (PoseIdx != INDEX_NONE)
+		{
+			BestPoseCandidates.Add(PoseIdx, PoseCandidateFlags, Cost);
+		}
+	}
+
+
 	struct FBestPoseCandidates
 	{
 		FBestPoseCandidates()
 		{
 			// preallocating memory to avoid multiple reallocations / rehashing
-			PoseCandidateHeap.Reserve(MaxPoseCandidates);
-			PoseIdxToFlags.Empty(MaxPoseCandidates);
+			PoseCandidateHeap.Reserve(MaxNumberOfCollectedPoseCandidatesPerDatabase);
+			PoseIdxToFlags.Empty(MaxNumberOfCollectedPoseCandidatesPerDatabase);
 		}
 
-		void Add(const FPoseSearchCost& Cost, int32 PoseIdx, const UPoseSearchDatabase* Database, EPoseCandidateFlags PoseCandidateFlags)
+		void Add(int32 PoseIdx, EPoseCandidateFlags PoseCandidateFlags, const FPoseSearchCost& Cost)
 		{
-			FPoseCandidate PoseCandidate;
-			PoseCandidate.PoseIdx = PoseIdx;
-			PoseCandidate.Database = Database;
-
-			if (EPoseCandidateFlags* PoseIdxPoseCandidateFlags = PoseIdxToFlags.Find(PoseCandidate))
+			check(PoseIdx >= 0);
+			if (EPoseCandidateFlags* PoseIdxPoseCandidateFlags = PoseIdxToFlags.Find(PoseIdx))
 			{
 				*PoseIdxPoseCandidateFlags |= PoseCandidateFlags;
 			}
-			else if (PoseCandidateHeap.Num() < MaxPoseCandidates || Cost < PoseCandidateHeap.HeapTop().Cost)
+			else if (PoseCandidateHeap.Num() < MaxNumberOfCollectedPoseCandidatesPerDatabase || Cost < PoseCandidateHeap.HeapTop().Cost)
 			{
 				bool bPoppedContinuingPoseCandidate = false;
 				FPoseCandidate ContinuingPoseCandidate;
-				while (PoseCandidateHeap.Num() >= MaxPoseCandidates)
+				while (PoseCandidateHeap.Num() >= MaxNumberOfCollectedPoseCandidatesPerDatabase)
 				{
 					FPoseCandidate PoppedPoseCandidate;
 					Pop(PoppedPoseCandidate);
@@ -296,38 +290,50 @@ public:
 					FPoseCandidate PoppedPoseCandidate;
 					Pop(PoppedPoseCandidate);
 					PoseCandidateHeap.HeapPush(ContinuingPoseCandidate);
-					PoseIdxToFlags.Add(ContinuingPoseCandidate, ContinuingPoseCandidate.PoseCandidateFlags);
+					PoseIdxToFlags.Add(ContinuingPoseCandidate.PoseIdx, ContinuingPoseCandidate.PoseCandidateFlags);
 				}
 
+				FPoseCandidate PoseCandidate;
+				PoseCandidate.PoseIdx = PoseIdx;
 				PoseCandidate.Cost = Cost;
 				PoseCandidateHeap.HeapPush(PoseCandidate);
-				PoseIdxToFlags.Add(PoseCandidate, PoseCandidateFlags);
+				PoseIdxToFlags.Add(PoseIdx, PoseCandidateFlags);
 			}
 		}
 
-		void Pop(FPoseCandidate& OutItem)
+		int32 Num() const
 		{
-			PoseCandidateHeap.HeapPop(OutItem, false);
-			OutItem.PoseCandidateFlags = PoseIdxToFlags.FindAndRemoveChecked(OutItem);
+			return PoseCandidateHeap.Num();
 		}
 
-		bool IsEmpty() const
+		FPoseCandidate GetUnsortedCandidate(int32 Index) const
 		{
-			return PoseCandidateHeap.IsEmpty();
+			FPoseCandidate PoseCandidate;
+			const FPoseCandidateIdCost& PoseCandidateIdCost = PoseCandidateHeap[Index];
+			PoseCandidate.PoseIdx = PoseCandidateIdCost.PoseIdx;
+			PoseCandidate.Cost = PoseCandidateIdCost.Cost;
+			PoseCandidate.PoseCandidateFlags = PoseIdxToFlags[PoseCandidateIdCost.PoseIdx];
+			return PoseCandidate;
 		}
 
-		void SetMaxPoseCandidates(int32 Value)
+	private:
+		void Pop(FPoseCandidate& OutItem)
 		{
-			MaxPoseCandidates = Value;
+			PoseCandidateHeap.HeapPop(OutItem, false);
+			OutItem.PoseCandidateFlags = PoseIdxToFlags.FindAndRemoveChecked(OutItem.PoseIdx);
 		}
 
-	private:
 		TArray<FPoseCandidateIdCost> PoseCandidateHeap;
-		TMap<FPoseCandidateId, EPoseCandidateFlags> PoseIdxToFlags;
-		int32 MaxPoseCandidates = 200;
+		TMap<int32, EPoseCandidateFlags> PoseIdxToFlags;
 	};
 	
-	FBestPoseCandidates BestCandidates;
+	const TMap<const UPoseSearchDatabase*, FBestPoseCandidates>& GetBestPoseCandidatesMap() const
+	{
+		return BestPoseCandidatesMap;
+	}
+
+private:
+	TMap<const UPoseSearchDatabase*, FBestPoseCandidates> BestPoseCandidatesMap;
 #endif // UE_POSE_SEARCH_TRACE_ENABLED
 };
 
diff --git a/Engine/Plugins/Experimental/Animation/PoseSearch/Source/Runtime/Public/PoseSearch/PoseSearchDefines.h b/Engine/Plugins/Experimental/Animation/PoseSearch/Source/Runtime/Public/PoseSearch/PoseSearchDefines.h
--- a/Engine/Plugins/Experimental/Animation/PoseSearch/Source/Runtime/Public/PoseSearch/PoseSearchDefines.h
+++ b/Engine/Plugins/Experimental/Animation/PoseSearch/Source/Runtime/Public/PoseSearch/PoseSearchDefines.h
@@ -28,8 +28,8 @@ static constexpr FBoneIndexType RootBoneIndexType = 0;
 static constexpr FBoneIndexType ComponentSpaceIndexType = FBoneIndexType(-1);
 
 static constexpr int32 PreallocatedCachedQueriesNum = 8;
-
 static constexpr float FiniteDelta = 1 / 60.0f; // Time delta used for computing pose derivatives
+static constexpr int32 MaxNumberOfCollectedPoseCandidatesPerDatabase = 200;
 
 constexpr int32 TMax(int32 A, int32 B) { return (A > B ? A : B); }
 template<typename ElementType> constexpr int32 TAlignOf() { return TMax(alignof(ElementType), 16); }
