motion matching
- debugger channels breakdown labels improvements
- database Data Details tab labels improvements

diff --git a/Engine/Plugins/Experimental/Animation/PoseSearch/Source/Editor/Private/PoseSearchDatabaseDataDetails.cpp b/Engine/Plugins/Experimental/Animation/PoseSearch/Source/Editor/Private/PoseSearchDatabaseDataDetails.cpp
--- a/Engine/Plugins/Experimental/Animation/PoseSearch/Source/Editor/Private/PoseSearchDatabaseDataDetails.cpp
+++ b/Engine/Plugins/Experimental/Animation/PoseSearch/Source/Editor/Private/PoseSearchDatabaseDataDetails.cpp
@@ -17,7 +17,8 @@ class FChannelItem
 {
 public:
 	FChannelItem(const UPoseSearchFeatureChannel* Channel, int32 ChannelComponentIdx = -1)
-	: Label(ComputeLabel(Channel, ChannelComponentIdx))
+	: FullLabel(ComputeLabel(Channel, ChannelComponentIdx, true))
+	, CompactLabel(ComputeLabel(Channel, ChannelComponentIdx, false))
 	, DataOffset(ComputeDataOffset(Channel, ChannelComponentIdx))
 	, Cardinality(ComputeCardinality(Channel, ChannelComponentIdx))
 	{
@@ -28,9 +29,14 @@ public:
 		return ChannelItems;
 	}
 
-	FString GetLabel() const
+	FString GetFullLabel() const
 	{
-		return Label;
+		return FullLabel;
+	}
+
+	FString GetCompactLabel() const
+	{
+		return CompactLabel;
 	}
 
 	int32 GetDataOffset() const
@@ -54,7 +60,7 @@ public:
 	}
 
 private:
-	static FString ComputeLabel(const UPoseSearchFeatureChannel* Channel, int32 ChannelComponentIdx)
+	static FString ComputeLabel(const UPoseSearchFeatureChannel* Channel, int32 ChannelComponentIdx, bool bFullLabel)
 	{
 		if (ChannelComponentIdx >= 0)
 		{
@@ -70,7 +76,8 @@ private:
 
 		if (Channel != nullptr)
 		{
-			return Channel->GetLabel();
+			TLabelBuilder LabelBuilder;
+			return Channel->GetLabel(LabelBuilder, bFullLabel ? ELabelFormat::Full_Horizontal : ELabelFormat::Compact_Horizontal).ToString();
 		}
 
 		return FString();
@@ -109,7 +116,8 @@ private:
 		return Cardinality;
 	}
 
-	const FString Label;
+	const FString FullLabel;
+	const FString CompactLabel;
 	const int32 DataOffset = 0;
 	const int32 Cardinality = 0;
 	bool bExpanded = false;
@@ -149,7 +157,7 @@ public:
 				.VAlign(VAlign_Center)
 				[
 					SNew(STextBlock)
-					.Text(FText::FromString(ChannelItem->GetLabel()))
+					.Text(FText::FromString(ChannelItem->GetCompactLabel()))
 				];
 		}
 		
@@ -262,7 +270,7 @@ void SDatabaseDataDetails::TrackExpandedItems(const TArray<FChannelItemPtr>& Cha
 	for (const FChannelItemPtr& ChannelItem : ChannelItems)
 	{
 		TrackExpandedItems(ChannelItem->GetChannelItems(), ExpandedItems);
-		ExpandedItems.FindOrAdd(ChannelItem->GetLabel()) = ChannelItem->IsExpanded();
+		ExpandedItems.FindOrAdd(ChannelItem->GetFullLabel()) = ChannelItem->IsExpanded();
 	}
 }
 
@@ -273,7 +281,7 @@ void SDatabaseDataDetails::SetExpandedItems(TArray<FChannelItemPtr>& ChannelItem
 		SetExpandedItems(ChannelItem->GetChannelItems(), ExpandedItems, ChannelItemsTreeView);
 		
 		bool bIsExpanded = false;
-		if (const bool* IsExpandedPtr = ExpandedItems.Find(ChannelItem->GetLabel()))
+		if (const bool* IsExpandedPtr = ExpandedItems.Find(ChannelItem->GetFullLabel()))
 		{
 			bIsExpanded = *IsExpandedPtr;
 		}
diff --git a/Engine/Plugins/Experimental/Animation/PoseSearch/Source/Editor/Private/PoseSearchDebuggerDatabaseColumns.h b/Engine/Plugins/Experimental/Animation/PoseSearch/Source/Editor/Private/PoseSearchDebuggerDatabaseColumns.h
--- a/Engine/Plugins/Experimental/Animation/PoseSearch/Source/Editor/Private/PoseSearchDebuggerDatabaseColumns.h
+++ b/Engine/Plugins/Experimental/Animation/PoseSearch/Source/Editor/Private/PoseSearchDebuggerDatabaseColumns.h
@@ -94,7 +94,7 @@ struct FPoseIdx : ITextColumn
 
 	virtual FText GetLabel() const override
 	{
-		return LOCTEXT("ColumnLabelPoseIndex", "PoseIndex");
+		return LOCTEXT("ColumnLabelPoseIndex", "Pose Id");
 	}
 
 	virtual FText GetLabelTooltip() const override
@@ -119,7 +119,7 @@ struct FAssetIdx : ITextColumn
 
 	virtual FText GetLabel() const override
 	{
-		return LOCTEXT("ColumnLabelAssetIndex", "AssetIndex");
+		return LOCTEXT("ColumnLabelAssetIndex", "Asset Id");
 	}
 
 	virtual FText GetLabelTooltip() const override
@@ -349,7 +349,7 @@ struct FPercentage : ITextColumn
 
 	virtual FText GetLabel() const override
 	{
-		return LOCTEXT("ColumnLabelPercentage", "Percentage");
+		return LOCTEXT("ColumnLabelPercentage", "Percent");
 	}
 
 	virtual FText GetLabelTooltip() const override
diff --git a/Engine/Plugins/Experimental/Animation/PoseSearch/Source/Editor/Private/PoseSearchDebuggerDatabaseView.cpp b/Engine/Plugins/Experimental/Animation/PoseSearch/Source/Editor/Private/PoseSearchDebuggerDatabaseView.cpp
--- a/Engine/Plugins/Experimental/Animation/PoseSearch/Source/Editor/Private/PoseSearchDebuggerDatabaseView.cpp
+++ b/Engine/Plugins/Experimental/Animation/PoseSearch/Source/Editor/Private/PoseSearchDebuggerDatabaseView.cpp
@@ -32,6 +32,7 @@ public:
 	SCostBreakDownData(const TArray<FTraceMotionMatchingStateDatabaseEntry>& DatabaseEntries, bool bIsVerbose)
 	{
 		// processing all the DatabaseEntries to collect the LabelToChannels
+		TLabelBuilder LabelBuilder;
 		for (const FTraceMotionMatchingStateDatabaseEntry& DbEntry : DatabaseEntries)
 		{
 			const UPoseSearchDatabase* Database = FTraceMotionMatchingState::GetObjectFromId<UPoseSearchDatabase>(DbEntry.DatabaseId);
@@ -39,7 +40,7 @@ public:
 			{
 				for (const TObjectPtr<UPoseSearchFeatureChannel>& ChannelPtr : Database->Schema->GetChannels())
 				{
-					AnalyzeChannelRecursively(ChannelPtr.Get(), bIsVerbose);
+					AnalyzeChannelRecursively(LabelBuilder, ChannelPtr.Get(), bIsVerbose);
 				}
 			}
 		}
@@ -102,9 +103,10 @@ public:
 	}
 
 private:
-	void AnalyzeChannelRecursively(const UPoseSearchFeatureChannel* Channel, bool bIsVerbose)
+	void AnalyzeChannelRecursively(TLabelBuilder& LabelBuilder, const UPoseSearchFeatureChannel* Channel, bool bIsVerbose)
 	{
-		const FText Label = FText::FromString(Channel->GetLabel());
+		LabelBuilder.Reset();
+		const FText Label = FText::FromString(Channel->GetLabel(LabelBuilder, ELabelFormat::Full_Vertical).ToString());
 
 		bool bLabelFound = false;
 		for (int32 i = 0; i < LabelToChannels.Num(); ++i)
@@ -128,7 +130,7 @@ private:
 			{
 				if (const UPoseSearchFeatureChannel* SubChannel = SubChannelPtr.Get())
 				{
-					AnalyzeChannelRecursively(SubChannel, bIsVerbose);
+					AnalyzeChannelRecursively(LabelBuilder, SubChannel, bIsVerbose);
 				}
 			}
 		}
@@ -456,7 +458,6 @@ void SDebuggerDatabaseView::Update(const FTraceMotionMatchingStateMessage& State
 			{
 				SHeaderRow::FColumn::FArguments ColumnArgs = SHeaderRow::FColumn::FArguments()
 					.ColumnId(Column.ColumnId)
-					.DefaultLabel(Column.GetLabel())
 					.DefaultTooltip(Column.GetLabelTooltip())
 					.SortMode(this, &SDebuggerDatabaseView::GetColumnSortMode, Column.ColumnId)
 					.OnSort(this, &SDebuggerDatabaseView::OnColumnSortModeChanged)
@@ -464,7 +465,20 @@ void SDebuggerDatabaseView::Update(const FTraceMotionMatchingStateMessage& State
 					.VAlignCell(VAlign_Center)
 					.VAlignHeader(VAlign_Center)
 					.HAlignHeader(HAlign_Center)
-					.HAlignCell(HAlign_Fill);
+					.HAlignCell(HAlign_Fill)
+					.HeaderContent()
+					[
+						SNew( SBox )
+						.VAlign( VAlign_Fill )
+						.HAlign( HAlign_Fill )
+						[
+							SNew( STextBlock )
+							.TextStyle(FAppStyle::Get(), "SmallText")
+							.Text(Column.GetLabel())
+							.Margin(FMargin(5.f, 5.f, 5.f, 5.f))
+							.Justification(ETextJustify::Center)
+						]
+					];
 
 				FilteredDatabaseView.HeaderRow->AddColumn(ColumnArgs);
 				ContinuingPoseView.HeaderRow->AddColumn(ColumnArgs);
diff --git a/Engine/Plugins/Experimental/Animation/PoseSearch/Source/Runtime/Private/PoseSearchFeatureChannel.cpp b/Engine/Plugins/Experimental/Animation/PoseSearch/Source/Runtime/Private/PoseSearchFeatureChannel.cpp
--- a/Engine/Plugins/Experimental/Animation/PoseSearch/Source/Runtime/Private/PoseSearchFeatureChannel.cpp
+++ b/Engine/Plugins/Experimental/Animation/PoseSearch/Source/Runtime/Private/PoseSearchFeatureChannel.cpp
@@ -129,20 +129,53 @@ void UPoseSearchFeatureChannel::GetPermutationTimeOffsets(EPermutationTimeType P
 }
 
 #if WITH_EDITOR
-FString UPoseSearchFeatureChannel::GetLabel() const
+void UPoseSearchFeatureChannel::GetOuterLabel(UE::PoseSearch::TLabelBuilder& LabelBuilder, UE::PoseSearch::ELabelFormat LabelFormat) const
 {
-	TStringBuilder<256> Label;
-	if (const UPoseSearchFeatureChannel* OuterChannel = Cast<UPoseSearchFeatureChannel>(GetOuter()))
+	if (LabelFormat != UE::PoseSearch::ELabelFormat::Compact_Horizontal)
 	{
-		Label.Append(OuterChannel->GetLabel());
-		Label.Append(TEXT("_"));
+		if (const UPoseSearchFeatureChannel* OuterChannel = Cast<UPoseSearchFeatureChannel>(GetOuter()))
+		{
+			if (LabelBuilder.Len() > 0)
+			{
+				LabelBuilder.Append(TEXT("_"));
+			}
+			OuterChannel->GetLabel(LabelBuilder, UE::PoseSearch::ELabelFormat::Full_Horizontal);
+		}
 	}
-	Label.Append(GetName());
-	return Label.ToString();
+}
+
+void UPoseSearchFeatureChannel::AppendLabelSeparator(UE::PoseSearch::TLabelBuilder& LabelBuilder, UE::PoseSearch::ELabelFormat LabelFormat, bool bTryUsingSpace)
+{
+	if (LabelBuilder.Len() == 0)
+	{
+		// do nothing
+	}
+	else if (LabelFormat == UE::PoseSearch::ELabelFormat::Full_Vertical)
+	{
+		LabelBuilder.Append(TEXT("\n"));
+	}
+	else if (bTryUsingSpace)
+	{
+		LabelBuilder.Append(TEXT(" "));
+	}
+	else
+	{
+		LabelBuilder.Append(TEXT("_"));
+	}
+}
+
+UE::PoseSearch::TLabelBuilder& UPoseSearchFeatureChannel::GetLabel(UE::PoseSearch::TLabelBuilder& LabelBuilder, UE::PoseSearch::ELabelFormat LabelFormat) const
+{
+	GetOuterLabel(LabelBuilder, LabelFormat);
+	AppendLabelSeparator(LabelBuilder, LabelFormat);
+	LabelBuilder.Append(GetName());
+	return LabelBuilder;
 }
 
 bool UPoseSearchFeatureChannel::CanBeNormalizedWith(const UPoseSearchFeatureChannel* Other) const
 {
+	using namespace UE::PoseSearch;
+
 	if (this == Other)
 	{
 		return true;
@@ -163,7 +196,8 @@ bool UPoseSearchFeatureChannel::CanBeNormalizedWith(const UPoseSearchFeatureChan
 		return false;
 	}
 
-	if (GetLabel() != Other->GetLabel())
+	TLabelBuilder ThisLabelBuilder, OtherLabelBuilder;
+	if (GetLabel(ThisLabelBuilder).ToString() != Other->GetLabel(OtherLabelBuilder).ToString())
 	{
 		return false;
 	}
diff --git a/Engine/Plugins/Experimental/Animation/PoseSearch/Source/Runtime/Private/PoseSearchFeatureChannel_FilterCrashingLegs.cpp b/Engine/Plugins/Experimental/Animation/PoseSearch/Source/Runtime/Private/PoseSearchFeatureChannel_FilterCrashingLegs.cpp
--- a/Engine/Plugins/Experimental/Animation/PoseSearch/Source/Runtime/Private/PoseSearchFeatureChannel_FilterCrashingLegs.cpp
+++ b/Engine/Plugins/Experimental/Animation/PoseSearch/Source/Runtime/Private/PoseSearchFeatureChannel_FilterCrashingLegs.cpp
@@ -171,8 +171,4 @@ bool UPoseSearchFeatureChannel_FilterCrashingLegs::IndexAsset(UE::PoseSearch::FA
 	return true;
 }
 
-FString UPoseSearchFeatureChannel_FilterCrashingLegs::GetLabel() const
-{
-	return Super::GetLabel();
-}
 #endif
\ No newline at end of file
diff --git a/Engine/Plugins/Experimental/Animation/PoseSearch/Source/Runtime/Private/PoseSearchFeatureChannel_FilterCrashingLegs.h b/Engine/Plugins/Experimental/Animation/PoseSearch/Source/Runtime/Private/PoseSearchFeatureChannel_FilterCrashingLegs.h
--- a/Engine/Plugins/Experimental/Animation/PoseSearch/Source/Runtime/Private/PoseSearchFeatureChannel_FilterCrashingLegs.h
+++ b/Engine/Plugins/Experimental/Animation/PoseSearch/Source/Runtime/Private/PoseSearchFeatureChannel_FilterCrashingLegs.h
@@ -65,7 +65,6 @@ public:
 #if WITH_EDITOR
 	virtual void FillWeights(TArrayView<float> Weights) const override;
 	virtual bool IndexAsset(UE::PoseSearch::FAssetIndexer& Indexer) const override;
-	virtual FString GetLabel() const override;
 #endif
 
 	// IPoseSearchFilter interface
diff --git a/Engine/Plugins/Experimental/Animation/PoseSearch/Source/Runtime/Private/PoseSearchFeatureChannel_Group.cpp b/Engine/Plugins/Experimental/Animation/PoseSearch/Source/Runtime/Private/PoseSearchFeatureChannel_Group.cpp
--- a/Engine/Plugins/Experimental/Animation/PoseSearch/Source/Runtime/Private/PoseSearchFeatureChannel_Group.cpp
+++ b/Engine/Plugins/Experimental/Animation/PoseSearch/Source/Runtime/Private/PoseSearchFeatureChannel_Group.cpp
@@ -113,15 +113,11 @@ bool UPoseSearchFeatureChannel_GroupBase::IsFilterValid(TConstArrayView<float> P
 }
 
 #if WITH_EDITOR
-FString UPoseSearchFeatureChannel_Group::GetLabel() const
+UE::PoseSearch::TLabelBuilder& UPoseSearchFeatureChannel_Group::GetLabel(UE::PoseSearch::TLabelBuilder& LabelBuilder, UE::PoseSearch::ELabelFormat LabelFormat) const
 {
-	TStringBuilder<256> Label;
-	if (const UPoseSearchFeatureChannel* OuterChannel = Cast<UPoseSearchFeatureChannel>(GetOuter()))
-	{
-		Label.Append(OuterChannel->GetLabel());
-		Label.Append(TEXT("_"));
-	}
-	Label.Append(TEXT("Group"));
-	return Label.ToString();
+	GetOuterLabel(LabelBuilder, LabelFormat);
+	AppendLabelSeparator(LabelBuilder, UE::PoseSearch::ELabelFormat::Full_Horizontal);
+	LabelBuilder.Append(TEXT("Group"));
+	return LabelBuilder;
 }
 #endif // WITH_EDITOR
diff --git a/Engine/Plugins/Experimental/Animation/PoseSearch/Source/Runtime/Private/PoseSearchFeatureChannel_Group.h b/Engine/Plugins/Experimental/Animation/PoseSearch/Source/Runtime/Private/PoseSearchFeatureChannel_Group.h
--- a/Engine/Plugins/Experimental/Animation/PoseSearch/Source/Runtime/Private/PoseSearchFeatureChannel_Group.h
+++ b/Engine/Plugins/Experimental/Animation/PoseSearch/Source/Runtime/Private/PoseSearchFeatureChannel_Group.h
@@ -41,7 +41,7 @@ class POSESEARCH_API UPoseSearchFeatureChannel_Group : public UPoseSearchFeature
 	virtual TConstArrayView<TObjectPtr<UPoseSearchFeatureChannel>> GetSubChannels() const override { return SubChannels; }
 
 #if WITH_EDITOR
-	virtual FString GetLabel() const override;
+	virtual UE::PoseSearch::TLabelBuilder& GetLabel(UE::PoseSearch::TLabelBuilder& LabelBuilder, UE::PoseSearch::ELabelFormat LabelFormat = UE::PoseSearch::ELabelFormat::Full_Horizontal) const override;
 #endif
 
 public:
diff --git a/Engine/Plugins/Experimental/Animation/PoseSearch/Source/Runtime/Private/PoseSearchFeatureChannel_Heading.cpp b/Engine/Plugins/Experimental/Animation/PoseSearch/Source/Runtime/Private/PoseSearchFeatureChannel_Heading.cpp
--- a/Engine/Plugins/Experimental/Animation/PoseSearch/Source/Runtime/Private/PoseSearchFeatureChannel_Heading.cpp
+++ b/Engine/Plugins/Experimental/Animation/PoseSearch/Source/Runtime/Private/PoseSearchFeatureChannel_Heading.cpp
@@ -170,61 +170,59 @@ bool UPoseSearchFeatureChannel_Heading::IndexAsset(UE::PoseSearch::FAssetIndexer
 	return true;
 }
 
-FString UPoseSearchFeatureChannel_Heading::GetLabel() const
+UE::PoseSearch::TLabelBuilder& UPoseSearchFeatureChannel_Heading::GetLabel(UE::PoseSearch::TLabelBuilder& LabelBuilder, UE::PoseSearch::ELabelFormat LabelFormat) const
 {
 	using namespace UE::PoseSearch;
 
-	TStringBuilder<256> Label;
-	if (const UPoseSearchFeatureChannel* OuterChannel = Cast<UPoseSearchFeatureChannel>(GetOuter()))
-	{
-		Label.Append(OuterChannel->GetLabel());
-		Label.Append(TEXT("_"));
-	}
+	GetOuterLabel(LabelBuilder, LabelFormat);
+	AppendLabelSeparator(LabelBuilder, LabelFormat);
 
-	Label.Append(TEXT("Head"));
+	LabelBuilder.Append(TEXT("Head"));
 	switch (HeadingAxis)
 	{
 	case EHeadingAxis::X:
-		Label.Append(TEXT("X"));
+		LabelBuilder.Append(TEXT("X"));
 		break;
 	case EHeadingAxis::Y:
-		Label.Append(TEXT("Y"));
+		LabelBuilder.Append(TEXT("Y"));
 		break;
 	case EHeadingAxis::Z:
-		Label.Append(TEXT("Z"));
+		LabelBuilder.Append(TEXT("Z"));
 		break;
 	}
 
 	if (ComponentStripping == EComponentStrippingVector::StripXY)
 	{
-		Label.Append(TEXT("_z"));
+		LabelBuilder.Append(TEXT("_z"));
 	}
 	else if (ComponentStripping == EComponentStrippingVector::StripZ)
 	{
-		Label.Append(TEXT("_xy"));
+		LabelBuilder.Append(TEXT("_xy"));
 	}
 
 	const UPoseSearchSchema* Schema = GetSchema();
 	check(Schema);
 	if (SchemaBoneIdx != RootSchemaBoneIdx)
 	{
-		Label.Append(TEXT("_"));
-		Label.Append(Schema->BoneReferences[SchemaBoneIdx].BoneName.ToString());
+		LabelBuilder.Append(TEXT("_"));
+		LabelBuilder.Append(Schema->BoneReferences[SchemaBoneIdx].BoneName.ToString());
 	}
 
 	if (SchemaOriginBoneIdx != RootSchemaBoneIdx)
 	{
-		Label.Append(TEXT("_"));
-		Label.Append(Schema->BoneReferences[SchemaOriginBoneIdx].BoneName.ToString());
+		LabelBuilder.Append(TEXT("_"));
+		LabelBuilder.Append(Schema->BoneReferences[SchemaOriginBoneIdx].BoneName.ToString());
 	}
 
-	Label.Appendf(TEXT(" %.2f"), SampleTimeOffset);
+	AppendLabelSeparator(LabelBuilder, LabelFormat, true);
+
+	LabelBuilder.Appendf(TEXT("%.2f"), SampleTimeOffset);
 
 	if (!FMath::IsNearlyZero(OriginTimeOffset))
 	{
-		Label.Appendf(TEXT("-%.2f"), OriginTimeOffset);
+		LabelBuilder.Appendf(TEXT("-%.2f"), OriginTimeOffset);
 	}
 
-	return Label.ToString();
+	return LabelBuilder;
 }
 #endif
\ No newline at end of file
diff --git a/Engine/Plugins/Experimental/Animation/PoseSearch/Source/Runtime/Private/PoseSearchFeatureChannel_Heading.h b/Engine/Plugins/Experimental/Animation/PoseSearch/Source/Runtime/Private/PoseSearchFeatureChannel_Heading.h
--- a/Engine/Plugins/Experimental/Animation/PoseSearch/Source/Runtime/Private/PoseSearchFeatureChannel_Heading.h
+++ b/Engine/Plugins/Experimental/Animation/PoseSearch/Source/Runtime/Private/PoseSearchFeatureChannel_Heading.h
@@ -96,7 +96,7 @@ public:
 #if WITH_EDITOR
 	virtual void FillWeights(TArrayView<float> Weights) const override;
 	virtual bool IndexAsset(UE::PoseSearch::FAssetIndexer& Indexer) const override;
-	virtual FString GetLabel() const override;
+	virtual UE::PoseSearch::TLabelBuilder& GetLabel(UE::PoseSearch::TLabelBuilder& LabelBuilder, UE::PoseSearch::ELabelFormat LabelFormat = UE::PoseSearch::ELabelFormat::Full_Horizontal) const override;
 #endif
 
 	FVector GetAxis(const FQuat& Rotation) const;
diff --git a/Engine/Plugins/Experimental/Animation/PoseSearch/Source/Runtime/Private/PoseSearchFeatureChannel_Padding.cpp b/Engine/Plugins/Experimental/Animation/PoseSearch/Source/Runtime/Private/PoseSearchFeatureChannel_Padding.cpp
--- a/Engine/Plugins/Experimental/Animation/PoseSearch/Source/Runtime/Private/PoseSearchFeatureChannel_Padding.cpp
+++ b/Engine/Plugins/Experimental/Animation/PoseSearch/Source/Runtime/Private/PoseSearchFeatureChannel_Padding.cpp
@@ -39,16 +39,11 @@ bool UPoseSearchFeatureChannel_Padding::IndexAsset(UE::PoseSearch::FAssetIndexer
 	return true;
 }
 
-FString UPoseSearchFeatureChannel_Padding::GetLabel() const
+UE::PoseSearch::TLabelBuilder& UPoseSearchFeatureChannel_Padding::GetLabel(UE::PoseSearch::TLabelBuilder& LabelBuilder, UE::PoseSearch::ELabelFormat LabelFormat) const
 {
-	TStringBuilder<256> Label;
-	if (const UPoseSearchFeatureChannel* OuterChannel = Cast<UPoseSearchFeatureChannel>(GetOuter()))
-	{
-		Label.Append(OuterChannel->GetLabel());
-		Label.Append(TEXT("_"));
-	}
-
-	Label.Appendf(TEXT("Pad_%d_%d"), ChannelDataOffset, ChannelDataOffset + ChannelCardinality);
-	return Label.ToString();
+	GetOuterLabel(LabelBuilder, LabelFormat);
+	AppendLabelSeparator(LabelBuilder, LabelFormat);
+	LabelBuilder.Appendf(TEXT("Pad_%d_%d"), ChannelDataOffset, ChannelDataOffset + ChannelCardinality);
+	return LabelBuilder;
 }
 #endif
\ No newline at end of file
diff --git a/Engine/Plugins/Experimental/Animation/PoseSearch/Source/Runtime/Private/PoseSearchFeatureChannel_Padding.h b/Engine/Plugins/Experimental/Animation/PoseSearch/Source/Runtime/Private/PoseSearchFeatureChannel_Padding.h
--- a/Engine/Plugins/Experimental/Animation/PoseSearch/Source/Runtime/Private/PoseSearchFeatureChannel_Padding.h
+++ b/Engine/Plugins/Experimental/Animation/PoseSearch/Source/Runtime/Private/PoseSearchFeatureChannel_Padding.h
@@ -21,7 +21,7 @@ public:
 #if WITH_EDITOR
 	virtual void FillWeights(TArrayView<float> Weights) const override;
 	virtual bool IndexAsset(UE::PoseSearch::FAssetIndexer& Indexer) const override;
-	virtual FString GetLabel() const override;
+	virtual UE::PoseSearch::TLabelBuilder& GetLabel(UE::PoseSearch::TLabelBuilder& LabelBuilder, UE::PoseSearch::ELabelFormat LabelFormat = UE::PoseSearch::ELabelFormat::Full_Horizontal) const override;
 #endif
 
 	static void AddToSchema(UPoseSearchSchema* Schema, int32 PaddingSize);
diff --git a/Engine/Plugins/Experimental/Animation/PoseSearch/Source/Runtime/Private/PoseSearchFeatureChannel_PermutationTime.cpp b/Engine/Plugins/Experimental/Animation/PoseSearch/Source/Runtime/Private/PoseSearchFeatureChannel_PermutationTime.cpp
--- a/Engine/Plugins/Experimental/Animation/PoseSearch/Source/Runtime/Private/PoseSearchFeatureChannel_PermutationTime.cpp
+++ b/Engine/Plugins/Experimental/Animation/PoseSearch/Source/Runtime/Private/PoseSearchFeatureChannel_PermutationTime.cpp
@@ -51,16 +51,11 @@ bool UPoseSearchFeatureChannel_PermutationTime::IndexAsset(UE::PoseSearch::FAsse
 	return true;
 }
 
-FString UPoseSearchFeatureChannel_PermutationTime::GetLabel() const
+UE::PoseSearch::TLabelBuilder& UPoseSearchFeatureChannel_PermutationTime::GetLabel(UE::PoseSearch::TLabelBuilder& LabelBuilder, UE::PoseSearch::ELabelFormat LabelFormat) const
 {
-	TStringBuilder<256> Label;
-	if (const UPoseSearchFeatureChannel* OuterChannel = Cast<UPoseSearchFeatureChannel>(GetOuter()))
-	{
-		Label.Append(OuterChannel->GetLabel());
-		Label.Append(TEXT("_"));
-	}
-
-	Label.Append(TEXT("PermTime"));
-	return Label.ToString();
+	GetOuterLabel(LabelBuilder, LabelFormat);
+	AppendLabelSeparator(LabelBuilder, LabelFormat);
+	LabelBuilder.Append(TEXT("PermTime"));
+	return LabelBuilder;
 }
 #endif
\ No newline at end of file
diff --git a/Engine/Plugins/Experimental/Animation/PoseSearch/Source/Runtime/Private/PoseSearchFeatureChannel_PermutationTime.h b/Engine/Plugins/Experimental/Animation/PoseSearch/Source/Runtime/Private/PoseSearchFeatureChannel_PermutationTime.h
--- a/Engine/Plugins/Experimental/Animation/PoseSearch/Source/Runtime/Private/PoseSearchFeatureChannel_PermutationTime.h
+++ b/Engine/Plugins/Experimental/Animation/PoseSearch/Source/Runtime/Private/PoseSearchFeatureChannel_PermutationTime.h
@@ -23,7 +23,7 @@ public:
 #if WITH_EDITOR
 	virtual void FillWeights(TArrayView<float> Weights) const override;
 	virtual bool IndexAsset(UE::PoseSearch::FAssetIndexer& Indexer) const override;
-	virtual FString GetLabel() const override;
+	virtual UE::PoseSearch::TLabelBuilder& GetLabel(UE::PoseSearch::TLabelBuilder& LabelBuilder, UE::PoseSearch::ELabelFormat LabelFormat = UE::PoseSearch::ELabelFormat::Full_Horizontal) const override;
 #endif
 
 	static void FindOrAddToSchema(UPoseSearchSchema* Schema);
diff --git a/Engine/Plugins/Experimental/Animation/PoseSearch/Source/Runtime/Private/PoseSearchFeatureChannel_Phase.cpp b/Engine/Plugins/Experimental/Animation/PoseSearch/Source/Runtime/Private/PoseSearchFeatureChannel_Phase.cpp
--- a/Engine/Plugins/Experimental/Animation/PoseSearch/Source/Runtime/Private/PoseSearchFeatureChannel_Phase.cpp
+++ b/Engine/Plugins/Experimental/Animation/PoseSearch/Source/Runtime/Private/PoseSearchFeatureChannel_Phase.cpp
@@ -408,27 +408,23 @@ bool UPoseSearchFeatureChannel_Phase::IndexAsset(UE::PoseSearch::FAssetIndexer&
 	return true;
 }
 
-FString UPoseSearchFeatureChannel_Phase::GetLabel() const
+UE::PoseSearch::TLabelBuilder& UPoseSearchFeatureChannel_Phase::GetLabel(UE::PoseSearch::TLabelBuilder& LabelBuilder, UE::PoseSearch::ELabelFormat LabelFormat) const
 {
 	using namespace UE::PoseSearch;
 
-	TStringBuilder<256> Label;
-	if (const UPoseSearchFeatureChannel* OuterChannel = Cast<UPoseSearchFeatureChannel>(GetOuter()))
-	{
-		Label.Append(OuterChannel->GetLabel());
-		Label.Append(TEXT("_"));
-	}
+	GetOuterLabel(LabelBuilder, LabelFormat);
+	AppendLabelSeparator(LabelBuilder, LabelFormat);
 
-	Label.Append(TEXT("Pha"));
+	LabelBuilder.Append(TEXT("Pha"));
 
 	const UPoseSearchSchema* Schema = GetSchema();
 	check(Schema);
 	if (SchemaBoneIdx != RootSchemaBoneIdx)
 	{
-		Label.Append(TEXT("_"));
-		Label.Append(Schema->BoneReferences[SchemaBoneIdx].BoneName.ToString());
+		LabelBuilder.Append(TEXT("_"));
+		LabelBuilder.Append(Schema->BoneReferences[SchemaBoneIdx].BoneName.ToString());
 	}
 
-	return Label.ToString();
+	return LabelBuilder;
 }
 #endif
\ No newline at end of file
diff --git a/Engine/Plugins/Experimental/Animation/PoseSearch/Source/Runtime/Private/PoseSearchFeatureChannel_Phase.h b/Engine/Plugins/Experimental/Animation/PoseSearch/Source/Runtime/Private/PoseSearchFeatureChannel_Phase.h
--- a/Engine/Plugins/Experimental/Animation/PoseSearch/Source/Runtime/Private/PoseSearchFeatureChannel_Phase.h
+++ b/Engine/Plugins/Experimental/Animation/PoseSearch/Source/Runtime/Private/PoseSearchFeatureChannel_Phase.h
@@ -45,6 +45,6 @@ public:
 #if WITH_EDITOR
 	virtual void FillWeights(TArrayView<float> Weights) const override;
 	virtual bool IndexAsset(UE::PoseSearch::FAssetIndexer& Indexer) const override;
-	virtual FString GetLabel() const override;
+	virtual UE::PoseSearch::TLabelBuilder& GetLabel(UE::PoseSearch::TLabelBuilder& LabelBuilder, UE::PoseSearch::ELabelFormat LabelFormat = UE::PoseSearch::ELabelFormat::Full_Horizontal) const override;
 #endif
 };
\ No newline at end of file
diff --git a/Engine/Plugins/Experimental/Animation/PoseSearch/Source/Runtime/Private/PoseSearchFeatureChannel_Pose.cpp b/Engine/Plugins/Experimental/Animation/PoseSearch/Source/Runtime/Private/PoseSearchFeatureChannel_Pose.cpp
--- a/Engine/Plugins/Experimental/Animation/PoseSearch/Source/Runtime/Private/PoseSearchFeatureChannel_Pose.cpp
+++ b/Engine/Plugins/Experimental/Animation/PoseSearch/Source/Runtime/Private/PoseSearchFeatureChannel_Pose.cpp
@@ -107,16 +107,12 @@ void UPoseSearchFeatureChannel_Pose::Finalize(UPoseSearchSchema* Schema)
 }
 
 #if WITH_EDITOR
-FString UPoseSearchFeatureChannel_Pose::GetLabel() const
+UE::PoseSearch::TLabelBuilder& UPoseSearchFeatureChannel_Pose::GetLabel(UE::PoseSearch::TLabelBuilder& LabelBuilder, UE::PoseSearch::ELabelFormat LabelFormat) const
 {
-	TStringBuilder<256> Label;
-	if (const UPoseSearchFeatureChannel* OuterChannel = Cast<UPoseSearchFeatureChannel>(GetOuter()))
-	{
-		Label.Append(OuterChannel->GetLabel());
-		Label.Append(TEXT("_"));
-	}
-	Label.Append(TEXT("Pose"));
-	return Label.ToString();
+	GetOuterLabel(LabelBuilder, LabelFormat);
+	AppendLabelSeparator(LabelBuilder, UE::PoseSearch::ELabelFormat::Full_Horizontal);
+	LabelBuilder.Append(TEXT("Pose"));
+	return LabelBuilder;
 }
 #endif // WITH_EDITOR
 
diff --git a/Engine/Plugins/Experimental/Animation/PoseSearch/Source/Runtime/Private/PoseSearchFeatureChannel_Pose.h b/Engine/Plugins/Experimental/Animation/PoseSearch/Source/Runtime/Private/PoseSearchFeatureChannel_Pose.h
--- a/Engine/Plugins/Experimental/Animation/PoseSearch/Source/Runtime/Private/PoseSearchFeatureChannel_Pose.h
+++ b/Engine/Plugins/Experimental/Animation/PoseSearch/Source/Runtime/Private/PoseSearchFeatureChannel_Pose.h
@@ -77,7 +77,7 @@ public:
 	virtual void Finalize(UPoseSearchSchema* Schema) override;
 
 #if WITH_EDITOR
-	virtual FString GetLabel() const override;
+	virtual UE::PoseSearch::TLabelBuilder& GetLabel(UE::PoseSearch::TLabelBuilder& LabelBuilder, UE::PoseSearch::ELabelFormat LabelFormat = UE::PoseSearch::ELabelFormat::Full_Horizontal) const override;
 #endif
 };
 
diff --git a/Engine/Plugins/Experimental/Animation/PoseSearch/Source/Runtime/Private/PoseSearchFeatureChannel_Position.cpp b/Engine/Plugins/Experimental/Animation/PoseSearch/Source/Runtime/Private/PoseSearchFeatureChannel_Position.cpp
--- a/Engine/Plugins/Experimental/Animation/PoseSearch/Source/Runtime/Private/PoseSearchFeatureChannel_Position.cpp
+++ b/Engine/Plugins/Experimental/Animation/PoseSearch/Source/Runtime/Private/PoseSearchFeatureChannel_Position.cpp
@@ -160,49 +160,47 @@ bool UPoseSearchFeatureChannel_Position::IndexAsset(UE::PoseSearch::FAssetIndexe
 	return true;
 }
 
-FString UPoseSearchFeatureChannel_Position::GetLabel() const
+UE::PoseSearch::TLabelBuilder& UPoseSearchFeatureChannel_Position::GetLabel(UE::PoseSearch::TLabelBuilder& LabelBuilder, UE::PoseSearch::ELabelFormat LabelFormat) const
 {
 	using namespace UE::PoseSearch;
 
-	TStringBuilder<256> Label;
-	if (const UPoseSearchFeatureChannel* OuterChannel = Cast<UPoseSearchFeatureChannel>(GetOuter()))
-	{
-		Label.Append(OuterChannel->GetLabel());
-		Label.Append(TEXT("_"));
-	}
+	GetOuterLabel(LabelBuilder, LabelFormat);
+	AppendLabelSeparator(LabelBuilder, LabelFormat);
 
-	Label.Append(TEXT("Pos"));
+	LabelBuilder.Append(TEXT("Pos"));
 
 	if (ComponentStripping == EComponentStrippingVector::StripXY)
 	{
-		Label.Append(TEXT("_z"));
+		LabelBuilder.Append(TEXT("_z"));
 	}
 	else if (ComponentStripping == EComponentStrippingVector::StripZ)
 	{
-		Label.Append(TEXT("_xy"));
+		LabelBuilder.Append(TEXT("_xy"));
 	}
 
 	const UPoseSearchSchema* Schema = GetSchema();
 	check(Schema);
 	if (SchemaBoneIdx != RootSchemaBoneIdx)
 	{
-		Label.Append(TEXT("_"));
-		Label.Append(Schema->BoneReferences[SchemaBoneIdx].BoneName.ToString());
+		LabelBuilder.Append(TEXT("_"));
+		LabelBuilder.Append(Schema->BoneReferences[SchemaBoneIdx].BoneName.ToString());
 	}
 
 	if (SchemaOriginBoneIdx != RootSchemaBoneIdx)
 	{
-		Label.Append(TEXT("_"));
-		Label.Append(Schema->BoneReferences[SchemaOriginBoneIdx].BoneName.ToString());
+		LabelBuilder.Append(TEXT("_"));
+		LabelBuilder.Append(Schema->BoneReferences[SchemaOriginBoneIdx].BoneName.ToString());
 	}
 
-	Label.Appendf(TEXT(" %.2f"), SampleTimeOffset);
+	AppendLabelSeparator(LabelBuilder, LabelFormat, true);
+
+	LabelBuilder.Appendf(TEXT("%.2f"), SampleTimeOffset);
 
 	if (!FMath::IsNearlyZero(OriginTimeOffset))
 	{
-		Label.Appendf(TEXT("-%.2f"), OriginTimeOffset);
+		LabelBuilder.Appendf(TEXT("-%.2f"), OriginTimeOffset);
 	}
 
-	return Label.ToString();
+	return LabelBuilder;
 }
 #endif
\ No newline at end of file
diff --git a/Engine/Plugins/Experimental/Animation/PoseSearch/Source/Runtime/Private/PoseSearchFeatureChannel_Position.h b/Engine/Plugins/Experimental/Animation/PoseSearch/Source/Runtime/Private/PoseSearchFeatureChannel_Position.h
--- a/Engine/Plugins/Experimental/Animation/PoseSearch/Source/Runtime/Private/PoseSearchFeatureChannel_Position.h
+++ b/Engine/Plugins/Experimental/Animation/PoseSearch/Source/Runtime/Private/PoseSearchFeatureChannel_Position.h
@@ -82,7 +82,7 @@ public:
 #if WITH_EDITOR
 	virtual void FillWeights(TArrayView<float> Weights) const override;
 	virtual bool IndexAsset(UE::PoseSearch::FAssetIndexer& Indexer) const override;
-	virtual FString GetLabel() const override;
+	virtual UE::PoseSearch::TLabelBuilder& GetLabel(UE::PoseSearch::TLabelBuilder& LabelBuilder, UE::PoseSearch::ELabelFormat LabelFormat = UE::PoseSearch::ELabelFormat::Full_Horizontal) const override;
 #endif
 
 	static void FindOrAddToSchema(UPoseSearchSchema* Schema, float SampleTimeOffset, const FName& BoneName = NAME_None, EPermutationTimeType PermutationTimeType = EPermutationTimeType::UseSampleTime);
diff --git a/Engine/Plugins/Experimental/Animation/PoseSearch/Source/Runtime/Private/PoseSearchFeatureChannel_SamplingTime.cpp b/Engine/Plugins/Experimental/Animation/PoseSearch/Source/Runtime/Private/PoseSearchFeatureChannel_SamplingTime.cpp
--- a/Engine/Plugins/Experimental/Animation/PoseSearch/Source/Runtime/Private/PoseSearchFeatureChannel_SamplingTime.cpp
+++ b/Engine/Plugins/Experimental/Animation/PoseSearch/Source/Runtime/Private/PoseSearchFeatureChannel_SamplingTime.cpp
@@ -38,16 +38,11 @@ bool UPoseSearchFeatureChannel_SamplingTime::IndexAsset(UE::PoseSearch::FAssetIn
 	return true;
 }
 
-FString UPoseSearchFeatureChannel_SamplingTime::GetLabel() const
+UE::PoseSearch::TLabelBuilder& UPoseSearchFeatureChannel_SamplingTime::GetLabel(UE::PoseSearch::TLabelBuilder& LabelBuilder, UE::PoseSearch::ELabelFormat LabelFormat) const
 {
-	TStringBuilder<256> Label;
-	if (const UPoseSearchFeatureChannel* OuterChannel = Cast<UPoseSearchFeatureChannel>(GetOuter()))
-	{
-		Label.Append(OuterChannel->GetLabel());
-		Label.Append(TEXT("_"));
-	}
-
-	Label.Append(TEXT("SamplTime"));
-	return Label.ToString();
+	GetOuterLabel(LabelBuilder, LabelFormat);
+	AppendLabelSeparator(LabelBuilder, LabelFormat);
+	LabelBuilder.Append(TEXT("SamplTime"));
+	return LabelBuilder;
 }
 #endif
\ No newline at end of file
diff --git a/Engine/Plugins/Experimental/Animation/PoseSearch/Source/Runtime/Private/PoseSearchFeatureChannel_SamplingTime.h b/Engine/Plugins/Experimental/Animation/PoseSearch/Source/Runtime/Private/PoseSearchFeatureChannel_SamplingTime.h
--- a/Engine/Plugins/Experimental/Animation/PoseSearch/Source/Runtime/Private/PoseSearchFeatureChannel_SamplingTime.h
+++ b/Engine/Plugins/Experimental/Animation/PoseSearch/Source/Runtime/Private/PoseSearchFeatureChannel_SamplingTime.h
@@ -25,6 +25,6 @@ public:
 #if WITH_EDITOR
 	virtual void FillWeights(TArrayView<float> Weights) const override;
 	virtual bool IndexAsset(UE::PoseSearch::FAssetIndexer& Indexer) const override;
-	virtual FString GetLabel() const override;
+	virtual UE::PoseSearch::TLabelBuilder& GetLabel(UE::PoseSearch::TLabelBuilder& LabelBuilder, UE::PoseSearch::ELabelFormat LabelFormat = UE::PoseSearch::ELabelFormat::Full_Horizontal) const override;
 #endif
 };
diff --git a/Engine/Plugins/Experimental/Animation/PoseSearch/Source/Runtime/Private/PoseSearchFeatureChannel_Trajectory.cpp b/Engine/Plugins/Experimental/Animation/PoseSearch/Source/Runtime/Private/PoseSearchFeatureChannel_Trajectory.cpp
--- a/Engine/Plugins/Experimental/Animation/PoseSearch/Source/Runtime/Private/PoseSearchFeatureChannel_Trajectory.cpp
+++ b/Engine/Plugins/Experimental/Animation/PoseSearch/Source/Runtime/Private/PoseSearchFeatureChannel_Trajectory.cpp
@@ -177,16 +177,12 @@ void UPoseSearchFeatureChannel_Trajectory::DebugDraw(const UE::PoseSearch::FDebu
 #endif // ENABLE_DRAW_DEBUG
 
 #if WITH_EDITOR
-FString UPoseSearchFeatureChannel_Trajectory::GetLabel() const
+UE::PoseSearch::TLabelBuilder& UPoseSearchFeatureChannel_Trajectory::GetLabel(UE::PoseSearch::TLabelBuilder& LabelBuilder, UE::PoseSearch::ELabelFormat LabelFormat) const
 {
-	TStringBuilder<256> Label;
-	if (const UPoseSearchFeatureChannel* OuterChannel = Cast<UPoseSearchFeatureChannel>(GetOuter()))
-	{
-		Label.Append(OuterChannel->GetLabel());
-		Label.Append(TEXT("_"));
-	}
-	Label.Append(TEXT("Traj"));
-	return Label.ToString();
+	GetOuterLabel(LabelBuilder, LabelFormat);
+	AppendLabelSeparator(LabelBuilder, UE::PoseSearch::ELabelFormat::Full_Horizontal);
+	LabelBuilder.Append(TEXT("Traj"));
+	return LabelBuilder;
 }
 #endif // WITH_EDITOR
 
diff --git a/Engine/Plugins/Experimental/Animation/PoseSearch/Source/Runtime/Private/PoseSearchFeatureChannel_Trajectory.h b/Engine/Plugins/Experimental/Animation/PoseSearch/Source/Runtime/Private/PoseSearchFeatureChannel_Trajectory.h
--- a/Engine/Plugins/Experimental/Animation/PoseSearch/Source/Runtime/Private/PoseSearchFeatureChannel_Trajectory.h
+++ b/Engine/Plugins/Experimental/Animation/PoseSearch/Source/Runtime/Private/PoseSearchFeatureChannel_Trajectory.h
@@ -76,7 +76,7 @@ public:
 #endif // ENABLE_DRAW_DEBUG
 
 #if WITH_EDITOR
-	virtual FString GetLabel() const override;
+	virtual UE::PoseSearch::TLabelBuilder& GetLabel(UE::PoseSearch::TLabelBuilder& LabelBuilder, UE::PoseSearch::ELabelFormat LabelFormat = UE::PoseSearch::ELabelFormat::Full_Horizontal) const override;
 #endif
 
 	float GetEstimatedSpeedRatio(TConstArrayView<float> QueryVector, TConstArrayView<float> PoseVector) const;
diff --git a/Engine/Plugins/Experimental/Animation/PoseSearch/Source/Runtime/Private/PoseSearchFeatureChannel_Velocity.cpp b/Engine/Plugins/Experimental/Animation/PoseSearch/Source/Runtime/Private/PoseSearchFeatureChannel_Velocity.cpp
--- a/Engine/Plugins/Experimental/Animation/PoseSearch/Source/Runtime/Private/PoseSearchFeatureChannel_Velocity.cpp
+++ b/Engine/Plugins/Experimental/Animation/PoseSearch/Source/Runtime/Private/PoseSearchFeatureChannel_Velocity.cpp
@@ -134,53 +134,51 @@ bool UPoseSearchFeatureChannel_Velocity::IndexAsset(UE::PoseSearch::FAssetIndexe
 	return true;
 }
 
-FString UPoseSearchFeatureChannel_Velocity::GetLabel() const
+UE::PoseSearch::TLabelBuilder& UPoseSearchFeatureChannel_Velocity::GetLabel(UE::PoseSearch::TLabelBuilder& LabelBuilder, UE::PoseSearch::ELabelFormat LabelFormat) const
 {
 	using namespace UE::PoseSearch;
 
-	TStringBuilder<256> Label;
-	if (const UPoseSearchFeatureChannel* OuterChannel = Cast<UPoseSearchFeatureChannel>(GetOuter()))
-	{
-		Label.Append(OuterChannel->GetLabel());
-		Label.Append(TEXT("_"));
-	}
+	GetOuterLabel(LabelBuilder, LabelFormat);
+	AppendLabelSeparator(LabelBuilder, LabelFormat);
 
-	Label.Append(TEXT("Vel"));
+	LabelBuilder.Append(TEXT("Vel"));
 	if (bNormalize)
 	{
-		Label.Append(TEXT("Dir"));
+		LabelBuilder.Append(TEXT("Dir"));
 	}
 
 	if (ComponentStripping == EComponentStrippingVector::StripXY)
 	{
-		Label.Append(TEXT("_z"));
+		LabelBuilder.Append(TEXT("_z"));
 	}
 	else if (ComponentStripping == EComponentStrippingVector::StripZ)
 	{
-		Label.Append(TEXT("_xy"));
+		LabelBuilder.Append(TEXT("_xy"));
 	}
 
 	const UPoseSearchSchema* Schema = GetSchema();
 	check(Schema);
 	if (SchemaBoneIdx != RootSchemaBoneIdx)
 	{
-		Label.Append(TEXT("_"));
-		Label.Append(Schema->BoneReferences[SchemaBoneIdx].BoneName.ToString());
+		LabelBuilder.Append(TEXT("_"));
+		LabelBuilder.Append(Schema->BoneReferences[SchemaBoneIdx].BoneName.ToString());
 	}
 
 	if (SchemaOriginBoneIdx != RootSchemaBoneIdx)
 	{
-		Label.Append(TEXT("_"));
-		Label.Append(Schema->BoneReferences[SchemaOriginBoneIdx].BoneName.ToString());
+		LabelBuilder.Append(TEXT("_"));
+		LabelBuilder.Append(Schema->BoneReferences[SchemaOriginBoneIdx].BoneName.ToString());
 	}
 
-	Label.Appendf(TEXT(" %.2f"), SampleTimeOffset);
+	AppendLabelSeparator(LabelBuilder, LabelFormat, true);
+
+	LabelBuilder.Appendf(TEXT("%.2f"), SampleTimeOffset);
 
 	if (!FMath::IsNearlyZero(OriginTimeOffset))
 	{
-		Label.Appendf(TEXT("-%.2f"), OriginTimeOffset);
+		LabelBuilder.Appendf(TEXT("-%.2f"), OriginTimeOffset);
 	}
 
-	return Label.ToString();
+	return LabelBuilder;
 }
 #endif
\ No newline at end of file
diff --git a/Engine/Plugins/Experimental/Animation/PoseSearch/Source/Runtime/Private/PoseSearchFeatureChannel_Velocity.h b/Engine/Plugins/Experimental/Animation/PoseSearch/Source/Runtime/Private/PoseSearchFeatureChannel_Velocity.h
--- a/Engine/Plugins/Experimental/Animation/PoseSearch/Source/Runtime/Private/PoseSearchFeatureChannel_Velocity.h
+++ b/Engine/Plugins/Experimental/Animation/PoseSearch/Source/Runtime/Private/PoseSearchFeatureChannel_Velocity.h
@@ -90,6 +90,6 @@ public:
 #if WITH_EDITOR
 	virtual void FillWeights(TArrayView<float> Weights) const override;
 	virtual bool IndexAsset(UE::PoseSearch::FAssetIndexer& Indexer) const override;
-	virtual FString GetLabel() const override;
+	virtual UE::PoseSearch::TLabelBuilder& GetLabel(UE::PoseSearch::TLabelBuilder& LabelBuilder, UE::PoseSearch::ELabelFormat LabelFormat = UE::PoseSearch::ELabelFormat::Full_Horizontal) const override;
 #endif
 };
diff --git a/Engine/Plugins/Experimental/Animation/PoseSearch/Source/Runtime/Public/PoseSearch/PoseSearchFeatureChannel.h b/Engine/Plugins/Experimental/Animation/PoseSearch/Source/Runtime/Public/PoseSearch/PoseSearchFeatureChannel.h
--- a/Engine/Plugins/Experimental/Animation/PoseSearch/Source/Runtime/Public/PoseSearch/PoseSearchFeatureChannel.h
+++ b/Engine/Plugins/Experimental/Animation/PoseSearch/Source/Runtime/Public/PoseSearch/PoseSearchFeatureChannel.h
@@ -61,6 +61,21 @@ struct FPoseMetadata;
 
 #if WITH_EDITOR
 class FAssetIndexer;
+
+enum class ELabelFormat : uint8
+{
+	// output label example: "Traj_Vel_xy 1.20"
+	Full_Horizontal,
+
+	// output label example: "Traj\nVel_xy\n1.20"
+	Full_Vertical,
+
+	// output label example: "Vel_xy 1.20"
+	Compact_Horizontal
+};
+
+typedef TStringBuilder<256> TLabelBuilder;
+
 #endif // WITH_EDITOR
 
 /** Helper class for extracting and encoding features into a float buffer */
@@ -135,8 +150,9 @@ public:
 	// Called at database build time to populate pose vectors with this channel's data
 	virtual bool IndexAsset(UE::PoseSearch::FAssetIndexer& Indexer) const PURE_VIRTUAL(UPoseSearchFeatureChannel::IndexAsset, return false;);
 
-	// returns the FString used editor side to identify this UPoseSearchFeatureChannel (for instance in the pose search debugger)
-	virtual FString GetLabel() const;
+	// returns the TLabelBuilder used editor side to identify this UPoseSearchFeatureChannel (for instance in the pose search debugger)
+	virtual UE::PoseSearch::TLabelBuilder& GetLabel(UE::PoseSearch::TLabelBuilder& LabelBuilder, UE::PoseSearch::ELabelFormat LabelFormat = UE::PoseSearch::ELabelFormat::Full_Horizontal) const;
+
 	virtual bool CanBeNormalizedWith(const UPoseSearchFeatureChannel* Other) const;
 	const UPoseSearchSchema* GetSchema() const;
 #endif
@@ -147,6 +163,10 @@ private:
 	USkeleton* GetSkeleton(bool& bInvalidSkeletonIsError, const IPropertyHandle* PropertyHandle) override;
 
 protected:
+#if WITH_EDITOR
+	void GetOuterLabel(UE::PoseSearch::TLabelBuilder& LabelBuilder, UE::PoseSearch::ELabelFormat LabelFormat) const;
+	static void AppendLabelSeparator(UE::PoseSearch::TLabelBuilder& LabelBuilder, UE::PoseSearch::ELabelFormat LabelFormat, bool bTryUsingSpace = false);
+#endif //WITH_EDITOR
 	friend class ::UPoseSearchSchema;
 
 	UPROPERTY(Transient)
