motion matching - pose history memory optimization

diff --git a/Engine/Plugins/Experimental/Animation/PoseSearch/Source/Runtime/Private/AnimNode_PoseSearchHistoryCollector.cpp b/Engine/Plugins/Experimental/Animation/PoseSearch/Source/Runtime/Private/AnimNode_PoseSearchHistoryCollector.cpp
--- a/Engine/Plugins/Experimental/Animation/PoseSearch/Source/Runtime/Private/AnimNode_PoseSearchHistoryCollector.cpp
+++ b/Engine/Plugins/Experimental/Animation/PoseSearch/Source/Runtime/Private/AnimNode_PoseSearchHistoryCollector.cpp
@@ -80,7 +80,7 @@ void FAnimNode_PoseSearchHistoryCollector_Base::CacheBones_AnyThread(const FAnim
 		Pose.SetBoneContainer(&Context.AnimInstanceProxy->GetRequiredBones());
 		FCSPose<FCompactPose> ComponentSpacePose;
 		ComponentSpacePose.InitPose(Pose);
-		PoseHistory.Update(0.f, ComponentSpacePose);
+		PoseHistory.Update(0.f, ComponentSpacePose, bStoreScales);
 	}
 }
 
@@ -128,7 +128,7 @@ void FAnimNode_PoseSearchHistoryCollector::Evaluate_AnyThread(FPoseContext& Outp
 
 	FCSPose<FCompactPose> ComponentSpacePose;
 	ComponentSpacePose.InitPose(Output.Pose);
-	PoseHistory.Update(Output.AnimInstanceProxy->GetDeltaSeconds(), ComponentSpacePose);
+	PoseHistory.Update(Output.AnimInstanceProxy->GetDeltaSeconds(), ComponentSpacePose, bStoreScales);
 
 #if ENABLE_DRAW_DEBUG && ENABLE_ANIM_DEBUG
 	if (CVarAnimPoseHistoryDebugDraw.GetValueOnAnyThread())
@@ -187,7 +187,7 @@ void FAnimNode_PoseSearchComponentSpaceHistoryCollector::EvaluateComponentSpace_
 	Super::EvaluateComponentSpace_AnyThread(Output);
 	Source.EvaluateComponentSpace(Output);
 
-	PoseHistory.Update(Output.AnimInstanceProxy->GetDeltaSeconds(), Output.Pose);
+	PoseHistory.Update(Output.AnimInstanceProxy->GetDeltaSeconds(), Output.Pose, bStoreScales);
 
 #if ENABLE_DRAW_DEBUG && ENABLE_ANIM_DEBUG
 	if (CVarAnimPoseHistoryDebugDraw.GetValueOnAnyThread())
diff --git a/Engine/Plugins/Experimental/Animation/PoseSearch/Source/Runtime/Private/PoseSearchHistory.cpp b/Engine/Plugins/Experimental/Animation/PoseSearch/Source/Runtime/Private/PoseSearchHistory.cpp
--- a/Engine/Plugins/Experimental/Animation/PoseSearch/Source/Runtime/Private/PoseSearchHistory.cpp
+++ b/Engine/Plugins/Experimental/Animation/PoseSearch/Source/Runtime/Private/PoseSearchHistory.cpp
@@ -67,7 +67,7 @@ FORCEINLINE auto LowerBound(IteratorType First, IteratorType Last, const ValueTy
 
 //////////////////////////////////////////////////////////////////////////
 // FPoseHistoryEntry
-void FPoseHistoryEntry::Update(float InTime, FCSPose<FCompactPose>& ComponentSpacePose, const FBoneToTransformMap& BoneToTransformMap)
+void FPoseHistoryEntry::Update(float InTime, FCSPose<FCompactPose>& ComponentSpacePose, const FBoneToTransformMap& BoneToTransformMap, bool bStoreScales)
 {
 	Time = InTime;
 
@@ -81,25 +81,55 @@ void FPoseHistoryEntry::Update(float InTime, FCSPose<FCompactPose>& ComponentSpa
 	if (BoneToTransformMap.IsEmpty())
 	{
 		// no mapping: we add all the transforms
-		ComponentSpaceTransforms.SetNum(NumSkeletonBones);
+		SetNum(NumSkeletonBones, bStoreScales);
 		for (FSkeletonPoseBoneIndex SkeletonBoneIdx(0); SkeletonBoneIdx != NumSkeletonBones; ++SkeletonBoneIdx)
 		{
 			const FCompactPoseBoneIndex CompactBoneIdx = BoneContainer.GetCompactPoseIndexFromSkeletonPoseIndex(SkeletonBoneIdx);
-			ComponentSpaceTransforms[SkeletonBoneIdx.GetInt()] = (CompactBoneIdx.IsValid() ? ComponentSpacePose.GetComponentSpaceTransform(CompactBoneIdx) : RefBonePose[SkeletonBoneIdx.GetInt()]);
+			SetComponentSpaceTransform(SkeletonBoneIdx.GetInt(), (CompactBoneIdx.IsValid() ? ComponentSpacePose.GetComponentSpaceTransform(CompactBoneIdx) : RefBonePose[SkeletonBoneIdx.GetInt()]));
 		}
 	}
 	else
 	{
-		ComponentSpaceTransforms.SetNum(BoneToTransformMap.Num());
+		SetNum(BoneToTransformMap.Num(), true);
 		for (const FBoneToTransformPair& BoneToTransformPair : BoneToTransformMap)
 		{
 			const FSkeletonPoseBoneIndex SkeletonBoneIdx(BoneToTransformPair.Key);
 			const FCompactPoseBoneIndex CompactBoneIdx = BoneContainer.GetCompactPoseIndexFromSkeletonPoseIndex(SkeletonBoneIdx);
-			ComponentSpaceTransforms[BoneToTransformPair.Value] = (CompactBoneIdx.IsValid() ? ComponentSpacePose.GetComponentSpaceTransform(CompactBoneIdx) : RefBonePose[SkeletonBoneIdx.GetInt()]);
+			SetComponentSpaceTransform(BoneToTransformPair.Value, (CompactBoneIdx.IsValid() ? ComponentSpacePose.GetComponentSpaceTransform(CompactBoneIdx) : RefBonePose[SkeletonBoneIdx.GetInt()]));
 		}
 	}
 }
 
+void FPoseHistoryEntry::SetNum(int32 Num, bool bStoreScales)
+{
+	ComponentSpaceRotations.SetNum(Num);
+	ComponentSpacePositions.SetNum(Num);
+	ComponentSpaceScales.SetNum(bStoreScales ? Num : 0);
+}
+
+int32 FPoseHistoryEntry::Num() const
+{
+	return ComponentSpaceRotations.Num();
+}
+
+void FPoseHistoryEntry::SetComponentSpaceTransform(int32 Index, const FTransform& Transform)
+{
+	ComponentSpaceRotations[Index] = FQuat4f(Transform.GetRotation());
+	ComponentSpacePositions[Index] = Transform.GetTranslation();
+	
+	if (!ComponentSpaceScales.IsEmpty())
+	{
+		ComponentSpaceScales[Index] = FVector3f(Transform.GetScale3D());
+	}
+}
+
+FTransform FPoseHistoryEntry::GetComponentSpaceTransform(int32 Index) const
+{
+	const FQuat Quat(ComponentSpaceRotations[Index]);
+	const FVector Scale(ComponentSpaceScales.IsEmpty() ? FVector3f::OneVector : ComponentSpaceScales[Index]);
+	return FTransform(Quat, ComponentSpacePositions[Index], Scale);
+}
+
 //////////////////////////////////////////////////////////////////////////
 // FPoseHistory
 void FPoseHistory::Init(int32 InNumPoses, float InTimeHorizon, const TArray<FBoneIndexType>& RequiredBones)
@@ -187,15 +217,15 @@ bool FPoseHistory::LerpEntries(float Time, bool bExtrapolate, const FPoseHistory
 		if (ReferenceBoneTransformIndex == ComponentSpaceIndexType)
 		{
 			OutBoneTransform.Blend(
-				PrevEntry.ComponentSpaceTransforms[BoneTransformIndex],
-				NextEntry.ComponentSpaceTransforms[BoneTransformIndex],
+				PrevEntry.GetComponentSpaceTransform(BoneTransformIndex),
+				NextEntry.GetComponentSpaceTransform(BoneTransformIndex),
 				LerpValue);
 		}
 		else
 		{
 			OutBoneTransform.Blend(
-				PrevEntry.ComponentSpaceTransforms[BoneTransformIndex] * PrevEntry.ComponentSpaceTransforms[ReferenceBoneTransformIndex].Inverse(),
-				NextEntry.ComponentSpaceTransforms[BoneTransformIndex] * NextEntry.ComponentSpaceTransforms[ReferenceBoneTransformIndex].Inverse(),
+				PrevEntry.GetComponentSpaceTransform(BoneTransformIndex) * PrevEntry.GetComponentSpaceTransform(ReferenceBoneTransformIndex).Inverse(),
+				NextEntry.GetComponentSpaceTransform(BoneTransformIndex) * NextEntry.GetComponentSpaceTransform(ReferenceBoneTransformIndex).Inverse(),
 				LerpValue);
 		}
 	}
@@ -251,7 +281,7 @@ void FPoseHistory::ClearHistory()
 	Entries.Reset();
 }
 
-void FPoseHistory::Update(float SecondsElapsed, FCSPose<FCompactPose>& ComponentSpacePose)
+void FPoseHistory::Update(float SecondsElapsed, FCSPose<FCompactPose>& ComponentSpacePose, bool bStoreScales)
 {
 	const USkeleton* Skeleton = ComponentSpacePose.GetPose().GetBoneContainer().GetSkeletonAsset();
 	if (LastUpdateSkeleton != Skeleton)
@@ -296,7 +326,7 @@ void FPoseHistory::Update(float SecondsElapsed, FCSPose<FCompactPose>& Component
 	}
 
 	// Regardless of the retention policy, we always update the most recent Entry
-	Entries.Last().Update(0.f, ComponentSpacePose, BoneToTransformMap);
+	Entries.Last().Update(0.f, ComponentSpacePose, BoneToTransformMap, bStoreScales);
 }
 
 #if ENABLE_DRAW_DEBUG && ENABLE_ANIM_DEBUG
@@ -308,25 +338,25 @@ void FPoseHistory::DebugDraw(FAnimInstanceProxy& AnimInstanceProxy, FColor Color
 	for (int32 EntryIndex = 0; EntryIndex < Entries.Num(); ++EntryIndex)
 	{
 		const FPoseHistoryEntry& Entry = Entries[EntryIndex];
-		if (Entry.ComponentSpaceTransforms.IsEmpty())
+		if (Entry.Num() == 0)
 		{
 			PrevGlobalTransforms.Reset();
 		}
-		else if (PrevGlobalTransforms.Num() != Entry.ComponentSpaceTransforms.Num())
+		else if (PrevGlobalTransforms.Num() != Entry.Num())
 		{
-			PrevGlobalTransforms.SetNum(Entry.ComponentSpaceTransforms.Num());
-			for (int32 i = 0; i < Entry.ComponentSpaceTransforms.Num(); ++i)
+			PrevGlobalTransforms.SetNum(Entry.Num());
+			for (int32 i = 0; i < Entry.Num(); ++i)
 			{
 				const FTransform RootTransform = bValidTrajectory ? Trajectory->GetSampleAtTime(-Entry.Time).GetTransform() : AnimInstanceProxy.GetComponentTransform();
-				PrevGlobalTransforms[i] = Entry.ComponentSpaceTransforms[i] * RootTransform;
+				PrevGlobalTransforms[i] = Entry.GetComponentSpaceTransform(i) * RootTransform;
 			}
 		}
 		else
 		{
-			for (int32 i = 0; i < Entry.ComponentSpaceTransforms.Num(); ++i)
+			for (int32 i = 0; i < Entry.Num(); ++i)
 			{
 				const FTransform RootTransform = bValidTrajectory ? Trajectory->GetSampleAtTime(-Entry.Time).GetTransform() : AnimInstanceProxy.GetComponentTransform();
-				const FTransform GlobalTransforms = Entry.ComponentSpaceTransforms[i] * RootTransform;
+				const FTransform GlobalTransforms = Entry.GetComponentSpaceTransform(i) * RootTransform;
 
 				AnimInstanceProxy.AnimDrawDebugLine(PrevGlobalTransforms[i].GetTranslation(), GlobalTransforms.GetTranslation(), Color, false, 0.f, ESceneDepthPriorityGroup::SDPG_Foreground);
 
@@ -390,7 +420,7 @@ void FExtendedPoseHistory::AddFuturePose(float SecondsInTheFuture, FCSPose<FComp
 	check(PoseHistory);	
 	const float SecondsAgo = -SecondsInTheFuture;
 	const int32 LowerBoundIdx = Algo::LowerBound(FutureEntries, SecondsAgo, [](const FPoseHistoryEntry& Entry, float Value) { return Value < Entry.Time; });
-	FutureEntries.InsertDefaulted_GetRef(LowerBoundIdx).Update(SecondsAgo, ComponentSpacePose, PoseHistory->GetBoneToTransformMap());
+	FutureEntries.InsertDefaulted_GetRef(LowerBoundIdx).Update(SecondsAgo, ComponentSpacePose, PoseHistory->GetBoneToTransformMap(), true);
 }
 
 #if ENABLE_DRAW_DEBUG && ENABLE_ANIM_DEBUG
@@ -404,25 +434,25 @@ void FExtendedPoseHistory::DebugDraw(FAnimInstanceProxy& AnimInstanceProxy, FCol
 	for (int32 EntryIndex = 0; EntryIndex < FutureEntries.Num(); ++EntryIndex)
 	{
 		const FPoseHistoryEntry& Entry = FutureEntries[EntryIndex];
-		if (Entry.ComponentSpaceTransforms.IsEmpty())
+		if (Entry.Num() == 0)
 		{
 			PrevGlobalTransforms.Reset();
 		}
-		else if (PrevGlobalTransforms.Num() != Entry.ComponentSpaceTransforms.Num())
+		else if (PrevGlobalTransforms.Num() != Entry.Num())
 		{
-			PrevGlobalTransforms.SetNum(Entry.ComponentSpaceTransforms.Num());
-			for (int32 i = 0; i < Entry.ComponentSpaceTransforms.Num(); ++i)
+			PrevGlobalTransforms.SetNum(Entry.Num());
+			for (int32 i = 0; i < Entry.Num(); ++i)
 			{
 				const FTransform RootTransform = bValidTrajectory ? Trajectory->GetSampleAtTime(-Entry.Time).GetTransform() : AnimInstanceProxy.GetComponentTransform();
-				PrevGlobalTransforms[i] = Entry.ComponentSpaceTransforms[i] * RootTransform;
+				PrevGlobalTransforms[i] = Entry.GetComponentSpaceTransform(i) * RootTransform;
 			}
 		}
 		else
 		{
-			for (int32 i = 0; i < Entry.ComponentSpaceTransforms.Num(); ++i)
+			for (int32 i = 0; i < Entry.Num(); ++i)
 			{
 				const FTransform RootTransform = bValidTrajectory ? Trajectory->GetSampleAtTime(-Entry.Time).GetTransform() : AnimInstanceProxy.GetComponentTransform();
-				const FTransform GlobalTransforms = Entry.ComponentSpaceTransforms[i] * RootTransform;
+				const FTransform GlobalTransforms = Entry.GetComponentSpaceTransform(i) * RootTransform;
 
 				AnimInstanceProxy.AnimDrawDebugLine(PrevGlobalTransforms[i].GetTranslation(), GlobalTransforms.GetTranslation(), Color, false, 0.f, ESceneDepthPriorityGroup::SDPG_Foreground);
 
diff --git a/Engine/Plugins/Experimental/Animation/PoseSearch/Source/Runtime/Public/PoseSearch/AnimNode_PoseSearchHistoryCollector.h b/Engine/Plugins/Experimental/Animation/PoseSearch/Source/Runtime/Public/PoseSearch/AnimNode_PoseSearchHistoryCollector.h
--- a/Engine/Plugins/Experimental/Animation/PoseSearch/Source/Runtime/Public/PoseSearch/AnimNode_PoseSearchHistoryCollector.h
+++ b/Engine/Plugins/Experimental/Animation/PoseSearch/Source/Runtime/Public/PoseSearch/AnimNode_PoseSearchHistoryCollector.h
@@ -33,6 +33,10 @@ public:
 	UPROPERTY(EditAnywhere, Category = Settings)
 	bool bResetOnBecomingRelevant = true;
 
+	// if true pose scales will be cached, otherwise implied to be unitary scales
+	UPROPERTY(EditAnywhere, Category = Settings)
+	bool bStoreScales = true;
+
 	// Update Counter for detecting being relevant
 	FGraphTraversalCounter UpdateCounter;
 
diff --git a/Engine/Plugins/Experimental/Animation/PoseSearch/Source/Runtime/Public/PoseSearch/PoseSearchHistory.h b/Engine/Plugins/Experimental/Animation/PoseSearch/Source/Runtime/Public/PoseSearch/PoseSearchHistory.h
--- a/Engine/Plugins/Experimental/Animation/PoseSearch/Source/Runtime/Public/PoseSearch/PoseSearchHistory.h
+++ b/Engine/Plugins/Experimental/Animation/PoseSearch/Source/Runtime/Public/PoseSearch/PoseSearchHistory.h
@@ -40,10 +40,18 @@ struct POSESEARCH_API IPoseHistory
 struct FPoseHistoryEntry
 {
 	// collected bones transforms in component space
-	TArray<FTransform> ComponentSpaceTransforms;
+	TArray<FQuat4f> ComponentSpaceRotations;
+	TArray<FVector> ComponentSpacePositions;
+	TArray<FVector3f> ComponentSpaceScales;
 	float Time = 0.f;
 
-	void Update(float InTime, FCSPose<FCompactPose>& ComponentSpacePose, const FBoneToTransformMap& BoneToTransformMap);
+	void Update(float InTime, FCSPose<FCompactPose>& ComponentSpacePose, const FBoneToTransformMap& BoneToTransformMap, bool bStoreScales);
+
+	void SetNum(int32 Num, bool bStoreScales);
+	int32 Num() const;
+
+	void SetComponentSpaceTransform(int32 Index, const FTransform& Transform);
+	FTransform GetComponentSpaceTransform(int32 Index) const;
 };
 
 typedef TRingBuffer<FPoseHistoryEntry> FPoseHistoryEntries;
@@ -52,7 +60,7 @@ typedef TArray<FPoseHistoryEntry> FPoseHistoryFutureEntries;
 struct FPoseHistory : public IPoseHistory
 {
 	void Init(int32 InNumPoses, float InTimeHorizon, const TArray<FBoneIndexType>& RequiredBones);
-	void Update(float SecondsElapsed, FCSPose<FCompactPose>& ComponentSpacePose);
+	void Update(float SecondsElapsed, FCSPose<FCompactPose>& ComponentSpacePose, bool bStoreScales);
 	float GetTimeHorizon() const { return TimeHorizon; }
 
 	const FBoneToTransformMap& GetBoneToTransformMap() const { return BoneToTransformMap; }
