motion matching - reducing FSearchContext memory footprint

diff --git a/Engine/Plugins/Experimental/Animation/PoseSearch/Source/Runtime/Public/PoseSearch/PoseSearchContext.h b/Engine/Plugins/Experimental/Animation/PoseSearch/Source/Runtime/Public/PoseSearch/PoseSearchContext.h
--- a/Engine/Plugins/Experimental/Animation/PoseSearch/Source/Runtime/Public/PoseSearch/PoseSearchContext.h
+++ b/Engine/Plugins/Experimental/Animation/PoseSearch/Source/Runtime/Public/PoseSearch/PoseSearchContext.h
@@ -207,11 +207,11 @@ private:
 
 	// Trajectory has been transformed in root bone world space reference system
 	const FPoseSearchQueryTrajectory* Trajectory = nullptr;
-	float DesiredPermutationTimeOffset = 0.f;
+	const float DesiredPermutationTimeOffset = 0.f;
 	const FPoseIndicesHistory* PoseIndicesHistory = nullptr;
-	FSearchResult CurrentResult;
-	FFloatInterval PoseJumpThresholdTime = FFloatInterval(0.f, 0.f);
-	bool bForceInterrupt = false;
+	const FSearchResult& CurrentResult;
+	const FFloatInterval& PoseJumpThresholdTime;
+	const bool bForceInterrupt = false;
 
 	TConstArrayView<float> CurrentResultPoseVector;
 	TStackAlignedArray<float> CurrentResultPoseVectorData;
