motion matching - fix for database Data Details tab refresh on an updated schema with different cardinality

diff --git a/Engine/Plugins/Experimental/Animation/PoseSearch/Source/Editor/Private/PoseSearchDatabaseDataDetails.cpp b/Engine/Plugins/Experimental/Animation/PoseSearch/Source/Editor/Private/PoseSearchDatabaseDataDetails.cpp
--- a/Engine/Plugins/Experimental/Animation/PoseSearch/Source/Editor/Private/PoseSearchDatabaseDataDetails.cpp
+++ b/Engine/Plugins/Experimental/Animation/PoseSearch/Source/Editor/Private/PoseSearchDatabaseDataDetails.cpp
@@ -241,15 +241,31 @@ public:
 									{
 										StringBuilder.Append(TEXT(", "));
 									}
-									const float Value = PoseValues[i + DataOffset];
-									StringBuilder.Appendf(TEXT("%.1f"), Value);
+
+									const int32 PoseValueIndex = i + DataOffset;
+									if (PoseValueIndex < PoseValues.Num())
+									{
+										const float Value = PoseValues[PoseValueIndex];
+										StringBuilder.Appendf(TEXT("%.1f"), Value);
+									}
+									else
+									{
+										StringBuilder.Append(TEXT("---"));
+									}
 								}
 							}
 							else
 							{
 								// using all the float digits 
-								const float Value = PoseValues[DataOffset];
-								StringBuilder.Appendf(TEXT("%f"), Value);
+								if (DataOffset < PoseValues.Num())
+								{
+									const float Value = PoseValues[DataOffset];
+									StringBuilder.Appendf(TEXT("%f"), Value);
+								}
+								else
+								{
+									StringBuilder.Append(TEXT("---"));
+								}
 							}
 						}
 					}
