motion matching - handling SelfAdjointEigenSolver errors

diff --git a/Engine/Plugins/Experimental/Animation/PoseSearch/Source/Runtime/Private/PoseSearchDerivedData.cpp b/Engine/Plugins/Experimental/Animation/PoseSearch/Source/Runtime/Private/PoseSearchDerivedData.cpp
--- a/Engine/Plugins/Experimental/Animation/PoseSearch/Source/Runtime/Private/PoseSearchDerivedData.cpp
+++ b/Engine/Plugins/Experimental/Animation/PoseSearch/Source/Runtime/Private/PoseSearchDerivedData.cpp
@@ -472,7 +472,7 @@ static void PreprocessSearchIndexWeights(FSearchIndex& SearchIndex, const UPoseS
 }
 
 // it calculates Mean, PCAValues, and PCAProjectionMatrix
-static void PreprocessSearchIndexPCAData(FSearchIndex& SearchIndex, int32 NumDimensions, int32 NumberOfPrincipalComponents, EPoseSearchMode PoseSearchMode)
+static Eigen::ComputationInfo PreprocessSearchIndexPCAData(FSearchIndex& SearchIndex, int32 NumDimensions, int32 NumberOfPrincipalComponents, EPoseSearchMode PoseSearchMode)
 {
 	// binding SearchIndex.Values and SearchIndex.PCAValues Eigen row major matrix maps
 	const int32 NumPoses = SearchIndex.GetNumPoses();
@@ -483,6 +483,7 @@ static void PreprocessSearchIndexPCAData(FSearchIndex& SearchIndex, int32 NumDim
 	SearchIndex.Mean.Reset();
 	SearchIndex.PCAProjectionMatrix.Reset();
 
+	Eigen::ComputationInfo ComputationInfo = Eigen::Success;
 	if (PoseSearchMode == EPoseSearchMode::PCAKDTree && NumDimensions > 0 && NumPoses > 0 && NumberOfPrincipalComponents > 0)
 	{
 		SearchIndex.PCAValues.AddZeroed(NumPoses * NumberOfPrincipalComponents);
@@ -515,84 +516,88 @@ static void PreprocessSearchIndexPCAData(FSearchIndex& SearchIndex, int32 NumDim
 		const ColMajorMatrix CovariantMatrix = (CenteredValues.transpose() * CenteredValues) / float(NumPoses - 1);
 		const Eigen::SelfAdjointEigenSolver<ColMajorMatrix> EigenSolver(CovariantMatrix);
 
-		check(EigenSolver.info() == Eigen::Success);
-
-		// validating EigenSolver results
-		const ColMajorMatrix EigenVectors = EigenSolver.eigenvectors().real();
+		ComputationInfo = EigenSolver.info();
+		if (ComputationInfo == Eigen::Success)
+		{
+			// validating EigenSolver results
+			const ColMajorMatrix EigenVectors = EigenSolver.eigenvectors().real();
 
 #if ENABLE_ANIM_DEBUG
-		if (AnyTestFlags(EMotionMatchTestFlags::ValidateKDTreeConstruct) && NumberOfPrincipalComponents == NumDimensions)
-		{
-			const RowMajorVector ReciprocalWeightsSqrt = MapWeightsSqrt.cwiseInverse();
-			const RowMajorMatrix ProjectedValues = CenteredValues * EigenVectors;
-			for (Eigen::Index RowIndex = 0; RowIndex < MapValues.rows(); ++RowIndex)
+			if (AnyTestFlags(EMotionMatchTestFlags::ValidateKDTreeConstruct) && NumberOfPrincipalComponents == NumDimensions)
 			{
-				const RowMajorVector WeightedReconstructedPoint = ProjectedValues.row(RowIndex) * EigenVectors.transpose() + MapMean;
-				const RowMajorVector ReconstructedPoint = WeightedReconstructedPoint.array() * ReciprocalWeightsSqrt.array();
-				const float Error = (ReconstructedPoint - MapValues.row(RowIndex)).squaredNorm();
-				check(Error < UE_KINDA_SMALL_NUMBER);
+				const RowMajorVector ReciprocalWeightsSqrt = MapWeightsSqrt.cwiseInverse();
+				const RowMajorMatrix ProjectedValues = CenteredValues * EigenVectors;
+				for (Eigen::Index RowIndex = 0; RowIndex < MapValues.rows(); ++RowIndex)
+				{
+					const RowMajorVector WeightedReconstructedPoint = ProjectedValues.row(RowIndex) * EigenVectors.transpose() + MapMean;
+					const RowMajorVector ReconstructedPoint = WeightedReconstructedPoint.array() * ReciprocalWeightsSqrt.array();
+					const float Error = (ReconstructedPoint - MapValues.row(RowIndex)).squaredNorm();
+					check(Error < UE_KINDA_SMALL_NUMBER);
+				}
 			}
-		}
 #endif // ENABLE_ANIM_DEBUG
 
-		// sorting EigenVectors by EigenValues, so we pick the most significant ones to compose our PCA projection matrix.
-		const RowMajorVector EigenValues = EigenSolver.eigenvalues().real();
-		TArray<int32> Indexer;
-		Indexer.Reserve(NumDimensions);
-		for (int32 DimensionIndex = 0; DimensionIndex < NumDimensions; ++DimensionIndex)
-		{
-			Indexer.Push(DimensionIndex);
-		}
-		Indexer.Sort([&EigenValues](int32 a, int32 b)
-		{
-			return EigenValues[a] > EigenValues[b];
-		});
+			// sorting EigenVectors by EigenValues, so we pick the most significant ones to compose our PCA projection matrix.
+			const RowMajorVector EigenValues = EigenSolver.eigenvalues().real();
+			TArray<int32> Indexer;
+			Indexer.Reserve(NumDimensions);
+			for (int32 DimensionIndex = 0; DimensionIndex < NumDimensions; ++DimensionIndex)
+			{
+				Indexer.Push(DimensionIndex);
+			}
+			Indexer.Sort([&EigenValues](int32 a, int32 b)
+				{
+					return EigenValues[a] > EigenValues[b];
+				});
 
-		// composing the PCA projection matrix with the PCANumComponents most significant EigenVectors
-		ColMajorMatrixMap PCAProjectionMatrix(SearchIndex.PCAProjectionMatrix.GetData(), NumDimensions, NumberOfPrincipalComponents);
-		float AccumulatedVariance = 0.f;
-		for (int32 PCAComponentIndex = 0; PCAComponentIndex < NumberOfPrincipalComponents; ++PCAComponentIndex)
-		{
-			PCAProjectionMatrix.col(PCAComponentIndex) = EigenVectors.col(Indexer[PCAComponentIndex]);
-			AccumulatedVariance += EigenValues[Indexer[PCAComponentIndex]];
-		}
+			// composing the PCA projection matrix with the PCANumComponents most significant EigenVectors
+			ColMajorMatrixMap PCAProjectionMatrix(SearchIndex.PCAProjectionMatrix.GetData(), NumDimensions, NumberOfPrincipalComponents);
+			float AccumulatedVariance = 0.f;
+			for (int32 PCAComponentIndex = 0; PCAComponentIndex < NumberOfPrincipalComponents; ++PCAComponentIndex)
+			{
+				PCAProjectionMatrix.col(PCAComponentIndex) = EigenVectors.col(Indexer[PCAComponentIndex]);
+				AccumulatedVariance += EigenValues[Indexer[PCAComponentIndex]];
+			}
 
-		// calculating the total variance knowing that eigen values measure variance along the principal components:
-		const float TotalVariance = EigenValues.sum();
-		// and explained variance as ratio between AccumulatedVariance and TotalVariance: https://ro-che.info/articles/2017-12-11-pca-explained-variance
-		SearchIndex.PCAExplainedVariance = TotalVariance > UE_KINDA_SMALL_NUMBER ? AccumulatedVariance / TotalVariance : 0.f;
+			// calculating the total variance knowing that eigen values measure variance along the principal components:
+			const float TotalVariance = EigenValues.sum();
+			// and explained variance as ratio between AccumulatedVariance and TotalVariance: https://ro-che.info/articles/2017-12-11-pca-explained-variance
+			SearchIndex.PCAExplainedVariance = TotalVariance > UE_KINDA_SMALL_NUMBER ? AccumulatedVariance / TotalVariance : 0.f;
 
-		MapPCAValues = CenteredValues * PCAProjectionMatrix;
+			MapPCAValues = CenteredValues * PCAProjectionMatrix;
 
 #if ENABLE_ANIM_DEBUG
-		if (AnyTestFlags(EMotionMatchTestFlags::ValidateKDTreeConstruct) && NumberOfPrincipalComponents == NumDimensions)
-		{
-			const RowMajorVector ReciprocalWeightsSqrt = MapWeightsSqrt.cwiseInverse();
-			for (Eigen::Index RowIndex = 0; RowIndex < MapValues.rows(); ++RowIndex)
+			if (AnyTestFlags(EMotionMatchTestFlags::ValidateKDTreeConstruct) && NumberOfPrincipalComponents == NumDimensions)
 			{
-				const RowMajorVector WeightedReconstructedValues = MapPCAValues.row(RowIndex) * PCAProjectionMatrix.transpose() + MapMean;
-				const RowMajorVector ReconstructedValues = WeightedReconstructedValues.array() * ReciprocalWeightsSqrt.array();
-				const float Error = (ReconstructedValues - MapValues.row(RowIndex)).squaredNorm();
-				check(Error < UE_KINDA_SMALL_NUMBER);
-			}
+				const RowMajorVector ReciprocalWeightsSqrt = MapWeightsSqrt.cwiseInverse();
+				for (Eigen::Index RowIndex = 0; RowIndex < MapValues.rows(); ++RowIndex)
+				{
+					const RowMajorVector WeightedReconstructedValues = MapPCAValues.row(RowIndex) * PCAProjectionMatrix.transpose() + MapMean;
+					const RowMajorVector ReconstructedValues = WeightedReconstructedValues.array() * ReciprocalWeightsSqrt.array();
+					const float Error = (ReconstructedValues - MapValues.row(RowIndex)).squaredNorm();
+					check(Error < UE_KINDA_SMALL_NUMBER);
+				}
 
-			TArray<float> ReconstructedPoseValues;
-			ReconstructedPoseValues.SetNumZeroed(NumDimensions);
-			for (int32 PoseIdx = 0; PoseIdx < NumPoses; ++PoseIdx)
-			{
-				SearchIndex.GetReconstructedPoseValues(PoseIdx, ReconstructedPoseValues);
-				TConstArrayView<float> PoseValues = SearchIndex.GetPoseValues(PoseIdx);
+				TArray<float> ReconstructedPoseValues;
+				ReconstructedPoseValues.SetNumZeroed(NumDimensions);
+				for (int32 PoseIdx = 0; PoseIdx < NumPoses; ++PoseIdx)
+				{
+					SearchIndex.GetReconstructedPoseValues(PoseIdx, ReconstructedPoseValues);
+					TConstArrayView<float> PoseValues = SearchIndex.GetPoseValues(PoseIdx);
 
-				check(ReconstructedPoseValues.Num() == PoseValues.Num());
-				Eigen::Map<const Eigen::ArrayXf> VA(ReconstructedPoseValues.GetData(), ReconstructedPoseValues.Num());
-				Eigen::Map<const Eigen::ArrayXf> VB(PoseValues.GetData(), PoseValues.Num());
+					check(ReconstructedPoseValues.Num() == PoseValues.Num());
+					Eigen::Map<const Eigen::ArrayXf> VA(ReconstructedPoseValues.GetData(), ReconstructedPoseValues.Num());
+					Eigen::Map<const Eigen::ArrayXf> VB(PoseValues.GetData(), PoseValues.Num());
 
-				const float Error = (VA - VB).square().sum();
-				check(Error < UE_KINDA_SMALL_NUMBER);
+					const float Error = (VA - VB).square().sum();
+					check(Error < UE_KINDA_SMALL_NUMBER);
+				}
 			}
-		}
 #endif // ENABLE_ANIM_DEBUG
+		}
 	}
+
+	return ComputationInfo;
 }
 
 static void PreprocessSearchIndexKDTree(FSearchIndex& SearchIndex, const UPoseSearchDatabase* Database)
@@ -1561,7 +1566,26 @@ void FPoseSearchDatabaseAsyncCacheTask::OnGetComplete(UE::DerivedData::FCacheGet
 					return;
 				}
 
-				PreprocessSearchIndexPCAData(SearchIndex, IndexBaseDatabases[0]->Schema->SchemaCardinality, IndexBaseDatabases[0]->GetNumberOfPrincipalComponents(), IndexBaseDatabases[0]->PoseSearchMode);
+				const Eigen::ComputationInfo ComputationInfo = PreprocessSearchIndexPCAData(SearchIndex, IndexBaseDatabases[0]->Schema->SchemaCardinality, IndexBaseDatabases[0]->GetNumberOfPrincipalComponents(), IndexBaseDatabases[0]->PoseSearchMode);
+				if (ComputationInfo != Eigen::Success)
+				{
+					FString Reason;
+					switch (ComputationInfo)
+					{
+					case Eigen::NumericalIssue: Reason = "Numerical Issues"; break;
+					case Eigen::NoConvergence:	Reason = "No Convergence";   break;
+					case Eigen::InvalidInput:	Reason = "Invalid Input"; 	 break;
+					default:					Reason = "Unknown Reasons";  break;
+					}
+					UE_LOG(LogPoseSearch, Error, TEXT("%s - %s BuildIndex Failed because of '%s' while calculating PCA data. Try with a different dataset or change the database 'Pose Search Mode'"), *LexToString(FullIndexKey.Hash), *IndexBaseDatabases[0]->GetName(), *Reason);
+
+					SearchIndex.Reset();
+#if ENABLE_ANIM_DEBUG
+					SearchIndexCompare.Reset();
+#endif //ENABLE_ANIM_DEBUG
+					return;
+				}
+
 				if (Owner.IsCanceled())
 				{
 					UE_LOG(LogPoseSearch, Log, TEXT("%s - %s BuildIndex Cancelled"), *LexToString(FullIndexKey.Hash), *IndexBaseDatabases[0]->GetName());
