Change void FAnimNode_MotionMatching::SetDatabasesToSearch() to support stack allocated arrays

diff --git a/Engine/Plugins/Experimental/Animation/PoseSearch/Source/Runtime/Private/AnimNode_MotionMatching.cpp b/Engine/Plugins/Experimental/Animation/PoseSearch/Source/Runtime/Private/AnimNode_MotionMatching.cpp
--- a/Engine/Plugins/Experimental/Animation/PoseSearch/Source/Runtime/Private/AnimNode_MotionMatching.cpp
+++ b/Engine/Plugins/Experimental/Animation/PoseSearch/Source/Runtime/Private/AnimNode_MotionMatching.cpp
@@ -231,7 +231,7 @@ void FAnimNode_MotionMatching::SetDatabaseToSearch(UPoseSearchDatabase* InDataba
 	}
 }
 
-void FAnimNode_MotionMatching::SetDatabasesToSearch(const TArray<UPoseSearchDatabase*>& InDatabases, bool bForceInterruptIfNew)
+void FAnimNode_MotionMatching::SetDatabasesToSearch(TConstArrayView<UPoseSearchDatabase*> InDatabases, bool bForceInterruptIfNew)
 {
 	// Check if InDatabases and DatabasesToSearch are the same.
 	bool bDatabasesAlreadySet = true;
diff --git a/Engine/Plugins/Experimental/Animation/PoseSearch/Source/Runtime/Public/PoseSearch/AnimNode_MotionMatching.h b/Engine/Plugins/Experimental/Animation/PoseSearch/Source/Runtime/Public/PoseSearch/AnimNode_MotionMatching.h
--- a/Engine/Plugins/Experimental/Animation/PoseSearch/Source/Runtime/Public/PoseSearch/AnimNode_MotionMatching.h
+++ b/Engine/Plugins/Experimental/Animation/PoseSearch/Source/Runtime/Public/PoseSearch/AnimNode_MotionMatching.h
@@ -21,7 +21,7 @@ public:
 	void SetDatabaseToSearch(UPoseSearchDatabase* InDatabase, bool bForceInterruptIfNew);
 
 	// Search InDatabases instead of the Database property on the node. Use bForceInterruptIfNew to ignore the continuing pose if InDatabases is new.
-	void SetDatabasesToSearch(const TArray<UPoseSearchDatabase*>& InDatabases, bool bForceInterruptIfNew);
+	void SetDatabasesToSearch(TConstArrayView<UPoseSearchDatabase*> InDatabases, bool bForceInterruptIfNew);
 
 	// Reset the effects of SetDatabaseToSearch/SetDatabasesToSearch and use the Database property on this node.
 	void ResetDatabasesToSearch(bool bInForceInterrupt);
