Don't use anim trace in test/shipping builds

diff --git a/Engine/Plugins/Experimental/Animation/PoseSearch/Source/Runtime/Public/PoseSearch/PoseSearchDefines.h b/Engine/Plugins/Experimental/Animation/PoseSearch/Source/Runtime/Public/PoseSearch/PoseSearchDefines.h
--- a/Engine/Plugins/Experimental/Animation/PoseSearch/Source/Runtime/Public/PoseSearch/PoseSearchDefines.h
+++ b/Engine/Plugins/Experimental/Animation/PoseSearch/Source/Runtime/Public/PoseSearch/PoseSearchDefines.h
@@ -4,10 +4,10 @@
 
 #include "Async/ParallelFor.h"
 #include "Logging/LogMacros.h"
-#include "ObjectTrace.h"
+#include "Animation/AnimTrace.h"
 
 // Enable this if object tracing is enabled, mimics animation tracing
-#define UE_POSE_SEARCH_TRACE_ENABLED OBJECT_TRACE_ENABLED
+#define UE_POSE_SEARCH_TRACE_ENABLED ANIM_TRACE_ENABLED
 
 #ifndef UE_POSE_SEARCH_FORCE_SINGLE_THREAD
 #define UE_POSE_SEARCH_FORCE_SINGLE_THREAD 0
diff --git a/Engine/Source/Runtime/Engine/Public/Animation/AnimTrace.h b/Engine/Source/Runtime/Engine/Public/Animation/AnimTrace.h
--- a/Engine/Source/Runtime/Engine/Public/Animation/AnimTrace.h
+++ b/Engine/Source/Runtime/Engine/Public/Animation/AnimTrace.h
@@ -7,7 +7,7 @@
 #include "CoreMinimal.h"
 #include "Trace/Trace.h"
 
-#define ANIM_TRACE_ENABLED OBJECT_TRACE_ENABLED
+#define ANIM_TRACE_ENABLED (OBJECT_TRACE_ENABLED && !(UE_BUILD_SHIPPING || UE_BUILD_TEST))
 
 #if ANIM_TRACE_ENABLED
 
