motion matching - pose search database commands cleanup

diff --git a/Engine/Plugins/Experimental/Animation/PoseSearch/Source/Editor/Private/PoseSearchDatabaseEditorCommands.cpp b/Engine/Plugins/Experimental/Animation/PoseSearch/Source/Editor/Private/PoseSearchDatabaseEditorCommands.cpp
--- a/Engine/Plugins/Experimental/Animation/PoseSearch/Source/Editor/Private/PoseSearchDatabaseEditorCommands.cpp
+++ b/Engine/Plugins/Experimental/Animation/PoseSearch/Source/Editor/Private/PoseSearchDatabaseEditorCommands.cpp
@@ -8,11 +8,6 @@ namespace UE::PoseSearch
 {
 	void FDatabaseEditorCommands::RegisterCommands()
 	{
-		UI_COMMAND(ShowPoseFeaturesNone, "None", "Hide all features", EUserInterfaceActionType::RadioButton, FInputChord());
-		UI_COMMAND(ShowPoseFeaturesAll, "All Features", "Show all features", EUserInterfaceActionType::RadioButton, FInputChord());
-		UI_COMMAND(ShowPoseFeaturesDetailed, "Detailed Features", "Show all features with mode details", EUserInterfaceActionType::RadioButton, FInputChord());
-		UI_COMMAND(ShowAnimationOriginalOnly, "Original Only", "Show only original animations", EUserInterfaceActionType::RadioButton, FInputChord());
-		UI_COMMAND(ShowAnimationOriginalAndMirrored, "Original and mirrored", "Show original and mirrored animations", EUserInterfaceActionType::RadioButton, FInputChord());
 		UI_COMMAND(ShowDisplayRootMotionSpeed, "Display Root Motion Speed", "Show root motion speed for the selected item", EUserInterfaceActionType::ToggleButton, FInputChord());
 		UI_COMMAND(ShowQuantizeAnimationToPoseData, "Quantize Animation To Pose Data", "Animations will not playing continuously, only using the associated pose quantized time", EUserInterfaceActionType::ToggleButton, FInputChord());
 		UI_COMMAND(ShowBones, "Show All Bones", "Debug Draw All Bones", EUserInterfaceActionType::ToggleButton, FInputChord());
diff --git a/Engine/Plugins/Experimental/Animation/PoseSearch/Source/Editor/Private/PoseSearchDatabasePreviewScene.cpp b/Engine/Plugins/Experimental/Animation/PoseSearch/Source/Editor/Private/PoseSearchDatabasePreviewScene.cpp
--- a/Engine/Plugins/Experimental/Animation/PoseSearch/Source/Editor/Private/PoseSearchDatabasePreviewScene.cpp
+++ b/Engine/Plugins/Experimental/Animation/PoseSearch/Source/Editor/Private/PoseSearchDatabasePreviewScene.cpp
@@ -63,8 +63,7 @@ namespace UE::PoseSearch
 		FDatabaseViewModel* ViewModel = GetEditor()->GetViewModel();
 		const UPoseSearchDatabase* Database = ViewModel->GetPoseSearchDatabase();
 
-		if (ViewModel->IsPoseFeaturesDrawMode(EFeaturesDrawMode::All | EFeaturesDrawMode::Detailed) && !ViewModel->GetPreviewActors().IsEmpty() &&
-			FAsyncPoseSearchDatabasesManagement::RequestAsyncBuildIndex(Database, ERequestAsyncBuildFlag::ContinueRequest))
+		if (!ViewModel->GetPreviewActors().IsEmpty() && FAsyncPoseSearchDatabasesManagement::RequestAsyncBuildIndex(Database, ERequestAsyncBuildFlag::ContinueRequest))
 		{
 			bool bDrawQueryVector = ViewModel->ShouldDrawQueryVector();
 			for (FDatabasePreviewActor& PreviewActor : ViewModel->GetPreviewActors())
diff --git a/Engine/Plugins/Experimental/Animation/PoseSearch/Source/Editor/Private/PoseSearchDatabaseViewModel.cpp b/Engine/Plugins/Experimental/Animation/PoseSearch/Source/Editor/Private/PoseSearchDatabaseViewModel.cpp
--- a/Engine/Plugins/Experimental/Animation/PoseSearch/Source/Editor/Private/PoseSearchDatabaseViewModel.cpp
+++ b/Engine/Plugins/Experimental/Animation/PoseSearch/Source/Editor/Private/PoseSearchDatabaseViewModel.cpp
@@ -305,30 +305,6 @@ namespace UE::PoseSearch
 		return RootTransform;
 	}
 
-	void FDatabaseViewModel::OnSetPoseFeaturesDrawMode(EFeaturesDrawMode DrawMode)
-	{
-		PoseFeaturesDrawMode = DrawMode;
-	}
-
-	bool FDatabaseViewModel::IsPoseFeaturesDrawMode(EFeaturesDrawMode DrawMode) const
-	{
-		return EnumHasAnyFlags(PoseFeaturesDrawMode, DrawMode);
-	}
-
-	void FDatabaseViewModel::OnSetAnimationPreviewMode(EAnimationPreviewMode PreviewMode)
-	{
-		if (PreviewMode != AnimationPreviewMode)
-		{
-			AnimationPreviewMode = PreviewMode;
-			RemovePreviewActors();
-		}
-	}
-
-	bool FDatabaseViewModel::IsAnimationPreviewMode(EAnimationPreviewMode PreviewMode) const
-	{
-		return EnumHasAnyFlags(AnimationPreviewMode, PreviewMode);
-	}
-
 	void FDatabaseViewModel::AddSequenceToDatabase(UAnimSequence* AnimSequence)
 	{
 		FPoseSearchDatabaseSequence NewAsset;
@@ -477,16 +453,13 @@ namespace UE::PoseSearch
 			for (int32 IndexAssetIndex = 0; IndexAssetIndex < SearchIndex.Assets.Num(); ++IndexAssetIndex)
 			{
 				const FSearchIndexAsset& IndexAsset = SearchIndex.Assets[IndexAssetIndex];
-				if (AnimationPreviewMode == EAnimationPreviewMode::OriginalAndMirrored || !IndexAsset.IsMirrored())
+				if (const int32* SelectedNodesIndex = AssociatedAssetIndices.Find(IndexAsset.GetSourceAssetIdx()))
 				{
-					if (const int32* SelectedNodesIndex = AssociatedAssetIndices.Find(IndexAsset.GetSourceAssetIdx()))
+					FDatabasePreviewActor PreviewActor = SpawnPreviewActor(IndexAssetIndex);
+					if (PreviewActor.IsValid())
 					{
-						FDatabasePreviewActor PreviewActor = SpawnPreviewActor(IndexAssetIndex);
-						if (PreviewActor.IsValid())
-						{
-							MaxPreviewPlayLength = FMath::Max(MaxPreviewPlayLength, IndexAsset.GetLastSampleTime(PoseSearchDatabase->Schema->SampleRate) - IndexAsset.GetFirstSampleTime(PoseSearchDatabase->Schema->SampleRate));
-							PreviewActors.Add(PreviewActor);
-						}
+						MaxPreviewPlayLength = FMath::Max(MaxPreviewPlayLength, IndexAsset.GetLastSampleTime(PoseSearchDatabase->Schema->SampleRate) - IndexAsset.GetFirstSampleTime(PoseSearchDatabase->Schema->SampleRate));
+						PreviewActors.Add(PreviewActor);
 					}
 				}
 			}
diff --git a/Engine/Plugins/Experimental/Animation/PoseSearch/Source/Editor/Private/SPoseSearchDatabaseViewport.cpp b/Engine/Plugins/Experimental/Animation/PoseSearch/Source/Editor/Private/SPoseSearchDatabaseViewport.cpp
--- a/Engine/Plugins/Experimental/Animation/PoseSearch/Source/Editor/Private/SPoseSearchDatabaseViewport.cpp
+++ b/Engine/Plugins/Experimental/Animation/PoseSearch/Source/Editor/Private/SPoseSearchDatabaseViewport.cpp
@@ -37,66 +37,6 @@ namespace UE::PoseSearch
 		TSharedRef<FDatabaseViewModel> ViewModelRef =
 			AssetEditorPtr.Pin()->GetViewModelSharedPtr().ToSharedRef();
 
-		CommandList->MapAction(
-			Commands.ShowPoseFeaturesNone,
-			FExecuteAction::CreateSP(
-				ViewModelRef,
-				&FDatabaseViewModel::OnSetPoseFeaturesDrawMode,
-				EFeaturesDrawMode::None),
-			FCanExecuteAction(),
-			FIsActionChecked::CreateSP(
-				ViewModelRef,
-				&FDatabaseViewModel::IsPoseFeaturesDrawMode,
-				EFeaturesDrawMode::None));
-
-		CommandList->MapAction(
-			Commands.ShowPoseFeaturesAll,
-			FExecuteAction::CreateSP(
-				ViewModelRef,
-				&FDatabaseViewModel::OnSetPoseFeaturesDrawMode,
-				EFeaturesDrawMode::All),
-			FCanExecuteAction(),
-			FIsActionChecked::CreateSP(
-				ViewModelRef,
-				&FDatabaseViewModel::IsPoseFeaturesDrawMode,
-				EFeaturesDrawMode::All));
-
-		CommandList->MapAction(
-			Commands.ShowPoseFeaturesDetailed,
-			FExecuteAction::CreateSP(
-				ViewModelRef,
-				&FDatabaseViewModel::OnSetPoseFeaturesDrawMode,
-				EFeaturesDrawMode::Detailed),
-			FCanExecuteAction(),
-			FIsActionChecked::CreateSP(
-				ViewModelRef,
-				&FDatabaseViewModel::IsPoseFeaturesDrawMode,
-				EFeaturesDrawMode::Detailed));
-
-		CommandList->MapAction(
-			Commands.ShowAnimationOriginalOnly,
-			FExecuteAction::CreateSP(
-				ViewModelRef,
-				&FDatabaseViewModel::OnSetAnimationPreviewMode,
-				EAnimationPreviewMode::OriginalOnly),
-			FCanExecuteAction(),
-			FIsActionChecked::CreateSP(
-				ViewModelRef,
-				&FDatabaseViewModel::IsAnimationPreviewMode,
-				EAnimationPreviewMode::OriginalOnly));
-
-		CommandList->MapAction(
-			Commands.ShowAnimationOriginalAndMirrored,
-			FExecuteAction::CreateSP(
-				ViewModelRef,
-				&FDatabaseViewModel::OnSetAnimationPreviewMode,
-				EAnimationPreviewMode::OriginalAndMirrored),
-			FCanExecuteAction(),
-			FIsActionChecked::CreateSP(
-				ViewModelRef,
-				&FDatabaseViewModel::IsAnimationPreviewMode,
-				EAnimationPreviewMode::OriginalAndMirrored));
-
 		CommandList->MapAction(
 			Commands.ShowDisplayRootMotionSpeed,
 			FExecuteAction::CreateSP(ViewModelRef, &FDatabaseViewModel::OnToggleDisplayRootMotionSpeed),
diff --git a/Engine/Plugins/Experimental/Animation/PoseSearch/Source/Editor/Private/SPoseSearchDatabaseViewportToolbar.cpp b/Engine/Plugins/Experimental/Animation/PoseSearch/Source/Editor/Private/SPoseSearchDatabaseViewportToolbar.cpp
--- a/Engine/Plugins/Experimental/Animation/PoseSearch/Source/Editor/Private/SPoseSearchDatabaseViewportToolbar.cpp
+++ b/Engine/Plugins/Experimental/Animation/PoseSearch/Source/Editor/Private/SPoseSearchDatabaseViewportToolbar.cpp
@@ -29,57 +29,12 @@ namespace UE::PoseSearch
 		const bool bInShouldCloseWindowAfterMenuSelection = true;
 		FMenuBuilder ShowMenuBuilder(bInShouldCloseWindowAfterMenuSelection, ViewportRef->GetCommandList());
 		{
-			ShowMenuBuilder.AddSubMenu(
-				LOCTEXT("ShowMenu_PoseFeaturesDrawSubMenu", "Pose Features"),
-				LOCTEXT("ShowMenu_PoseFeaturesDrawSubMenuToolTip", "Pose Feature Drawing Options"),
-				FNewMenuDelegate::CreateLambda([](FMenuBuilder& SubMenuBuilder)
-			{
-				const FDatabaseEditorCommands& Commands = FDatabaseEditorCommands::Get();
-
-				SubMenuBuilder.BeginSection("PoseFeatures", LOCTEXT("ShowMenu_PoseFeaturesLabel", "Pose Features"));
-				{
-					SubMenuBuilder.AddMenuEntry(Commands.ShowPoseFeaturesNone);
-					SubMenuBuilder.AddMenuEntry(Commands.ShowPoseFeaturesAll);
-					SubMenuBuilder.AddMenuEntry(Commands.ShowPoseFeaturesDetailed);
-				}
-				SubMenuBuilder.EndSection();
-			})
-			);
-
-			ShowMenuBuilder.AddSubMenu(
-				LOCTEXT("ShowMenu_AnimationsSubMenu", "Animations"),
-				LOCTEXT("ShowMenu_AnimationsSubMenuToolTip", "Animation Preview Options"),
-				FNewMenuDelegate::CreateLambda([](FMenuBuilder& SubMenuBuilder)
-			{
-				const FDatabaseEditorCommands& Commands = FDatabaseEditorCommands::Get();
-
-				SubMenuBuilder.BeginSection("Animations", LOCTEXT("ShowMenu_AnimationsLabel", "Animations"));
-				{
-					SubMenuBuilder.AddMenuEntry(Commands.ShowAnimationOriginalOnly);
-					SubMenuBuilder.AddMenuEntry(Commands.ShowAnimationOriginalAndMirrored);
-				}
-				SubMenuBuilder.EndSection();
-			})
-			);
-
-			
-			ShowMenuBuilder.AddSubMenu(
-				LOCTEXT("ShowMenu_DebugSubMenu", "Debug"),
-				LOCTEXT("ShowMenu_DebugSubMenuToolTip", "Debug Options"),
-				FNewMenuDelegate::CreateLambda([](FMenuBuilder& SubMenuBuilder)
-			{
-				const FDatabaseEditorCommands& Commands = FDatabaseEditorCommands::Get();
-
-				SubMenuBuilder.BeginSection("Debug", LOCTEXT("ShowMenu_DebugLabel", "Debug"));
-				{
-					SubMenuBuilder.AddMenuEntry(Commands.ShowDisplayRootMotionSpeed);
-					SubMenuBuilder.AddMenuEntry(Commands.ShowQuantizeAnimationToPoseData);
-					SubMenuBuilder.AddMenuEntry(Commands.ShowBones);
-				}
-				SubMenuBuilder.EndSection();
-			})
-			);
-			
+			ShowMenuBuilder.BeginSection("Debug", LOCTEXT("ShowMenu_DebugLabel", "Debug"));
+			const FDatabaseEditorCommands& Commands = FDatabaseEditorCommands::Get();
+			ShowMenuBuilder.AddMenuEntry(Commands.ShowDisplayRootMotionSpeed);
+			ShowMenuBuilder.AddMenuEntry(Commands.ShowQuantizeAnimationToPoseData);
+			ShowMenuBuilder.AddMenuEntry(Commands.ShowBones);
+			ShowMenuBuilder.EndSection();
 		}
 
 		return ShowMenuBuilder.MakeWidget();
diff --git a/Engine/Plugins/Experimental/Animation/PoseSearch/Source/Editor/Public/PoseSearchDatabaseEditorCommands.h b/Engine/Plugins/Experimental/Animation/PoseSearch/Source/Editor/Public/PoseSearchDatabaseEditorCommands.h
--- a/Engine/Plugins/Experimental/Animation/PoseSearch/Source/Editor/Public/PoseSearchDatabaseEditorCommands.h
+++ b/Engine/Plugins/Experimental/Animation/PoseSearch/Source/Editor/Public/PoseSearchDatabaseEditorCommands.h
@@ -24,11 +24,6 @@ namespace UE::PoseSearch
 		virtual void RegisterCommands() override;
 
 	public:
-		TSharedPtr<FUICommandInfo> ShowPoseFeaturesNone;
-		TSharedPtr<FUICommandInfo> ShowPoseFeaturesAll;
-		TSharedPtr<FUICommandInfo> ShowPoseFeaturesDetailed;
-		TSharedPtr<FUICommandInfo> ShowAnimationOriginalOnly;
-		TSharedPtr<FUICommandInfo> ShowAnimationOriginalAndMirrored;
 		TSharedPtr<FUICommandInfo> ShowDisplayRootMotionSpeed;
 		TSharedPtr<FUICommandInfo> ShowQuantizeAnimationToPoseData;
 		TSharedPtr<FUICommandInfo> ShowBones;
diff --git a/Engine/Plugins/Experimental/Animation/PoseSearch/Source/Editor/Public/PoseSearchDatabaseViewModel.h b/Engine/Plugins/Experimental/Animation/PoseSearch/Source/Editor/Public/PoseSearchDatabaseViewModel.h
--- a/Engine/Plugins/Experimental/Animation/PoseSearch/Source/Editor/Public/PoseSearchDatabaseViewModel.h
+++ b/Engine/Plugins/Experimental/Animation/PoseSearch/Source/Editor/Public/PoseSearchDatabaseViewModel.h
@@ -25,21 +25,6 @@ namespace UE::PoseSearch
 	struct FSearchIndexAsset;
 	class SDatabaseDataDetails;
 
-	enum class EFeaturesDrawMode : uint8
-	{
-		None,
-		All,
-		Detailed
-	};
-	ENUM_CLASS_FLAGS(EFeaturesDrawMode);
-
-	enum class EAnimationPreviewMode : uint8
-	{
-		OriginalOnly,
-		OriginalAndMirrored
-	};
-	ENUM_CLASS_FLAGS(EAnimationPreviewMode);
-
 	struct FDatabasePreviewActor
 	{
 	public:
@@ -87,12 +72,6 @@ namespace UE::PoseSearch
 		TArray<FDatabasePreviewActor>& GetPreviewActors() { return PreviewActors; }
 		const TArray<FDatabasePreviewActor>& GetPreviewActors() const { return PreviewActors; }
 
-		void OnSetPoseFeaturesDrawMode(EFeaturesDrawMode DrawMode);
-		bool IsPoseFeaturesDrawMode(EFeaturesDrawMode DrawMode) const;
-
-		void OnSetAnimationPreviewMode(EAnimationPreviewMode PreviewMode);
-		bool IsAnimationPreviewMode(EAnimationPreviewMode PreviewMode) const;
-
 		void OnToggleDisplayRootMotionSpeed() { bDisplayRootMotionSpeed = !bDisplayRootMotionSpeed; }
 		bool IsDisplayRootMotionSpeedChecked() const { return bDisplayRootMotionSpeed; };
 
@@ -156,12 +135,6 @@ namespace UE::PoseSearch
 		bool bDrawQueryVector = false;
 		TArray<float> QueryVector;
 
-		/** What features to show in the viewport */
-		EFeaturesDrawMode PoseFeaturesDrawMode = EFeaturesDrawMode::All;
-
-		/** What animations to show in the viewport */
-		EAnimationPreviewMode AnimationPreviewMode = EAnimationPreviewMode::OriginalAndMirrored;
-
 		/** Is animation debug draw enabled */
 		bool bDisplayRootMotionSpeed = false;
 
