motion matching
- UE::PoseSearch::FMirrorDataCache as container of cached CompactPoseMirrorBones and ComponentSpaceRefRotations
- implemented database editor display root motion speed

diff --git a/Engine/Plugins/Experimental/Animation/PoseSearch/Source/Editor/Private/PoseSearchDatabaseDataDetails.cpp b/Engine/Plugins/Experimental/Animation/PoseSearch/Source/Editor/Private/PoseSearchDatabaseDataDetails.cpp
--- a/Engine/Plugins/Experimental/Animation/PoseSearch/Source/Editor/Private/PoseSearchDatabaseDataDetails.cpp
+++ b/Engine/Plugins/Experimental/Animation/PoseSearch/Source/Editor/Private/PoseSearchDatabaseDataDetails.cpp
@@ -224,9 +224,9 @@ public:
 					{
 						if (const FDatabasePreviewActor* FoundPreviewActor = ViewModel->GetPreviewActors().FindByPredicate(
 							[ColumnName](const FDatabasePreviewActor& PreviewActor)
-							{ return *PreviewActor.Actor->GetName() == ColumnName; }))
+							{ return *PreviewActor.GetActor()->GetName() == ColumnName; }))
 						{
-							const int32 PoseIdx = FoundPreviewActor->CurrentPoseIndex;
+							const int32 PoseIdx = FoundPreviewActor->GetCurrentPoseIndex();
 							const TArray<float> PoseValues = PoseSearchDatabase->GetSearchIndex().GetPoseValuesSafe(PoseIdx);
 
 							const int32 DataOffset = ChannelItem->GetDataOffset();
@@ -357,12 +357,8 @@ void SDatabaseDataDetails::Reconstruct(int32 MaxPreviewActors)
 	for (int32 PreviewActorIdx = 0; PreviewActorIdx < PreviewActorNum; ++PreviewActorIdx)
 	{
 		const FDatabasePreviewActor& PreviewActor = PreviewActors[PreviewActorIdx];
-		if (PreviewActor.IsValid() && PreviewActor.Sampler.IsInitialized())
-		{
-			HeaderRow->AddColumn(
-				SHeaderRow::Column(*PreviewActor.Actor->GetName())
-				.DefaultLabel(FText::FromString(PreviewActor.Sampler.GetAsset()->GetName())));
-		}
+		HeaderRow->AddColumn(SHeaderRow::Column(*PreviewActor.GetActor()->GetName())
+			.DefaultLabel(FText::FromString(PreviewActor.GetSampler().GetAsset()->GetName())));
 	}
 
 	ChannelItemsTreeView = SNew(SChannelItemsTreeView)
diff --git a/Engine/Plugins/Experimental/Animation/PoseSearch/Source/Editor/Private/PoseSearchDatabaseEdMode.cpp b/Engine/Plugins/Experimental/Animation/PoseSearch/Source/Editor/Private/PoseSearchDatabaseEdMode.cpp
--- a/Engine/Plugins/Experimental/Animation/PoseSearch/Source/Editor/Private/PoseSearchDatabaseEdMode.cpp
+++ b/Engine/Plugins/Experimental/Animation/PoseSearch/Source/Editor/Private/PoseSearchDatabaseEdMode.cpp
@@ -46,15 +46,15 @@ namespace UE::PoseSearch
 
 	void FDatabaseEdMode::Render(const FSceneView* View, FViewport* Viewport, FPrimitiveDrawInterface* PDI)
 	{
-		if (ViewModel->IsShowBones())
+		if (ViewModel->IsShowBonesChecked())
 		{
 			const UPersonaOptions* PersonaOptions = GetDefault<UPersonaOptions>();
 
 			TArray<FTransform> WorldTransforms;
 			TArray<FLinearColor> BoneColors;
-			for (FDatabasePreviewActor& PreviewActor : ViewModel->GetPreviewActors())
+			for (const FDatabasePreviewActor& PreviewActor : ViewModel->GetPreviewActors())
 			{
-				UDebugSkelMeshComponent* MeshComponent = PreviewActor.GetDebugSkelMeshComponent();
+				const UDebugSkelMeshComponent* MeshComponent = PreviewActor.GetDebugSkelMeshComponent();
 				if (MeshComponent && MeshComponent->GetSkeletalMeshAsset() && MeshComponent->GetNumDrawTransform() > 0 && MeshComponent->SkeletonDrawMode != ESkeletonDrawMode::Hidden)
 				{
 					WorldTransforms.SetNumUninitialized(MeshComponent->GetNumDrawTransform());
diff --git a/Engine/Plugins/Experimental/Animation/PoseSearch/Source/Editor/Private/PoseSearchDatabasePreviewScene.cpp b/Engine/Plugins/Experimental/Animation/PoseSearch/Source/Editor/Private/PoseSearchDatabasePreviewScene.cpp
--- a/Engine/Plugins/Experimental/Animation/PoseSearch/Source/Editor/Private/PoseSearchDatabasePreviewScene.cpp
+++ b/Engine/Plugins/Experimental/Animation/PoseSearch/Source/Editor/Private/PoseSearchDatabasePreviewScene.cpp
@@ -9,13 +9,10 @@
 #include "PoseSearch/PoseSearchContext.h"
 #include "PoseSearch/PoseSearchDatabase.h"
 #include "PoseSearch/PoseSearchDerivedData.h"
+#include "PoseSearch/PoseSearchSchema.h"
 #include "PoseSearchDatabaseEditor.h"
 #include "PoseSearchDatabaseViewModel.h"
 
-#if ENABLE_ANIM_DEBUG
-static TAutoConsoleVariable<float> CVarDatabasePreviewDebugDrawSamplerSize(TEXT("a.DatabasePreview.DebugDrawSamplerSize"), 0.f, TEXT("Debug Draw Sampler Positions Size"));
-#endif
-
 namespace UE::PoseSearch
 {
 	FDatabasePreviewScene::FDatabasePreviewScene(
@@ -65,46 +62,12 @@ namespace UE::PoseSearch
 
 		if (!ViewModel->GetPreviewActors().IsEmpty() && FAsyncPoseSearchDatabasesManagement::RequestAsyncBuildIndex(Database, ERequestAsyncBuildFlag::ContinueRequest))
 		{
+			const bool bDisplayRootMotionSpeed = ViewModel->IsDisplayRootMotionSpeedChecked();
 			bool bDrawQueryVector = ViewModel->ShouldDrawQueryVector();
+
 			for (FDatabasePreviewActor& PreviewActor : ViewModel->GetPreviewActors())
 			{
-				if (Database->GetSearchIndex().IsValidPoseIndex(PreviewActor.CurrentPoseIndex))
-				{
-					if (UDebugSkelMeshComponent* Mesh = PreviewActor.GetDebugSkelMeshComponent())
-					{
-						UE::PoseSearch::FDebugDrawParams DrawParams(GetWorld(), Mesh, PreviewActor.QuantizedTimeRootTransform, Database);
-						DrawParams.DrawFeatureVector(PreviewActor.CurrentPoseIndex);
-
-						if (bDrawQueryVector)
-						{
-							DrawParams.DrawFeatureVector(ViewModel->GetQueryVector());
-							bDrawQueryVector = false;
-						}
-
-#if ENABLE_ANIM_DEBUG
-						const float DebugDrawSamplerSize = CVarDatabasePreviewDebugDrawSamplerSize.GetValueOnAnyThread();
-						if (DebugDrawSamplerSize > UE_KINDA_SMALL_NUMBER)
-						{
-							// drawing the pose extracted from the Sampler to visually compare with the pose features and the mesh drawing
-							FMemMark Mark(FMemStack::Get());
-							FCompactPose Pose;
-							Pose.SetBoneContainer(&PreviewActor.GetAnimPreviewInstance()->GetRequiredBonesOnAnyThread());
-							PreviewActor.Sampler.ExtractPose(PreviewActor.CurrentTime, Pose);
-
-							const FTransform RootTransform = PreviewActor.Sampler.ExtractRootTransform(PreviewActor.CurrentTime);
-
-							FCSPose<FCompactPose> ComponentSpacePose;
-							ComponentSpacePose.InitPose(MoveTemp(Pose));
-
-							for (int32 BoneIndex = 0; BoneIndex < ComponentSpacePose.GetPose().GetNumBones(); ++BoneIndex)
-							{
-								const FTransform BoneWorldTransforms = ComponentSpacePose.GetComponentSpaceTransform(FCompactPoseBoneIndex(BoneIndex)) * RootTransform;
-								DrawParams.DrawPoint(BoneWorldTransforms.GetTranslation(), FColor::Red, DebugDrawSamplerSize);
-							}
-						}
-#endif // ENABLE_ANIM_DEBUG
-					}
-				}
+				bDrawQueryVector &= !PreviewActor.DrawPreviewActor(Database, bDisplayRootMotionSpeed, bDrawQueryVector ? ViewModel->GetQueryVector() : TConstArrayView<float>());
 			}
 		}
 	}
diff --git a/Engine/Plugins/Experimental/Animation/PoseSearch/Source/Editor/Private/PoseSearchDatabaseViewModel.cpp b/Engine/Plugins/Experimental/Animation/PoseSearch/Source/Editor/Private/PoseSearchDatabaseViewModel.cpp
--- a/Engine/Plugins/Experimental/Animation/PoseSearch/Source/Editor/Private/PoseSearchDatabaseViewModel.cpp
+++ b/Engine/Plugins/Experimental/Animation/PoseSearch/Source/Editor/Private/PoseSearchDatabaseViewModel.cpp
@@ -15,6 +15,7 @@
 #include "InstancedStruct.h"
 #include "Modules/ModuleManager.h"
 #include "PoseSearch/PoseSearchAnimNotifies.h"
+#include "PoseSearch/PoseSearchContext.h"
 #include "PoseSearch/PoseSearchDatabase.h"
 #include "PoseSearch/PoseSearchDefines.h"
 #include "PoseSearch/PoseSearchDerivedData.h"
@@ -27,550 +28,659 @@
 
 namespace UE::PoseSearch
 {
-	constexpr float StepDeltaTime = 1.0f / 30.0f;
+#if ENABLE_ANIM_DEBUG
+static TAutoConsoleVariable<float> CVarDatabasePreviewDebugDrawSamplerSize(TEXT("a.DatabasePreview.DebugDrawSamplerSize"), 0.f, TEXT("Debug Draw Sampler Positions Size"));
+#endif
 
-	// FDatabasePreviewActor
-	bool FDatabasePreviewActor::IsValid() const
-	{
-		return Actor != nullptr;
-	}
+constexpr float StepDeltaTime = 1.0f / 30.0f;
 
-	void FDatabasePreviewActor::Process(const FBoneContainer& BoneContainer)
-	{
-		Sampler.Process(BoneContainer);
-	}
+// FDatabasePreviewActor
+bool FDatabasePreviewActor::SpawnPreviewActor(UWorld* World, const UPoseSearchDatabase* PoseSearchDatabase, int32 IndexAssetIdx, int32 PoseIdxForTimeOffset)
+{
+	check(PoseSearchDatabase);
+	const FSearchIndex& SearchIndex = PoseSearchDatabase->GetSearchIndex();
+	const FSearchIndexAsset& IndexAsset = SearchIndex.Assets[IndexAssetIdx];
 
-	UDebugSkelMeshComponent* FDatabasePreviewActor::GetDebugSkelMeshComponent()
+	const FPoseSearchDatabaseAnimationAssetBase* DatabaseAnimationAsset = PoseSearchDatabase->GetAnimationAssetBase(IndexAsset.GetSourceAssetIdx());
+	UAnimationAsset* PreviewAsset = DatabaseAnimationAsset->GetAnimationAsset();
+	if (!PreviewAsset)
 	{
-		if (IsValid())
-		{
-			return Cast<UDebugSkelMeshComponent>(Actor->GetRootComponent());
-		}
-		return nullptr;
+		return false;
 	}
 
-	UAnimPreviewInstance* FDatabasePreviewActor::GetAnimPreviewInstance()
+	Sampler.Init(PreviewAsset, IndexAsset.GetBlendParameters());
+
+	FBoneContainer BoneContainer;
+	BoneContainer.InitializeTo(PoseSearchDatabase->Schema->BoneIndicesWithParents, UE::Anim::FCurveFilterSettings(UE::Anim::ECurveFilterMode::DisallowAll), *PoseSearchDatabase->Schema->Skeleton);
+	Sampler.Process(BoneContainer);
+
+	IndexAssetIndex = IndexAssetIdx;
+	CurrentPoseIndex = INDEX_NONE;
+	PlayTimeOffset = PoseIdxForTimeOffset < 0 ? 0.f : PoseSearchDatabase->GetRealAssetTime(PoseIdxForTimeOffset);
+
+	// @todo: should we always use the PlayTimeOffset to extract the root transform?
+	if (PlayTimeOffset != 0.f)
 	{
-		if (UDebugSkelMeshComponent* Mesh = GetDebugSkelMeshComponent())
-		{
-			return Mesh->PreviewInstance.Get();
-		}
-		return nullptr;
+		RootTransformOrigin = ExtractRootTransform(PlayTimeOffset);
 	}
 
-	// FDatabaseViewModel
-	FDatabaseViewModel::FDatabaseViewModel()
-		: PoseSearchDatabase(nullptr)
+	FActorSpawnParameters Params;
+	Params.SpawnCollisionHandlingOverride = ESpawnActorCollisionHandlingMethod::AdjustIfPossibleButAlwaysSpawn;
+	Actor = World->SpawnActor<AActor>(AActor::StaticClass(), FTransform::Identity, Params);
+	Actor->SetFlags(RF_Transient);
+
+	UDebugSkelMeshComponent* Mesh = NewObject<UDebugSkelMeshComponent>(Actor.Get());
+	Mesh->RegisterComponentWithWorld(World);
+
+	UAnimPreviewInstance* AnimInstance = NewObject<UAnimPreviewInstance>(Mesh);
+
+	Mesh->PreviewInstance = AnimInstance;
+	AnimInstance->InitializeAnimation();
+
+	USkeletalMesh* DatabasePreviewMesh = PoseSearchDatabase->PreviewMesh;
+	Mesh->SetSkeletalMesh(DatabasePreviewMesh ? DatabasePreviewMesh : PoseSearchDatabase->Schema->Skeleton->GetPreviewMesh(true));
+	Mesh->EnablePreview(true, PreviewAsset);
+		
+	AnimInstance->SetAnimationAsset(PreviewAsset, false, 0.0f);
+	AnimInstance->SetBlendSpacePosition(IndexAsset.GetBlendParameters());
+		
+	if (IndexAsset.IsMirrored() && PoseSearchDatabase->Schema)
 	{
+		AnimInstance->SetMirrorDataTable(PoseSearchDatabase->Schema->MirrorDataTable);
+		MirrorDataCache.Init(PoseSearchDatabase->Schema->MirrorDataTable, AnimInstance->GetRequiredBonesOnAnyThread());
 	}
 
-	void FDatabaseViewModel::AddReferencedObjects(FReferenceCollector& Collector)
+	AnimInstance->PlayAnim(false, 0.0f);
+
+	if (!Actor->GetRootComponent())
 	{
-		Collector.AddReferencedObject(PoseSearchDatabase);
+		Actor->SetRootComponent(Mesh);
 	}
 
-	void FDatabaseViewModel::Initialize(UPoseSearchDatabase* InPoseSearchDatabase, const TSharedRef<FDatabasePreviewScene>& InPreviewScene, const TSharedRef<SDatabaseDataDetails>& InDatabaseDataDetails)
-	{
-		PoseSearchDatabase = InPoseSearchDatabase;
-		PreviewScenePtr = InPreviewScene;
-		DatabaseDataDetails = InDatabaseDataDetails;
+	AnimInstance->SetPlayRate(0.f);
 
-		RemovePreviewActors();
-	}
+	UE_LOG(LogPoseSearchEditor, Log, TEXT("Spawned preview Actor: %s"), *GetNameSafe(Actor.Get()));
+	return true;
+}
+
+void FDatabasePreviewActor::UpdatePreviewActor(const UPoseSearchDatabase* PoseSearchDatabase, float PlayTime, bool bQuantizeAnimationToPoseData)
+{
+	check(PoseSearchDatabase);
 
-	void FDatabaseViewModel::BuildSearchIndex()
+	const FSearchIndex& SearchIndex = PoseSearchDatabase->GetSearchIndex();
+
+	UAnimPreviewInstance* AnimInstance = GetAnimPreviewInstanceInternal();
+	if (!AnimInstance || IndexAssetIndex >= SearchIndex.Assets.Num())
 	{
-		using namespace UE::PoseSearch;
-		FAsyncPoseSearchDatabasesManagement::RequestAsyncBuildIndex(PoseSearchDatabase, ERequestAsyncBuildFlag::NewRequest);
+		return;
 	}
 
-	void FDatabaseViewModel::PreviewBackwardEnd()
+	const UAnimationAsset* PreviewAsset = AnimInstance->GetAnimationAsset();
+	if (!PreviewAsset)
 	{
-		PlayTime = MinPreviewPlayLength;
+		return;
 	}
 
-	void FDatabaseViewModel::PreviewBackwardStep()
+	CurrentTime = 0.f;
+	const FSearchIndexAsset& IndexAsset = SearchIndex.Assets[IndexAssetIndex];
+	float CurrentPlayTime = PlayTime + IndexAsset.GetFirstSampleTime(PoseSearchDatabase->Schema->SampleRate) + PlayTimeOffset;
+	FAnimationRuntime::AdvanceTime(false, CurrentPlayTime, CurrentTime, IndexAsset.GetLastSampleTime(PoseSearchDatabase->Schema->SampleRate));
+			 
+	// time to pose index
+	CurrentPoseIndex = PoseSearchDatabase->GetPoseIndexFromTime(CurrentTime, IndexAsset);
+
+	const float QuantizedTime = CurrentPoseIndex >= 0 ? PoseSearchDatabase->GetRealAssetTime(CurrentPoseIndex) : CurrentTime;
+	if (bQuantizeAnimationToPoseData)
 	{
-		PlayTime = FMath::Clamp(PlayTime - StepDeltaTime, MinPreviewPlayLength, MaxPreviewPlayLength);
-		DeltaTimeMultiplier = 0.0f;
+		CurrentTime = QuantizedTime;
 	}
 
-	void FDatabaseViewModel::PreviewBackward()
+	// SetPosition is in [0..1] range for blendspaces
+	AnimInstance->SetPosition(Sampler.ToNormalizedTime(CurrentTime));
+	AnimInstance->SetPlayRate(0.f);
+	AnimInstance->SetBlendSpacePosition(IndexAsset.GetBlendParameters());
+
+	// updating root transforms
+	RootTransformCurrentQuantizedTime = ExtractRootTransform(QuantizedTime);
+	RootTransformCurrentQuantizedTime.SetToRelativeTransform(RootTransformOrigin);
+
+	if (CurrentTime == QuantizedTime)
 	{
-		DeltaTimeMultiplier = -1.0f;
+		RootTransformCurrent = RootTransformCurrentQuantizedTime;
 	}
-
-	void FDatabaseViewModel::PreviewPause()
+	else
 	{
-		DeltaTimeMultiplier = 0.0f;
+		RootTransformCurrent = ExtractRootTransform(CurrentTime);
+		RootTransformCurrent.SetToRelativeTransform(RootTransformOrigin);
 	}
 
-	void FDatabaseViewModel::PreviewForward()
+	Actor->SetActorTransform(RootTransformCurrent);
+}
+
+void FDatabasePreviewActor::Destroy()
+{
+	check(Actor);
+	Actor->Destroy();
+}
+
+bool FDatabasePreviewActor::DrawPreviewActor(const UPoseSearchDatabase* PoseSearchDatabase, bool bDisplayRootMotionSpeed, TConstArrayView<float> QueryVector)
+{
+	if (!PoseSearchDatabase->GetSearchIndex().IsValidPoseIndex(GetCurrentPoseIndex()))
 	{
-		DeltaTimeMultiplier = 1.0f;
+		return false;
 	}
 
-	void FDatabaseViewModel::PreviewForwardStep()
+	const UDebugSkelMeshComponent* Mesh = GetDebugSkelMeshComponent();
+	if (!Mesh)
 	{
-		PlayTime = FMath::Clamp(PlayTime + StepDeltaTime, MinPreviewPlayLength, MaxPreviewPlayLength);
-		DeltaTimeMultiplier = 0.0f;
+		return false;
 	}
 
-	void FDatabaseViewModel::PreviewForwardEnd()
+	UE::PoseSearch::FDebugDrawParams DrawParams(Mesh->GetWorld(), Mesh, RootTransformCurrentQuantizedTime, PoseSearchDatabase);
+	DrawParams.DrawFeatureVector(GetCurrentPoseIndex());
+
+	if (!QueryVector.IsEmpty())
 	{
-		PlayTime = MaxPreviewPlayLength;
+		DrawParams.DrawFeatureVector(QueryVector);
 	}
 
-	FDatabasePreviewActor FDatabaseViewModel::SpawnPreviewActor(int32 IndexAssetIndex, int32 PoseIdxForTimeOffset)
+	if (bDisplayRootMotionSpeed)
 	{
-		FDatabasePreviewActor PreviewActor;
-
-		check(PoseSearchDatabase);
-		const FSearchIndex& SearchIndex = PoseSearchDatabase->GetSearchIndex();
-		const FSearchIndexAsset& IndexAsset = SearchIndex.Assets[IndexAssetIndex];
-
-		const FPoseSearchDatabaseAnimationAssetBase* DatabaseAnimationAsset = PoseSearchDatabase->GetAnimationAssetBase(IndexAsset.GetSourceAssetIdx());
-		UAnimationAsset* PreviewAsset = DatabaseAnimationAsset->GetAnimationAsset();
-		if (!PreviewAsset)
+		// initializing SampledRootMotion if required
+		if (SampledRootMotion.IsEmpty())
 		{
-			return PreviewActor;
-		}
-
-		const FInstancedStruct& DatabaseAsset = PoseSearchDatabase->GetAnimationAssetStruct(IndexAsset);
-		const FPoseSearchDatabaseAnimationAssetBase* DatabaseAnimationAssetBase = DatabaseAsset.GetPtr<FPoseSearchDatabaseAnimationAssetBase>();
-		check(DatabaseAnimationAssetBase);
-		PreviewActor.Sampler.Init(DatabaseAnimationAssetBase->GetAnimationAsset(), IndexAsset.GetBlendParameters());
-		PreviewActor.IndexAssetIndex = IndexAssetIndex;
-		PreviewActor.CurrentPoseIndex = INDEX_NONE;
-		PreviewActor.PlayTimeOffset = PoseIdxForTimeOffset < 0 ? 0.f : PoseSearchDatabase->GetRealAssetTime(PoseIdxForTimeOffset);
+			const FSearchIndex& SearchIndex = PoseSearchDatabase->GetSearchIndex();
+			const FSearchIndexAsset& IndexAsset = SearchIndex.Assets[IndexAssetIndex];
 
-		FActorSpawnParameters Params;
-		Params.SpawnCollisionHandlingOverride = ESpawnActorCollisionHandlingMethod::AdjustIfPossibleButAlwaysSpawn;
-		PreviewActor.Actor = GetWorld()->SpawnActor<AActor>(AActor::StaticClass(), FTransform::Identity, Params);
-		PreviewActor.Actor->SetFlags(RF_Transient);
+			const int NumPoses = IndexAsset.GetNumPoses();
+			if (NumPoses > 1)
+			{
+				SampledRootMotion.SetNumUninitialized(NumPoses);
+				SampledRootMotionSpeed.SetNumUninitialized(NumPoses);
 
-		UDebugSkelMeshComponent* Mesh = NewObject<UDebugSkelMeshComponent>(PreviewActor.Actor.Get());
-		Mesh->RegisterComponentWithWorld(GetWorld());
+				for (int32 Index = 0; Index < NumPoses; ++Index)
+				{
+					const int32 IndexAssetPoseIdx = Index + IndexAsset.GetFirstPoseIdx();
+					const float IndexAssetPoseTime = IndexAsset.GetTimeFromPoseIndex(IndexAssetPoseIdx, PoseSearchDatabase->Schema->SampleRate);
+					FTransform IndexAssetPoseTransform = ExtractRootTransform(IndexAssetPoseTime);
+					IndexAssetPoseTransform.SetToRelativeTransform(RootTransformOrigin);
 
-		UAnimPreviewInstance* AnimInstance = NewObject<UAnimPreviewInstance>(Mesh);
+					SampledRootMotion[Index] = IndexAssetPoseTransform.GetTranslation();
+				}
 
-		Mesh->PreviewInstance = AnimInstance;
-		AnimInstance->InitializeAnimation();
+				for (int32 Index = 1; Index < NumPoses; ++Index)
+				{
+					const FVector& Start = SampledRootMotion[Index - 1];
+					const FVector& End = SampledRootMotion[Index];
+					SampledRootMotionSpeed[Index] = (Start - End).Length() * PoseSearchDatabase->Schema->PermutationsSampleRate;
+				}
+				SampledRootMotionSpeed[0] = SampledRootMotionSpeed[1];
+			}
+		}
 
-		USkeletalMesh* DatabasePreviewMesh = PoseSearchDatabase->PreviewMesh;
-		Mesh->SetSkeletalMesh(DatabasePreviewMesh ? DatabasePreviewMesh : PoseSearchDatabase->Schema->Skeleton->GetPreviewMesh(true));
-		Mesh->EnablePreview(true, PreviewAsset);
-		
-		AnimInstance->SetAnimationAsset(PreviewAsset, false, 0.0f);
-		AnimInstance->SetBlendSpacePosition(IndexAsset.GetBlendParameters());
-		
-		if (IndexAsset.IsMirrored() && PoseSearchDatabase->Schema)
+		// drawing PreviewActor.SampledRootMotion
+		const int32 SampledRootMotionNum = SampledRootMotion.Num();
+		if (SampledRootMotionNum > 1)
 		{
-			AnimInstance->SetMirrorDataTable(PoseSearchDatabase->Schema->MirrorDataTable);
+			for (int32 Index = 0; Index < SampledRootMotion.Num(); ++Index)
+			{
+				const FVector& EndDown = SampledRootMotion[Index];
+				const FVector EndUp = EndDown + (SampledRootMotionSpeed[Index] * FVector::UpVector);
+
+				DrawParams.DrawLine(EndDown, EndUp, FColor::Black);
+				if (Index > 0)
+				{
+					const FColor RootMotionColor = Index % 2 == 0 ? FColor::Purple : FColor::Orange;
+					const FVector& StartDown = SampledRootMotion[Index - 1];
+					const FVector StartUp = StartDown + (SampledRootMotionSpeed[Index - 1] * FVector::UpVector);
+					DrawParams.DrawLine(StartDown, EndDown, RootMotionColor);
+					DrawParams.DrawLine(StartUp, EndUp, RootMotionColor);
+				}
+			}
 		}
+	}
+
+#if ENABLE_ANIM_DEBUG
+	const float DebugDrawSamplerSize = CVarDatabasePreviewDebugDrawSamplerSize.GetValueOnAnyThread();
+	if (DebugDrawSamplerSize > UE_KINDA_SMALL_NUMBER)
+	{
+		// drawing the pose extracted from the Sampler to visually compare with the pose features and the mesh drawing
+		FMemMark Mark(FMemStack::Get());
+		FCompactPose Pose;
+		Pose.SetBoneContainer(&GetAnimPreviewInstance()->GetRequiredBonesOnAnyThread());
+		ExtractPose(CurrentTime, Pose);
 
-		AnimInstance->PlayAnim(false, 0.0f);
+		const FTransform RootTransform = ExtractRootTransform(CurrentTime);
 
-		if (!PreviewActor.Actor->GetRootComponent())
+		FCSPose<FCompactPose> ComponentSpacePose;
+		ComponentSpacePose.InitPose(MoveTemp(Pose));
+
+		for (int32 BoneIndex = 0; BoneIndex < ComponentSpacePose.GetPose().GetNumBones(); ++BoneIndex)
 		{
-			PreviewActor.Actor->SetRootComponent(Mesh);
+			const FTransform BoneWorldTransforms = ComponentSpacePose.GetComponentSpaceTransform(FCompactPoseBoneIndex(BoneIndex)) * RootTransform;
+			DrawParams.DrawPoint(BoneWorldTransforms.GetTranslation(), FColor::Red, DebugDrawSamplerSize);
 		}
+	}
+#endif // ENABLE_ANIM_DEBUG
 
-		AnimInstance->SetPlayRate(0.f);
+	return true;
+}
 
-		UE_LOG(LogPoseSearchEditor, Log, TEXT("Spawned preview Actor: %s"), *GetNameSafe(PreviewActor.Actor.Get()));
-		return PreviewActor;
-	}
+FTransform FDatabasePreviewActor::ExtractRootTransform(float Time) const
+{
+	return MirrorDataCache.MirrorTransform(Sampler.ExtractRootTransform(Time));
+}
+
+void FDatabasePreviewActor::ExtractPose(float Time, FCompactPose& OutPose) const
+{
+	Sampler.ExtractPose(Time, OutPose);
+	MirrorDataCache.MirrorPose(OutPose);
+}
 
-	UWorld* FDatabaseViewModel::GetWorld()
+const UDebugSkelMeshComponent* FDatabasePreviewActor::GetDebugSkelMeshComponent() const
+{
+	if (Actor)
 	{
-		check(PreviewScenePtr.IsValid());
-		return PreviewScenePtr.Pin()->GetWorld();
+		return Cast<UDebugSkelMeshComponent>(Actor->GetRootComponent());
 	}
+	return nullptr;
+}
 
-	void FDatabaseViewModel::OnPreviewActorClassChanged()
+const UAnimPreviewInstance* FDatabasePreviewActor::GetAnimPreviewInstance() const
+{
+	if (const UDebugSkelMeshComponent* Mesh = GetDebugSkelMeshComponent())
 	{
-		// todo: implement
+		return Mesh->PreviewInstance.Get();
 	}
+	return nullptr;
+}
 
-	void FDatabaseViewModel::Tick(float DeltaSeconds)
+UAnimPreviewInstance* FDatabasePreviewActor::GetAnimPreviewInstanceInternal()
+{
+	if (Actor)
 	{
-		const float DeltaPlayTime = DeltaSeconds * DeltaTimeMultiplier;
-		if (!FMath::IsNearlyZero(DeltaPlayTime))
+		if (UDebugSkelMeshComponent* Mesh = Cast<UDebugSkelMeshComponent>(Actor->GetRootComponent()))
 		{
-			PlayTime += DeltaPlayTime;
-			PlayTime = FMath::Clamp(PlayTime, MinPreviewPlayLength, MaxPreviewPlayLength);
-
-			if (FAsyncPoseSearchDatabasesManagement::RequestAsyncBuildIndex(PoseSearchDatabase, ERequestAsyncBuildFlag::ContinueRequest))
-			{
-				UpdatePreviewActors(true);
-			}
+			return Mesh->PreviewInstance.Get();
 		}
 	}
+	return nullptr;
+}
 
-	void FDatabaseViewModel::UpdatePreviewActors(bool bInTickPlayTime)
-	{
-		const FSearchIndex& SearchIndex = PoseSearchDatabase->GetSearchIndex();
-		TSharedPtr<FDatabasePreviewScene> PreviewScene = PreviewScenePtr.Pin();
-		for (FDatabasePreviewActor& PreviewActor : GetPreviewActors())
-		{
-			UAnimPreviewInstance* AnimInstance = PreviewActor.GetAnimPreviewInstance();
-			if (!AnimInstance || PreviewActor.IndexAssetIndex >= SearchIndex.Assets.Num())
-			{
-				continue;
-			}
+// FDatabaseViewModel
+FDatabaseViewModel::FDatabaseViewModel()
+	: PoseSearchDatabase(nullptr)
+{
+}
 
-			const UAnimationAsset* PreviewAsset = AnimInstance->GetAnimationAsset();
-			if (!PreviewAsset || !PreviewActor.Sampler.IsInitialized())
-			{
-				continue;
-			}
+void FDatabaseViewModel::AddReferencedObjects(FReferenceCollector& Collector)
+{
+	Collector.AddReferencedObject(PoseSearchDatabase);
+}
 
-			PreviewActor.CurrentTime = 0.f;
-			const FSearchIndexAsset& IndexAsset = SearchIndex.Assets[PreviewActor.IndexAssetIndex];
-			float CurrentPlayTime = PlayTime + IndexAsset.GetFirstSampleTime(PoseSearchDatabase->Schema->SampleRate) + PreviewActor.PlayTimeOffset;
-			FAnimationRuntime::AdvanceTime(false, CurrentPlayTime, PreviewActor.CurrentTime, IndexAsset.GetLastSampleTime(PoseSearchDatabase->Schema->SampleRate));
-			 
-			// time to pose index
-			PreviewActor.CurrentPoseIndex = PoseSearchDatabase->GetPoseIndexFromTime(PreviewActor.CurrentTime, IndexAsset);
+void FDatabaseViewModel::Initialize(UPoseSearchDatabase* InPoseSearchDatabase, const TSharedRef<FDatabasePreviewScene>& InPreviewScene, const TSharedRef<SDatabaseDataDetails>& InDatabaseDataDetails)
+{
+	PoseSearchDatabase = InPoseSearchDatabase;
+	PreviewScenePtr = InPreviewScene;
+	DatabaseDataDetails = InDatabaseDataDetails;
 
-			const float QuantizedTime = PreviewActor.CurrentPoseIndex >= 0 ? PoseSearchDatabase->GetRealAssetTime(PreviewActor.CurrentPoseIndex) : PreviewActor.CurrentTime;
-			if (bQuantizeAnimationToPoseData)
-			{
-				PreviewActor.CurrentTime = QuantizedTime;
-			}
+	RemovePreviewActors();
+}
 
-			// SetPosition is in [0..1] range for blendspaces
-			AnimInstance->SetPosition(PreviewActor.Sampler.ToNormalizedTime(PreviewActor.CurrentTime));
-			AnimInstance->SetPlayRate(0.f);
-			AnimInstance->SetBlendSpacePosition(IndexAsset.GetBlendParameters());
+void FDatabaseViewModel::BuildSearchIndex()
+{
+	using namespace UE::PoseSearch;
+	FAsyncPoseSearchDatabasesManagement::RequestAsyncBuildIndex(PoseSearchDatabase, ERequestAsyncBuildFlag::NewRequest);
+}
 
-			PreviewActor.QuantizedTimeRootTransform = PreviewActor.Sampler.ExtractRootTransform(QuantizedTime);
-			FTransform RootMotion = PreviewActor.Sampler.ExtractRootTransform(PreviewActor.CurrentTime);
-			if (AnimInstance->GetMirrorDataTable())
-			{
-				RootMotion = MirrorRootTransform(RootMotion);
-				PreviewActor.QuantizedTimeRootTransform = MirrorRootTransform(PreviewActor.QuantizedTimeRootTransform);
-			}
+void FDatabaseViewModel::PreviewBackwardEnd()
+{
+	PlayTime = MinPreviewPlayLength;
+}
 
-			if (PreviewActor.PlayTimeOffset != 0.f)
-			{
-				FTransform OriginRootMotion = PreviewActor.Sampler.ExtractRootTransform(PreviewActor.PlayTimeOffset);
-				if (AnimInstance->GetMirrorDataTable())
-				{
-					OriginRootMotion = MirrorRootTransform(OriginRootMotion);
-				}
-				RootMotion.SetToRelativeTransform(OriginRootMotion);
-				PreviewActor.QuantizedTimeRootTransform.SetToRelativeTransform(OriginRootMotion);
-			}
+void FDatabaseViewModel::PreviewBackwardStep()
+{
+	PlayTime = FMath::Clamp(PlayTime - StepDeltaTime, MinPreviewPlayLength, MaxPreviewPlayLength);
+	DeltaTimeMultiplier = 0.0f;
+}
 
-			PreviewActor.Actor->SetActorTransform(RootMotion);
-		}
-	}
+void FDatabaseViewModel::PreviewBackward()
+{
+	DeltaTimeMultiplier = -1.0f;
+}
 
-	void FDatabaseViewModel::RemovePreviewActors()
-	{
-		PlayTime = 0.f;
-		DeltaTimeMultiplier = 1.f;
-		MaxPreviewPlayLength = 0.f;
-		MinPreviewPlayLength = 0.f;
-		bIsEditorSelection = true;
-		bDrawQueryVector = false;
+void FDatabaseViewModel::PreviewPause()
+{
+	DeltaTimeMultiplier = 0.0f;
+}
 
-		for (FDatabasePreviewActor& PreviewActor : PreviewActors)
-		{
-			PreviewActor.Actor->Destroy();
-		}
-		PreviewActors.Reset();
-	}
+void FDatabaseViewModel::PreviewForward()
+{
+	DeltaTimeMultiplier = 1.0f;
+}
 
-	FTransform FDatabaseViewModel::MirrorRootTransform(const FTransform& RootTransform)
+void FDatabaseViewModel::PreviewForwardStep()
+{
+	PlayTime = FMath::Clamp(PlayTime + StepDeltaTime, MinPreviewPlayLength, MaxPreviewPlayLength);
+	DeltaTimeMultiplier = 0.0f;
+}
+
+void FDatabaseViewModel::PreviewForwardEnd()
+{
+	PlayTime = MaxPreviewPlayLength;
+}
+
+UWorld* FDatabaseViewModel::GetWorld()
+{
+	check(PreviewScenePtr.IsValid());
+	return PreviewScenePtr.Pin()->GetWorld();
+}
+
+void FDatabaseViewModel::OnPreviewActorClassChanged()
+{
+	// todo: implement
+}
+
+void FDatabaseViewModel::Tick(float DeltaSeconds)
+{
+	const float DeltaPlayTime = DeltaSeconds * DeltaTimeMultiplier;
+	if (!FMath::IsNearlyZero(DeltaPlayTime))
 	{
-		check(PoseSearchDatabase && PoseSearchDatabase->Schema);
-		if (PoseSearchDatabase->Schema->MirrorDataTable && PoseSearchDatabase->Schema->Skeleton)
+		PlayTime += DeltaPlayTime;
+		PlayTime = FMath::Clamp(PlayTime, MinPreviewPlayLength, MaxPreviewPlayLength);
+
+		if (FAsyncPoseSearchDatabasesManagement::RequestAsyncBuildIndex(PoseSearchDatabase, ERequestAsyncBuildFlag::ContinueRequest))
 		{
-			const FTransform RootReferenceTransform = PoseSearchDatabase->Schema->Skeleton->GetReferenceSkeleton().GetRefBonePose()[0];
-			const FQuat RootReferenceRotation = RootReferenceTransform.GetRotation();
-
-			const EAxis::Type MirrorAxis = PoseSearchDatabase->Schema->MirrorDataTable->MirrorAxis;
-			FVector T = RootTransform.GetTranslation();
-			T = FAnimationRuntime::MirrorVector(T, MirrorAxis);
-			FQuat Q = RootTransform.GetRotation();
-			Q = FAnimationRuntime::MirrorQuat(Q, MirrorAxis);
-			Q *= FAnimationRuntime::MirrorQuat(RootReferenceRotation, MirrorAxis).Inverse() * RootReferenceRotation;
-
-			const FTransform MirroredRootTransform = FTransform(Q, T, RootTransform.GetScale3D());
-			return MirroredRootTransform;
+			for (FDatabasePreviewActor& PreviewActor : GetPreviewActors())
+			{
+				PreviewActor.UpdatePreviewActor(PoseSearchDatabase, PlayTime, bQuantizeAnimationToPoseData);
+			}
 		}
-		return RootTransform;
 	}
+}
 
-	void FDatabaseViewModel::AddSequenceToDatabase(UAnimSequence* AnimSequence)
-	{
-		FPoseSearchDatabaseSequence NewAsset;
-		NewAsset.Sequence = AnimSequence;
-		PoseSearchDatabase->AnimationAssets.Add(FInstancedStruct::Make(NewAsset));
-	}
+void FDatabaseViewModel::RemovePreviewActors()
+{
+	PlayTime = 0.f;
+	DeltaTimeMultiplier = 1.f;
+	MaxPreviewPlayLength = 0.f;
+	MinPreviewPlayLength = 0.f;
+	bIsEditorSelection = true;
+	bDrawQueryVector = false;
 
-	void FDatabaseViewModel::AddBlendSpaceToDatabase(UBlendSpace* BlendSpace)
+	for (FDatabasePreviewActor& PreviewActor : PreviewActors)
 	{
-		FPoseSearchDatabaseBlendSpace NewAsset;
-		NewAsset.BlendSpace = BlendSpace;
-		PoseSearchDatabase->AnimationAssets.Add(FInstancedStruct::Make(NewAsset));
+		PreviewActor.Destroy();
 	}
 
-	void FDatabaseViewModel::AddAnimCompositeToDatabase(UAnimComposite* AnimComposite)
-	{
-		FPoseSearchDatabaseAnimComposite NewAsset;
-		NewAsset.AnimComposite = AnimComposite;
-		PoseSearchDatabase->AnimationAssets.Add(FInstancedStruct::Make(NewAsset));
-	}
+	PreviewActors.Reset();
+}
 
-	void FDatabaseViewModel::AddAnimMontageToDatabase(UAnimMontage* AnimMontage)
-	{
-		FPoseSearchDatabaseAnimMontage NewAsset;
-		NewAsset.AnimMontage = AnimMontage;
-		PoseSearchDatabase->AnimationAssets.Add(FInstancedStruct::Make(NewAsset));
-	}
+void FDatabaseViewModel::AddSequenceToDatabase(UAnimSequence* AnimSequence)
+{
+	FPoseSearchDatabaseSequence NewAsset;
+	NewAsset.Sequence = AnimSequence;
+	PoseSearchDatabase->AnimationAssets.Add(FInstancedStruct::Make(NewAsset));
+}
+
+void FDatabaseViewModel::AddBlendSpaceToDatabase(UBlendSpace* BlendSpace)
+{
+	FPoseSearchDatabaseBlendSpace NewAsset;
+	NewAsset.BlendSpace = BlendSpace;
+	PoseSearchDatabase->AnimationAssets.Add(FInstancedStruct::Make(NewAsset));
+}
 
-	bool FDatabaseViewModel::DeleteFromDatabase(int32 AnimationAssetIndex)
+void FDatabaseViewModel::AddAnimCompositeToDatabase(UAnimComposite* AnimComposite)
+{
+	FPoseSearchDatabaseAnimComposite NewAsset;
+	NewAsset.AnimComposite = AnimComposite;
+	PoseSearchDatabase->AnimationAssets.Add(FInstancedStruct::Make(NewAsset));
+}
+
+void FDatabaseViewModel::AddAnimMontageToDatabase(UAnimMontage* AnimMontage)
+{
+	FPoseSearchDatabaseAnimMontage NewAsset;
+	NewAsset.AnimMontage = AnimMontage;
+	PoseSearchDatabase->AnimationAssets.Add(FInstancedStruct::Make(NewAsset));
+}
+
+bool FDatabaseViewModel::DeleteFromDatabase(int32 AnimationAssetIndex)
+{
+	if (const FPoseSearchDatabaseAnimationAssetBase* DatabaseAnimationAssetBase = PoseSearchDatabase->GetAnimationAssetBase(AnimationAssetIndex))
 	{
-		if (const FPoseSearchDatabaseAnimationAssetBase* DatabaseAnimationAssetBase = PoseSearchDatabase->GetAnimationAssetBase(AnimationAssetIndex))
+		if (!DatabaseAnimationAssetBase->bSynchronizeWithExternalDependency)
 		{
-			if (!DatabaseAnimationAssetBase->bSynchronizeWithExternalDependency)
-			{
-				PoseSearchDatabase->AnimationAssets.RemoveAt(AnimationAssetIndex);
-				PoseSearchDatabase->Modify();
+			PoseSearchDatabase->AnimationAssets.RemoveAt(AnimationAssetIndex);
+			PoseSearchDatabase->Modify();
 		
-				return true;
-			}
+			return true;
 		}
-
-		return false;
 	}
 
-	void FDatabaseViewModel::SetDisableReselection(int32 AnimationAssetIndex, bool bEnabled)
+	return false;
+}
+
+void FDatabaseViewModel::SetDisableReselection(int32 AnimationAssetIndex, bool bEnabled)
+{
+	if (FPoseSearchDatabaseAnimationAssetBase* DatabaseAnimationAsset = PoseSearchDatabase->GetMutableAnimationAssetBase(AnimationAssetIndex))
 	{
-		if (FPoseSearchDatabaseAnimationAssetBase* DatabaseAnimationAsset = PoseSearchDatabase->GetMutableAnimationAssetBase(AnimationAssetIndex))
-		{
-			DatabaseAnimationAsset->SetDisableReselection(bEnabled);
-		}
+		DatabaseAnimationAsset->SetDisableReselection(bEnabled);
 	}
+}
 
-	bool FDatabaseViewModel::IsDisableReselection(int32 AnimationAssetIndex) const
+bool FDatabaseViewModel::IsDisableReselection(int32 AnimationAssetIndex) const
+{
+	if (const FPoseSearchDatabaseAnimationAssetBase* DatabaseAnimationAsset = PoseSearchDatabase->GetAnimationAssetBase(AnimationAssetIndex))
 	{
-		if (const FPoseSearchDatabaseAnimationAssetBase* DatabaseAnimationAsset = PoseSearchDatabase->GetAnimationAssetBase(AnimationAssetIndex))
-		{
-			return DatabaseAnimationAsset->IsDisableReselection();
-		}
-
-		return false;
+		return DatabaseAnimationAsset->IsDisableReselection();
 	}
 
-	void FDatabaseViewModel::SetIsEnabled(int32 AnimationAssetIndex, bool bEnabled)
+	return false;
+}
+
+void FDatabaseViewModel::SetIsEnabled(int32 AnimationAssetIndex, bool bEnabled)
+{
+	if (FPoseSearchDatabaseAnimationAssetBase* DatabaseAnimationAsset = PoseSearchDatabase->GetMutableAnimationAssetBase(AnimationAssetIndex))
 	{
-		if (FPoseSearchDatabaseAnimationAssetBase* DatabaseAnimationAsset = PoseSearchDatabase->GetMutableAnimationAssetBase(AnimationAssetIndex))
-		{
-			GetPoseSearchDatabase()->Modify();
+		GetPoseSearchDatabase()->Modify();
 
-			DatabaseAnimationAsset->SetIsEnabled(bEnabled);
-		}
+		DatabaseAnimationAsset->SetIsEnabled(bEnabled);
 	}
+}
 
-	bool FDatabaseViewModel::IsEnabled(int32 AnimationAssetIndex) const
+bool FDatabaseViewModel::IsEnabled(int32 AnimationAssetIndex) const
+{
+	if (const FPoseSearchDatabaseAnimationAssetBase* DatabaseAnimationAsset = PoseSearchDatabase->GetAnimationAssetBase(AnimationAssetIndex))
 	{
-		if (const FPoseSearchDatabaseAnimationAssetBase* DatabaseAnimationAsset = PoseSearchDatabase->GetAnimationAssetBase(AnimationAssetIndex))
-		{
-			return DatabaseAnimationAsset->IsEnabled();
-		}
-
-		return false;
+		return DatabaseAnimationAsset->IsEnabled();
 	}
 
-	int32 FDatabaseViewModel::SetSelectedNode(int32 PoseIdx, bool bClearSelection, bool bDrawQuery, TConstArrayView<float> InQueryVector)
-	{
-		int32 SelectedSourceAssetIdx = INDEX_NONE;
+	return false;
+}
 
-		if (bClearSelection)
-		{
-			RemovePreviewActors();
-		}
+int32 FDatabaseViewModel::SetSelectedNode(int32 PoseIdx, bool bClearSelection, bool bDrawQuery, TConstArrayView<float> InQueryVector)
+{
+	int32 SelectedSourceAssetIdx = INDEX_NONE;
 
-		bIsEditorSelection = false;
-		bDrawQueryVector = bDrawQuery;
-		QueryVector = InQueryVector;
+	if (bClearSelection)
+	{
+		RemovePreviewActors();
+	}
 
-		if (FAsyncPoseSearchDatabasesManagement::RequestAsyncBuildIndex(PoseSearchDatabase, ERequestAsyncBuildFlag::ContinueRequest))
+	bIsEditorSelection = false;
+	bDrawQueryVector = bDrawQuery;
+	QueryVector = InQueryVector;
+
+	if (FAsyncPoseSearchDatabasesManagement::RequestAsyncBuildIndex(PoseSearchDatabase, ERequestAsyncBuildFlag::ContinueRequest))
+	{
+		const FSearchIndex& SearchIndex = PoseSearchDatabase->GetSearchIndex();
+		if (SearchIndex.PoseMetadata.IsValidIndex(PoseIdx))
 		{
-			const FSearchIndex& SearchIndex = PoseSearchDatabase->GetSearchIndex();
-			if (SearchIndex.PoseMetadata.IsValidIndex(PoseIdx))
+			const uint32 IndexAssetIndex = SearchIndex.PoseMetadata[PoseIdx].GetAssetIndex();
+			if (SearchIndex.Assets.IsValidIndex(IndexAssetIndex))
 			{
-				const uint32 IndexAssetIndex = SearchIndex.PoseMetadata[PoseIdx].GetAssetIndex();
-				if (SearchIndex.Assets.IsValidIndex(IndexAssetIndex))
+				FDatabasePreviewActor PreviewActor;
+				if (PreviewActor.SpawnPreviewActor(GetWorld(), PoseSearchDatabase, IndexAssetIndex, PoseIdx))
 				{
-					FDatabasePreviewActor PreviewActor = SpawnPreviewActor(IndexAssetIndex, PoseIdx);
-					if (PreviewActor.IsValid())
-					{
-						const FSearchIndexAsset& IndexAsset = SearchIndex.Assets[IndexAssetIndex];
-						MaxPreviewPlayLength = FMath::Max(MaxPreviewPlayLength, IndexAsset.GetLastSampleTime(PoseSearchDatabase->Schema->SampleRate) - PreviewActor.PlayTimeOffset);
-						MinPreviewPlayLength = FMath::Min(MinPreviewPlayLength, IndexAsset.GetFirstSampleTime(PoseSearchDatabase->Schema->SampleRate) - PreviewActor.PlayTimeOffset);
-						PreviewActors.Add(PreviewActor);
-						SelectedSourceAssetIdx = IndexAsset.GetSourceAssetIdx();
-					}
+					const FSearchIndexAsset& IndexAsset = SearchIndex.Assets[IndexAssetIndex];
+					MaxPreviewPlayLength = FMath::Max(MaxPreviewPlayLength, IndexAsset.GetLastSampleTime(PoseSearchDatabase->Schema->SampleRate) - PreviewActor.GetPlayTimeOffset());
+					MinPreviewPlayLength = FMath::Min(MinPreviewPlayLength, IndexAsset.GetFirstSampleTime(PoseSearchDatabase->Schema->SampleRate) - PreviewActor.GetPlayTimeOffset());
+					PreviewActors.Add(PreviewActor);
+					SelectedSourceAssetIdx = IndexAsset.GetSourceAssetIdx();
 				}
 			}
+		}
 
-			FBoneContainer BoneContainer;
-			BoneContainer.InitializeTo(PoseSearchDatabase->Schema->BoneIndicesWithParents, UE::Anim::FCurveFilterSettings(UE::Anim::ECurveFilterMode::DisallowAll), *PoseSearchDatabase->Schema->Skeleton);
-			// capturing BoneContainer by copy, since it has mutable properties
-			ParallelFor(PreviewActors.Num(), [this, BoneContainer](int32 PreviewActorIndex) { PreviewActors[PreviewActorIndex].Process(BoneContainer); }, ParallelForFlags);
+		DatabaseDataDetails.Pin()->Reconstruct();
 
-			DatabaseDataDetails.Pin()->Reconstruct();
+		for (FDatabasePreviewActor& PreviewActor : GetPreviewActors())
+		{
+			PreviewActor.UpdatePreviewActor(PoseSearchDatabase, PlayTime, bQuantizeAnimationToPoseData);
+		}
 
-			UpdatePreviewActors();
+		SetPlayTime(0.f, false);
+	}
 
-			SetPlayTime(0.f, false);
-		}
+	ProcessSelectedActor(nullptr);
 
-		ProcessSelectedActor(nullptr);
+	return SelectedSourceAssetIdx;
+}
 
-		return SelectedSourceAssetIdx;
-	}
+void FDatabaseViewModel::SetSelectedNodes(const TArrayView<TSharedPtr<FDatabaseAssetTreeNode>>& InSelectedNodes)
+{
+	RemovePreviewActors();
 
-	void FDatabaseViewModel::SetSelectedNodes(const TArrayView<TSharedPtr<FDatabaseAssetTreeNode>>& InSelectedNodes)
+	if (FAsyncPoseSearchDatabasesManagement::RequestAsyncBuildIndex(PoseSearchDatabase, ERequestAsyncBuildFlag::ContinueRequest))
 	{
-		RemovePreviewActors();
-
-		if (FAsyncPoseSearchDatabasesManagement::RequestAsyncBuildIndex(PoseSearchDatabase, ERequestAsyncBuildFlag::ContinueRequest))
+		TMap<int32, int32> AssociatedAssetIndices;
+		for (int32 i = 0; i < InSelectedNodes.Num(); ++i)
 		{
-			TMap<int32, int32> AssociatedAssetIndices;
-			for (int32 i = 0; i < InSelectedNodes.Num(); ++i)
-			{
-				AssociatedAssetIndices.FindOrAdd(InSelectedNodes[i]->SourceAssetIdx) = i;
-			}
+			AssociatedAssetIndices.FindOrAdd(InSelectedNodes[i]->SourceAssetIdx) = i;
+		}
 
-			const FSearchIndex& SearchIndex = PoseSearchDatabase->GetSearchIndex();
-			for (int32 IndexAssetIndex = 0; IndexAssetIndex < SearchIndex.Assets.Num(); ++IndexAssetIndex)
+		const FSearchIndex& SearchIndex = PoseSearchDatabase->GetSearchIndex();
+		for (int32 IndexAssetIndex = 0; IndexAssetIndex < SearchIndex.Assets.Num(); ++IndexAssetIndex)
+		{
+			const FSearchIndexAsset& IndexAsset = SearchIndex.Assets[IndexAssetIndex];
+			if (const int32* SelectedNodesIndex = AssociatedAssetIndices.Find(IndexAsset.GetSourceAssetIdx()))
 			{
-				const FSearchIndexAsset& IndexAsset = SearchIndex.Assets[IndexAssetIndex];
-				if (const int32* SelectedNodesIndex = AssociatedAssetIndices.Find(IndexAsset.GetSourceAssetIdx()))
+				FDatabasePreviewActor PreviewActor;
+				if (PreviewActor.SpawnPreviewActor(GetWorld(), PoseSearchDatabase, IndexAssetIndex))
 				{
-					FDatabasePreviewActor PreviewActor = SpawnPreviewActor(IndexAssetIndex);
-					if (PreviewActor.IsValid())
-					{
-						MaxPreviewPlayLength = FMath::Max(MaxPreviewPlayLength, IndexAsset.GetLastSampleTime(PoseSearchDatabase->Schema->SampleRate) - IndexAsset.GetFirstSampleTime(PoseSearchDatabase->Schema->SampleRate));
-						PreviewActors.Add(PreviewActor);
-					}
+					MaxPreviewPlayLength = FMath::Max(MaxPreviewPlayLength, IndexAsset.GetLastSampleTime(PoseSearchDatabase->Schema->SampleRate) - IndexAsset.GetFirstSampleTime(PoseSearchDatabase->Schema->SampleRate));
+					PreviewActors.Add(PreviewActor);
 				}
 			}
-
-			FBoneContainer BoneContainer;
-			BoneContainer.InitializeTo(PoseSearchDatabase->Schema->BoneIndicesWithParents, UE::Anim::FCurveFilterSettings(UE::Anim::ECurveFilterMode::DisallowAll), *PoseSearchDatabase->Schema->Skeleton);
-			// capturing BoneContainer by copy, since it has mutable properties
-			ParallelFor(PreviewActors.Num(), [this, BoneContainer](int32 PreviewActorIndex) { PreviewActors[PreviewActorIndex].Process(BoneContainer); }, ParallelForFlags);
-
-			DatabaseDataDetails.Pin()->Reconstruct();
-			UpdatePreviewActors();
 		}
 
-		ProcessSelectedActor(nullptr);
+		DatabaseDataDetails.Pin()->Reconstruct();
+		for (FDatabasePreviewActor& PreviewActor : GetPreviewActors())
+		{
+			PreviewActor.UpdatePreviewActor(PoseSearchDatabase, PlayTime, bQuantizeAnimationToPoseData);
+		}
 	}
 
-	void FDatabaseViewModel::ProcessSelectedActor(AActor* Actor)
+	ProcessSelectedActor(nullptr);
+}
+
+void FDatabaseViewModel::ProcessSelectedActor(AActor* Actor)
+{
+	SelectedActorIndexAssetIndex = INDEX_NONE;
+	if (const FDatabasePreviewActor* SelectedPreviewActor = PreviewActors.FindByPredicate([Actor](const FDatabasePreviewActor& PreviewActor) { return PreviewActor.GetActor() == Actor; }))
 	{
-		SelectedActorIndexAssetIndex = INDEX_NONE;
-		if (const FDatabasePreviewActor* SelectedPreviewActor = PreviewActors.FindByPredicate([Actor](const FDatabasePreviewActor& PreviewActor) { return PreviewActor.Actor == Actor; }))
-		{
-			SelectedActorIndexAssetIndex = SelectedPreviewActor->IndexAssetIndex;
-		}
+		SelectedActorIndexAssetIndex = SelectedPreviewActor->GetIndexAssetIndex();
 	}
+}
 
-	void FDatabaseViewModel::SetDrawQueryVector(bool bValue)
+void FDatabaseViewModel::SetDrawQueryVector(bool bValue)
+{
+	if (bDrawQueryVector != bValue)
 	{
-		if (bDrawQueryVector != bValue)
-		{
-			bDrawQueryVector = bValue;
-			DatabaseDataDetails.Pin()->Reconstruct();
-		}
+		bDrawQueryVector = bValue;
+		DatabaseDataDetails.Pin()->Reconstruct();
 	}
+}
 
-	const FSearchIndexAsset* FDatabaseViewModel::GetSelectedActorIndexAsset() const
+const FSearchIndexAsset* FDatabaseViewModel::GetSelectedActorIndexAsset() const
+{
+	if (FAsyncPoseSearchDatabasesManagement::RequestAsyncBuildIndex(PoseSearchDatabase, ERequestAsyncBuildFlag::ContinueRequest))
 	{
-		if (FAsyncPoseSearchDatabasesManagement::RequestAsyncBuildIndex(PoseSearchDatabase, ERequestAsyncBuildFlag::ContinueRequest))
+		const FSearchIndex& SearchIndex = PoseSearchDatabase->GetSearchIndex();
+		if (SearchIndex.Assets.IsValidIndex(SelectedActorIndexAssetIndex))
 		{
-			const FSearchIndex& SearchIndex = PoseSearchDatabase->GetSearchIndex();
-			if (SearchIndex.Assets.IsValidIndex(SelectedActorIndexAssetIndex))
-			{
-				return &SearchIndex.Assets[SelectedActorIndexAssetIndex];
-			}
+			return &SearchIndex.Assets[SelectedActorIndexAssetIndex];
 		}
-		return nullptr;
 	}
+	return nullptr;
+}
 
-	TRange<double> FDatabaseViewModel::GetPreviewPlayRange() const
-	{
-		constexpr double ViewRangeSlack = 0.2;
-		return TRange<double>(MinPreviewPlayLength - ViewRangeSlack, MaxPreviewPlayLength + ViewRangeSlack);
-	}
+TRange<double> FDatabaseViewModel::GetPreviewPlayRange() const
+{
+	constexpr double ViewRangeSlack = 0.2;
+	return TRange<double>(MinPreviewPlayLength - ViewRangeSlack, MaxPreviewPlayLength + ViewRangeSlack);
+}
 
-	float FDatabaseViewModel::GetPlayTime() const
-	{
-		return PlayTime;
-	}
+float FDatabaseViewModel::GetPlayTime() const
+{
+	return PlayTime;
+}
 
-	void FDatabaseViewModel::SetPlayTime(float NewPlayTime, bool bInTickPlayTime)
-	{
-		NewPlayTime = FMath::Clamp(NewPlayTime, MinPreviewPlayLength, MaxPreviewPlayLength);
-		DeltaTimeMultiplier = bInTickPlayTime ? DeltaTimeMultiplier : 0.0f;
+void FDatabaseViewModel::SetPlayTime(float NewPlayTime, bool bInTickPlayTime)
+{
+	NewPlayTime = FMath::Clamp(NewPlayTime, MinPreviewPlayLength, MaxPreviewPlayLength);
+	DeltaTimeMultiplier = bInTickPlayTime ? DeltaTimeMultiplier : 0.0f;
 
-		if (!FMath::IsNearlyEqual(PlayTime, NewPlayTime))
+	if (!FMath::IsNearlyEqual(PlayTime, NewPlayTime))
+	{
+		PlayTime = NewPlayTime;
+		if (FAsyncPoseSearchDatabasesManagement::RequestAsyncBuildIndex(PoseSearchDatabase, ERequestAsyncBuildFlag::ContinueRequest))
 		{
-			PlayTime = NewPlayTime;
-			if (FAsyncPoseSearchDatabasesManagement::RequestAsyncBuildIndex(PoseSearchDatabase, ERequestAsyncBuildFlag::ContinueRequest))
+			for (FDatabasePreviewActor& PreviewActor : GetPreviewActors())
 			{
-				UpdatePreviewActors(true);
+				PreviewActor.UpdatePreviewActor(PoseSearchDatabase, PlayTime, bQuantizeAnimationToPoseData);
 			}
 		}
 	}
+}
 
-	bool FDatabaseViewModel::GetAnimationTime(int32 SourceAssetIdx, float& CurrentPlayTime, FVector& BlendParameters) const
+bool FDatabaseViewModel::GetAnimationTime(int32 SourceAssetIdx, float& CurrentPlayTime, FVector& BlendParameters) const
+{
+	if (FAsyncPoseSearchDatabasesManagement::RequestAsyncBuildIndex(PoseSearchDatabase, ERequestAsyncBuildFlag::ContinueRequest))
 	{
-		if (FAsyncPoseSearchDatabasesManagement::RequestAsyncBuildIndex(PoseSearchDatabase, ERequestAsyncBuildFlag::ContinueRequest))
+		const FSearchIndex& SearchIndex = PoseSearchDatabase->GetSearchIndex();
+		for (const FDatabasePreviewActor& PreviewActor : GetPreviewActors())
 		{
-			const FSearchIndex& SearchIndex = PoseSearchDatabase->GetSearchIndex();
-			for (const FDatabasePreviewActor& PreviewActor : GetPreviewActors())
+			if (PreviewActor.GetIndexAssetIndex() >= 0 && PreviewActor.GetIndexAssetIndex() < SearchIndex.Assets.Num())
 			{
-				if (PreviewActor.IndexAssetIndex >= 0 && PreviewActor.IndexAssetIndex < SearchIndex.Assets.Num() && PreviewActor.Sampler.IsInitialized())
+				const FSearchIndexAsset& IndexAsset = SearchIndex.Assets[PreviewActor.GetIndexAssetIndex()];
+				if (IndexAsset.GetSourceAssetIdx() == SourceAssetIdx)
 				{
-					const FSearchIndexAsset& IndexAsset = SearchIndex.Assets[PreviewActor.IndexAssetIndex];
-					if (IndexAsset.GetSourceAssetIdx() == SourceAssetIdx)
-					{
-						CurrentPlayTime = PreviewActor.Sampler.ToNormalizedTime(PlayTime + IndexAsset.GetFirstSampleTime(PoseSearchDatabase->Schema->SampleRate) + PreviewActor.PlayTimeOffset);
-						BlendParameters = IndexAsset.GetBlendParameters();
-						return true;
-					}
+					CurrentPlayTime = PreviewActor.GetSampler().ToNormalizedTime(PlayTime + IndexAsset.GetFirstSampleTime(PoseSearchDatabase->Schema->SampleRate) + PreviewActor.GetPlayTimeOffset());
+					BlendParameters = IndexAsset.GetBlendParameters();
+					return true;
 				}
 			}
+		}
 
-			for (const FSearchIndexAsset& IndexAsset : SearchIndex.Assets)
+		for (const FSearchIndexAsset& IndexAsset : SearchIndex.Assets)
+		{
+			if (IndexAsset.GetSourceAssetIdx() == SourceAssetIdx)
 			{
-				if (IndexAsset.GetSourceAssetIdx() == SourceAssetIdx)
-				{
-					CurrentPlayTime = PlayTime + IndexAsset.GetFirstSampleTime(PoseSearchDatabase->Schema->SampleRate);
-					BlendParameters = IndexAsset.GetBlendParameters();
+				CurrentPlayTime = PlayTime + IndexAsset.GetFirstSampleTime(PoseSearchDatabase->Schema->SampleRate);
+				BlendParameters = IndexAsset.GetBlendParameters();
 
-					const bool bIsBlendSpace = PoseSearchDatabase->GetAnimationAssetStruct(IndexAsset).GetPtr<FPoseSearchDatabaseBlendSpace>() != nullptr;
-					if (bIsBlendSpace && !FMath::IsNearlyEqual(MaxPreviewPlayLength, MinPreviewPlayLength))
-					{
-						CurrentPlayTime = (CurrentPlayTime - MaxPreviewPlayLength) / (MaxPreviewPlayLength - MinPreviewPlayLength);
-					}
-					return true;
+				const bool bIsBlendSpace = PoseSearchDatabase->GetAnimationAssetStruct(IndexAsset).GetPtr<FPoseSearchDatabaseBlendSpace>() != nullptr;
+				if (bIsBlendSpace && !FMath::IsNearlyEqual(MaxPreviewPlayLength, MinPreviewPlayLength))
+				{
+					CurrentPlayTime = (CurrentPlayTime - MaxPreviewPlayLength) / (MaxPreviewPlayLength - MinPreviewPlayLength);
 				}
+				return true;
 			}
 		}
-
-		CurrentPlayTime = 0.f;
-		BlendParameters = FVector::ZeroVector;
-		return false;
 	}
+
+	CurrentPlayTime = 0.f;
+	BlendParameters = FVector::ZeroVector;
+	return false;
+}
 }
diff --git a/Engine/Plugins/Experimental/Animation/PoseSearch/Source/Editor/Private/PoseSearchDebuggerViewModel.cpp b/Engine/Plugins/Experimental/Animation/PoseSearch/Source/Editor/Private/PoseSearchDebuggerViewModel.cpp
--- a/Engine/Plugins/Experimental/Animation/PoseSearch/Source/Editor/Private/PoseSearchDebuggerViewModel.cpp
+++ b/Engine/Plugins/Experimental/Animation/PoseSearch/Source/Editor/Private/PoseSearchDebuggerViewModel.cpp
@@ -164,7 +164,7 @@ void FDebuggerViewModel::UpdatePoseSearchContext(UPoseSearchMeshComponent::FUpda
 		{
 			InOutContext.StartTime = Skeletons[SelectedPose].Time;
 			InOutContext.Time = Skeletons[SelectedPose].Time;
-			InOutContext.bMirrored = Skeletons[SelectedPose].bMirrored;
+			InOutContext.MirrorDataCache = Skeletons[SelectedPose].bMirrored ? &MirrorDataCache : nullptr;
 			InOutContext.bLoop = DatabaseAsset->IsLooping();
 		}
 
@@ -289,29 +289,25 @@ const USkinnedMeshComponent* FDebuggerViewModel::GetMeshComponent() const
 
 void FDebuggerViewModel::FillCompactPoseAndComponentRefRotations()
 {
-	bool bResetMirrorBonesAndRotations = true;
+	bool bResetMirrorDataCache = true;
 	if (const UPoseSearchDatabase* Database = GetCurrentDatabase())
 	{
-		if (UMirrorDataTable* MirrorDataTable = Database->Schema->MirrorDataTable)
+		if (MirrorDataCache.GetMirrorDataTable() != Database->Schema->MirrorDataTable)
 		{
 			if (UPoseSearchMeshComponent* MeshComponent = Skeletons[ActivePose].Component.Get())
 			{
 				if (MeshComponent->RequiredBones.IsValid())
 				{
-					if (CompactPoseMirrorBones.Num() == 0 || ComponentSpaceRefRotations.Num() == 0)
-					{
-						MirrorDataTable->FillCompactPoseAndComponentRefRotations(MeshComponent->RequiredBones, CompactPoseMirrorBones, ComponentSpaceRefRotations);
-					}
-					bResetMirrorBonesAndRotations = false;
+					MirrorDataCache.Init(Database->Schema->MirrorDataTable, MeshComponent->RequiredBones);
+					bResetMirrorDataCache = false;
 				}
 			}
 		}
 	}
 
-	if (bResetMirrorBonesAndRotations)
+	if (bResetMirrorDataCache)
 	{
-		CompactPoseMirrorBones.Reset();
-		ComponentSpaceRefRotations.Reset();
+		MirrorDataCache.Reset();
 	}
 }
 
diff --git a/Engine/Plugins/Experimental/Animation/PoseSearch/Source/Editor/Private/PoseSearchDebuggerViewModel.h b/Engine/Plugins/Experimental/Animation/PoseSearch/Source/Editor/Private/PoseSearchDebuggerViewModel.h
--- a/Engine/Plugins/Experimental/Animation/PoseSearch/Source/Editor/Private/PoseSearchDebuggerViewModel.h
+++ b/Engine/Plugins/Experimental/Animation/PoseSearch/Source/Editor/Private/PoseSearchDebuggerViewModel.h
@@ -5,6 +5,7 @@
 #include "CoreMinimal.h"
 #include "Misc/Attribute.h"
 #include "PoseSearchMeshComponent.h"
+#include "PoseSearch/PoseSearchMirrorDataCache.h"
 #include "Templates/SharedPointer.h"
 #include "UObject/WeakObjectPtrTemplates.h"
 
@@ -114,12 +115,7 @@ private:
 	/** Anim Instance associated with this debugger instance */
 	uint64 AnimInstanceId = 0;
 
-	/** Compact pose format of Mirror Bone Map */
-	TCustomBoneIndexArray<FCompactPoseBoneIndex, FCompactPoseBoneIndex> CompactPoseMirrorBones;
-
-	/** Pre-calculated component space rotations of reference pose */
-	TCustomBoneIndexArray<FQuat, FCompactPoseBoneIndex> ComponentSpaceRefRotations;
-
+	FMirrorDataCache MirrorDataCache;
 
 	/** Index for each type of skeleton we store for debug visualization */
 	enum ESkeletonIndex
diff --git a/Engine/Plugins/Experimental/Animation/PoseSearch/Source/Editor/Private/PoseSearchMeshComponent.cpp b/Engine/Plugins/Experimental/Animation/PoseSearch/Source/Editor/Private/PoseSearchMeshComponent.cpp
--- a/Engine/Plugins/Experimental/Animation/PoseSearch/Source/Editor/Private/PoseSearchMeshComponent.cpp
+++ b/Engine/Plugins/Experimental/Animation/PoseSearch/Source/Editor/Private/PoseSearchMeshComponent.cpp
@@ -7,7 +7,7 @@
 #include "Animation/BlendSpace.h"
 #include "Animation/MirrorDataTable.h"
 #include "Engine/SkinnedAsset.h"
-#include "PoseSearch/PoseSearchContext.h"
+#include "PoseSearch/PoseSearchMirrorDataCache.h"
 
 void UPoseSearchMeshComponent::Initialize(const FTransform& InComponentToWorld)
 {
@@ -51,6 +51,7 @@ void UPoseSearchMeshComponent::ResetToStart()
 
 void UPoseSearchMeshComponent::UpdatePose(const FUpdateContext& UpdateContext)
 {
+	// @todo: reuse asset sampler here if needed at all
 	using namespace UE::PoseSearch;
 
 	FMemMark Mark(FMemStack::Get());
@@ -126,13 +127,9 @@ void UPoseSearchMeshComponent::UpdatePose(const FUpdateContext& UpdateContext)
 		}
 	}
 
-	if (UpdateContext.bMirrored)
+	if (UpdateContext.MirrorDataCache)
 	{
-		FAnimationRuntime::MirrorPose(
-			CompactPose, 
-			UpdateContext.MirrorDataTable->MirrorAxis, 
-			*UpdateContext.CompactPoseMirrorBones, 
-			*UpdateContext.ComponentSpaceRefRotations);
+		UpdateContext.MirrorDataCache->MirrorPose(CompactPose);
 	}
 
 	FCSPose<FCompactPose> ComponentSpacePose;
@@ -146,9 +143,9 @@ void UPoseSearchMeshComponent::UpdatePose(const FUpdateContext& UpdateContext)
 		SetBoneTransformByName(BoneName, BoneTransform, EBoneSpaces::ComponentSpace);
 	}
 
-	if (UpdateContext.bMirrored)
+	if (UpdateContext.MirrorDataCache)
 	{
-		LastRootMotionDelta = UE::PoseSearch::MirrorTransform(LastRootMotionDelta, UpdateContext.MirrorDataTable->MirrorAxis, (*UpdateContext.ComponentSpaceRefRotations)[FCompactPoseBoneIndex(RootBoneIndexType)]);
+		LastRootMotionDelta = UpdateContext.MirrorDataCache->MirrorTransform(LastRootMotionDelta);
 	}
 
 	const FTransform ComponentTransform = LastRootMotionDelta * StartingTransform;
diff --git a/Engine/Plugins/Experimental/Animation/PoseSearch/Source/Editor/Private/PoseSearchMeshComponent.h b/Engine/Plugins/Experimental/Animation/PoseSearch/Source/Editor/Private/PoseSearchMeshComponent.h
--- a/Engine/Plugins/Experimental/Animation/PoseSearch/Source/Editor/Private/PoseSearchMeshComponent.h
+++ b/Engine/Plugins/Experimental/Animation/PoseSearch/Source/Editor/Private/PoseSearchMeshComponent.h
@@ -9,6 +9,11 @@
 
 class UBlendSpace;
 
+namespace UE::PoseSearch
+{
+	struct FMirrorDataCache;
+};
+
 UCLASS()
 class UPoseSearchMeshComponent : public UPoseableMeshComponent
 {
@@ -22,11 +27,8 @@ public:
 		float StartTime = 0.0f;
 		float Time = 0.0f;
 		bool bLoop = false;
-		bool bMirrored = false;
 		FVector BlendParameters = FVector::Zero();
-		const UMirrorDataTable* MirrorDataTable = nullptr;
-		TCustomBoneIndexArray<FCompactPoseBoneIndex, FCompactPoseBoneIndex>* CompactPoseMirrorBones = nullptr;
-		TCustomBoneIndexArray<FQuat, FCompactPoseBoneIndex>* ComponentSpaceRefRotations = nullptr;
+		const UE::PoseSearch::FMirrorDataCache* MirrorDataCache = nullptr;
 	};
 
 	void Refresh();
diff --git a/Engine/Plugins/Experimental/Animation/PoseSearch/Source/Editor/Private/SPoseSearchDatabaseViewport.cpp b/Engine/Plugins/Experimental/Animation/PoseSearch/Source/Editor/Private/SPoseSearchDatabaseViewport.cpp
--- a/Engine/Plugins/Experimental/Animation/PoseSearch/Source/Editor/Private/SPoseSearchDatabaseViewport.cpp
+++ b/Engine/Plugins/Experimental/Animation/PoseSearch/Source/Editor/Private/SPoseSearchDatabaseViewport.cpp
@@ -28,32 +28,62 @@ namespace UE::PoseSearch
 		);
 	}
 
+	void SDatabaseViewport::OnToggleDisplayRootMotionSpeed()
+	{
+		AssetEditorPtr.Pin()->GetViewModelSharedPtr().ToSharedRef()->ToggleDisplayRootMotionSpeed();
+		Invalidate();
+	}
+
+	bool SDatabaseViewport::IsDisplayRootMotionSpeedChecked() const
+	{
+		return AssetEditorPtr.Pin()->GetViewModelSharedPtr().ToSharedRef()->IsDisplayRootMotionSpeedChecked();
+	}
+
+	void SDatabaseViewport::OnToggleQuantizeAnimationToPoseData()
+	{
+		AssetEditorPtr.Pin()->GetViewModelSharedPtr().ToSharedRef()->ToggleQuantizeAnimationToPoseData();
+		Invalidate();
+	}
+
+	bool SDatabaseViewport::IsQuantizeAnimationToPoseDataChecked() const
+	{
+		return AssetEditorPtr.Pin()->GetViewModelSharedPtr().ToSharedRef()->IsQuantizeAnimationToPoseDataChecked();
+	}
+
+	void SDatabaseViewport::OnToggleShowBones()
+	{
+		AssetEditorPtr.Pin()->GetViewModelSharedPtr().ToSharedRef()->ToggleShowBones();
+		Invalidate();
+	}
+
+	bool SDatabaseViewport::IsShowBonesChecked() const
+	{
+		return AssetEditorPtr.Pin()->GetViewModelSharedPtr().ToSharedRef()->IsShowBonesChecked();
+	}
+
 	void SDatabaseViewport::BindCommands()
 	{
 		SEditorViewport::BindCommands();
 
 		const FDatabaseEditorCommands& Commands = FDatabaseEditorCommands::Get();
 
-		TSharedRef<FDatabaseViewModel> ViewModelRef =
-			AssetEditorPtr.Pin()->GetViewModelSharedPtr().ToSharedRef();
-
 		CommandList->MapAction(
 			Commands.ShowDisplayRootMotionSpeed,
-			FExecuteAction::CreateSP(ViewModelRef, &FDatabaseViewModel::OnToggleDisplayRootMotionSpeed),
+			FExecuteAction::CreateSP(this, &SDatabaseViewport::OnToggleDisplayRootMotionSpeed),
 			FCanExecuteAction(),
-			FIsActionChecked::CreateSP(ViewModelRef, &FDatabaseViewModel::IsDisplayRootMotionSpeedChecked));
+			FIsActionChecked::CreateSP(this, &SDatabaseViewport::IsDisplayRootMotionSpeedChecked));
 
 		CommandList->MapAction(
 			Commands.ShowQuantizeAnimationToPoseData,
-			FExecuteAction::CreateSP(ViewModelRef, &FDatabaseViewModel::OnToggleQuantizeAnimationToPoseData),
+			FExecuteAction::CreateSP(this, &SDatabaseViewport::OnToggleQuantizeAnimationToPoseData),
 			FCanExecuteAction(),
-			FIsActionChecked::CreateSP(ViewModelRef, &FDatabaseViewModel::IsQuantizeAnimationToPoseDataChecked));
+			FIsActionChecked::CreateSP(this, &SDatabaseViewport::IsQuantizeAnimationToPoseDataChecked));
 
 		CommandList->MapAction(
 			Commands.ShowBones,
-			FExecuteAction::CreateSP(ViewModelRef, &FDatabaseViewModel::OnToggleShowBones),
+			FExecuteAction::CreateSP(this, &SDatabaseViewport::OnToggleShowBones),
 			FCanExecuteAction(),
-			FIsActionChecked::CreateSP(ViewModelRef, &FDatabaseViewModel::IsShowBones));
+			FIsActionChecked::CreateSP(this, &SDatabaseViewport::IsShowBonesChecked));
 		
 	}
 
diff --git a/Engine/Plugins/Experimental/Animation/PoseSearch/Source/Editor/Public/PoseSearchDatabaseViewModel.h b/Engine/Plugins/Experimental/Animation/PoseSearch/Source/Editor/Public/PoseSearchDatabaseViewModel.h
--- a/Engine/Plugins/Experimental/Animation/PoseSearch/Source/Editor/Public/PoseSearchDatabaseViewModel.h
+++ b/Engine/Plugins/Experimental/Animation/PoseSearch/Source/Editor/Public/PoseSearchDatabaseViewModel.h
@@ -6,7 +6,7 @@
 #include "InstancedStruct.h"
 #include "MovieSceneFwd.h"
 #include "PoseSearch/PoseSearchAssetSampler.h"
-#include "PoseSearch/PoseSearchIndex.h"
+#include "PoseSearch/PoseSearchMirrorDataCache.h"
 #include "PoseSearchDatabasePreviewScene.h"
 #include "UObject/GCObject.h"
 
@@ -28,18 +28,40 @@ namespace UE::PoseSearch
 	struct FDatabasePreviewActor
 	{
 	public:
+		bool SpawnPreviewActor(UWorld* World, const UPoseSearchDatabase* PoseSearchDatabase, int32 IndexAssetIdx, int32 PoseIdxForTimeOffset = INDEX_NONE);
+		void UpdatePreviewActor(const UPoseSearchDatabase* PoseSearchDatabase, float PlayTime, bool bQuantizeAnimationToPoseData);
+		bool DrawPreviewActor(const UPoseSearchDatabase* PoseSearchDatabase, bool bDisplayRootMotionSpeed, TConstArrayView<float> QueryVector);
+		void Destroy();
+
+		const UDebugSkelMeshComponent* GetDebugSkelMeshComponent() const;
+		const UAnimPreviewInstance* GetAnimPreviewInstance() const;
+		const FAnimationAssetSampler& GetSampler() const { return Sampler; }
+		FTransform ExtractRootTransform(float Time) const;
+		void ExtractPose(float Time, FCompactPose& OutPose) const;
+
+		const AActor* GetActor() const { return Actor.Get(); }
+		int32 GetIndexAssetIndex() const { return IndexAssetIndex; }
+		int32 GetCurrentPoseIndex() const { return CurrentPoseIndex; }
+		float GetPlayTimeOffset() const { return PlayTimeOffset; }
+		
+	private:
+		UAnimPreviewInstance* GetAnimPreviewInstanceInternal();
+
 		TObjectPtr<AActor> Actor;
 		int32 IndexAssetIndex = INDEX_NONE;
 		int32 CurrentPoseIndex = INDEX_NONE;
 		float PlayTimeOffset = 0.f;
 		float CurrentTime = 0.f;
+
+		FTransform RootTransformCurrentQuantizedTime = FTransform::Identity;
+		FTransform RootTransformCurrent = FTransform::Identity;
+		FTransform RootTransformOrigin = FTransform::Identity;
+
 		FAnimationAssetSampler Sampler;
-		FTransform QuantizedTimeRootTransform = FTransform::Identity;
+		FMirrorDataCache MirrorDataCache;
 
-		bool IsValid() const;
-		void Process(const FBoneContainer& BoneContainer);
-		UDebugSkelMeshComponent* GetDebugSkelMeshComponent();
-		UAnimPreviewInstance* GetAnimPreviewInstance();
+		TArray<FVector> SampledRootMotion;
+		TArray<float> SampledRootMotionSpeed;
 	};
 
 	class FDatabaseViewModel : public TSharedFromThis<FDatabaseViewModel>, public FGCObject
@@ -72,14 +94,14 @@ namespace UE::PoseSearch
 		TArray<FDatabasePreviewActor>& GetPreviewActors() { return PreviewActors; }
 		const TArray<FDatabasePreviewActor>& GetPreviewActors() const { return PreviewActors; }
 
-		void OnToggleDisplayRootMotionSpeed() { bDisplayRootMotionSpeed = !bDisplayRootMotionSpeed; }
-		bool IsDisplayRootMotionSpeedChecked() const { return bDisplayRootMotionSpeed; };
+		void ToggleDisplayRootMotionSpeed() { bDisplayRootMotionSpeed = !bDisplayRootMotionSpeed;	}
+		bool IsDisplayRootMotionSpeedChecked() const { return bDisplayRootMotionSpeed; }
 
-		void OnToggleQuantizeAnimationToPoseData() { bQuantizeAnimationToPoseData = !bQuantizeAnimationToPoseData; }
-		bool IsQuantizeAnimationToPoseDataChecked() const { return bQuantizeAnimationToPoseData; };
+		void ToggleQuantizeAnimationToPoseData() { bQuantizeAnimationToPoseData = !bQuantizeAnimationToPoseData; }
+		bool IsQuantizeAnimationToPoseDataChecked() const { return bQuantizeAnimationToPoseData; }
 
-		void OnToggleShowBones() { bShowBones = !bShowBones; }
-		bool IsShowBones() const { return bShowBones; };
+		void ToggleShowBones() { bShowBones = !bShowBones; }
+		bool IsShowBonesChecked() const { return bShowBones; }
 
 		void AddSequenceToDatabase(UAnimSequence* AnimSequence);
 		void AddBlendSpaceToDatabase(UBlendSpace* BlendSpace);
@@ -110,8 +132,9 @@ namespace UE::PoseSearch
 		bool IsEditorSelection() const { return bIsEditorSelection; }
 		bool GetAnimationTime(int32 SourceAssetIdx, float& CurrentPlayTime, FVector& BlendParameters) const;
 
-
 	private:
+		UWorld* GetWorld();
+
 		float PlayTime = 0.f;
 		float DeltaTimeMultiplier = 1.f;
 
@@ -137,19 +160,9 @@ namespace UE::PoseSearch
 
 		/** Is animation debug draw enabled */
 		bool bDisplayRootMotionSpeed = false;
-
 		bool bQuantizeAnimationToPoseData = false;
-
 		bool bShowBones = false;
 
 		int32 SelectedActorIndexAssetIndex = INDEX_NONE;
-
-		UWorld* GetWorld();
-
-		FDatabasePreviewActor SpawnPreviewActor(int32 IndexAssetIndex, int32 PoseIdxForTimeOffset = INDEX_NONE);
-
-		void UpdatePreviewActors(bool bInTickPlayTime = false);
-
-		FTransform MirrorRootTransform(const FTransform& RootTransform);
 	};
 }
diff --git a/Engine/Plugins/Experimental/Animation/PoseSearch/Source/Editor/Public/SPoseSearchDatabaseViewport.h b/Engine/Plugins/Experimental/Animation/PoseSearch/Source/Editor/Public/SPoseSearchDatabaseViewport.h
--- a/Engine/Plugins/Experimental/Animation/PoseSearch/Source/Editor/Public/SPoseSearchDatabaseViewport.h
+++ b/Engine/Plugins/Experimental/Animation/PoseSearch/Source/Editor/Public/SPoseSearchDatabaseViewport.h
@@ -52,6 +52,13 @@ namespace UE::PoseSearch
 		virtual TSharedPtr<SWidget> MakeViewportToolbar() override;
 		// ~End of SEditorViewport interface
 
+		void OnToggleDisplayRootMotionSpeed();
+		bool IsDisplayRootMotionSpeedChecked() const;
+		void OnToggleQuantizeAnimationToPoseData();
+		bool IsQuantizeAnimationToPoseDataChecked() const;
+		void OnToggleShowBones();
+		bool IsShowBonesChecked() const;
+
 		/** The viewport toolbar */
 		TSharedPtr<SPoseSearchDatabaseViewportToolBar> ViewportToolbar;
 
diff --git a/Engine/Plugins/Experimental/Animation/PoseSearch/Source/Runtime/Private/PoseSearchAssetIndexer.cpp b/Engine/Plugins/Experimental/Animation/PoseSearch/Source/Runtime/Private/PoseSearchAssetIndexer.cpp
--- a/Engine/Plugins/Experimental/Animation/PoseSearch/Source/Runtime/Private/PoseSearchAssetIndexer.cpp
+++ b/Engine/Plugins/Experimental/Animation/PoseSearch/Source/Runtime/Private/PoseSearchAssetIndexer.cpp
@@ -82,26 +82,11 @@ static FSamplingParam WrapOrClampSamplingParam(bool bCanWrap, float SamplingPara
 //////////////////////////////////////////////////////////////////////////
 // FAssetSamplingContext
 FAssetSamplingContext::FAssetSamplingContext(const UPoseSearchDatabase& Database, const FBoneContainer& BoneContainer)
+: FMirrorDataCache(Database.Schema ? Database.Schema->MirrorDataTable : nullptr, BoneContainer)
 {
 	check(Database.Schema);
-	MirrorDataTable = Database.Schema->MirrorDataTable;
 	BaseCostBias = Database.BaseCostBias;
 	LoopingCostBias = Database.LoopingCostBias;
-
-	if (MirrorDataTable)
-	{
-		MirrorDataTable->FillCompactPoseAndComponentRefRotations(BoneContainer, CompactPoseMirrorBones, ComponentSpaceRefRotations);
-	}
-	else
-	{
-		CompactPoseMirrorBones.Reset();
-		ComponentSpaceRefRotations.Reset();
-	}
-}
-
-FTransform FAssetSamplingContext::MirrorTransform(const FTransform& InTransform) const
-{
-	return UE::PoseSearch::MirrorTransform(InTransform, MirrorDataTable->MirrorAxis, ComponentSpaceRefRotations[FCompactPoseBoneIndex(RootBoneIndexType)]);
 }
 
 //////////////////////////////////////////////////////////////////////////
@@ -364,15 +349,9 @@ FAssetIndexer::CachedEntry& FAssetIndexer::GetEntry(float SampleTime)
 
 		Pose[FCompactPoseBoneIndex(RootBoneIndexType)].SetIdentity();
 
-		if (SearchIndexAsset.IsMirrored() && Schema.MirrorDataTable)
+		if (SearchIndexAsset.IsMirrored())
 		{
-			FAnimationRuntime::MirrorPose(
-				Pose,
-				Schema.MirrorDataTable->MirrorAxis,
-				SamplingContext.CompactPoseMirrorBones,
-				SamplingContext.ComponentSpaceRefRotations
-			);
-			// Note curves and attributes are not used during the indexing process and therefore don't need to be mirrored
+			SamplingContext.MirrorPose(Pose);
 		}
 
 		FCSPose<FCompactPose> StackComponentSpacePose;
@@ -427,8 +406,7 @@ FTransform FAssetIndexer::GetTransform(float SampleTime, bool& bClamped, const F
 	CachedEntry& Entry = GetEntry(SampleTime);
 	bClamped = Entry.bClamped;
 
-	const FTransform MirroredRootTransform = MirrorTransform(Entry.RootTransform);
-	return CalculateComponentSpaceTransform(Entry, BoneReference) * MirroredRootTransform;
+	return CalculateComponentSpaceTransform(Entry, BoneReference) * MirrorTransform(Entry.RootTransform);
 }
 
 FTransform FAssetIndexer::CalculateComponentSpaceTransform(FAssetIndexer::CachedEntry& Entry, int8 SchemaBoneIdx)
diff --git a/Engine/Plugins/Experimental/Animation/PoseSearch/Source/Runtime/Private/PoseSearchContext.cpp b/Engine/Plugins/Experimental/Animation/PoseSearch/Source/Runtime/Private/PoseSearchContext.cpp
--- a/Engine/Plugins/Experimental/Animation/PoseSearch/Source/Runtime/Private/PoseSearchContext.cpp
+++ b/Engine/Plugins/Experimental/Animation/PoseSearch/Source/Runtime/Private/PoseSearchContext.cpp
@@ -620,13 +620,4 @@ bool FSearchContext::IsCurrentResultFromDatabase(const UPoseSearchDatabase* Data
 	return CurrentResult.IsValid() && CurrentResult.Database == Database;
 }
 
-FTransform MirrorTransform(const FTransform& InTransform, EAxis::Type MirrorAxis, const FQuat& ReferenceRotation)
-{
-	const FVector T = FAnimationRuntime::MirrorVector(InTransform.GetTranslation(), MirrorAxis);
-	const FQuat Q = FAnimationRuntime::MirrorQuat(InTransform.GetRotation(), MirrorAxis);
-	const FQuat QR = Q * FAnimationRuntime::MirrorQuat(ReferenceRotation, MirrorAxis).Inverse() * ReferenceRotation;
-	const FTransform Result = FTransform(QR, T, InTransform.GetScale3D());
-	return Result;
-}
-
 } // namespace UE::PoseSearch
\ No newline at end of file
diff --git a/Engine/Plugins/Experimental/Animation/PoseSearch/Source/Runtime/Private/PoseSearchMirrorDataCache.cpp b/Engine/Plugins/Experimental/Animation/PoseSearch/Source/Runtime/Private/PoseSearchMirrorDataCache.cpp
new file mode 100644
--- /dev/null
+++ b/Engine/Plugins/Experimental/Animation/PoseSearch/Source/Runtime/Private/PoseSearchMirrorDataCache.cpp
@@ -0,0 +1,70 @@
+// Copyright Epic Games, Inc. All Rights Reserved.
+
+#include "PoseSearch/PoseSearchMirrorDataCache.h"
+#include "AnimationRuntime.h"
+#include "PoseSearch/PoseSearchDefines.h"
+
+namespace UE::PoseSearch
+{
+
+FMirrorDataCache::FMirrorDataCache()
+{
+}
+
+FMirrorDataCache::FMirrorDataCache(const UMirrorDataTable* InMirrorDataTable, const FBoneContainer& BoneContainer)
+{
+	Init(InMirrorDataTable, BoneContainer);
+}
+
+void FMirrorDataCache::Init(const UMirrorDataTable* InMirrorDataTable, const FBoneContainer& BoneContainer)
+{
+	if (InMirrorDataTable != nullptr)
+	{
+		check(BoneContainer.IsValid());
+		MirrorDataTable = InMirrorDataTable;
+		MirrorDataTable->FillCompactPoseAndComponentRefRotations(BoneContainer, CompactPoseMirrorBones, ComponentSpaceRefRotations);
+	}
+	else
+	{
+		Reset();
+	}
+}
+
+void FMirrorDataCache::Reset()
+{
+	MirrorDataTable = nullptr;
+	CompactPoseMirrorBones.Reset();
+	ComponentSpaceRefRotations.Reset();
+}
+
+FTransform FMirrorDataCache::MirrorTransform(const FTransform& InTransform) const
+{
+	if (MirrorDataTable != nullptr)
+	{
+		const EAxis::Type MirrorAxis = MirrorDataTable->MirrorAxis;
+		const FQuat& ReferenceRotation = ComponentSpaceRefRotations[FCompactPoseBoneIndex(RootBoneIndexType)];
+
+		const FVector T = FAnimationRuntime::MirrorVector(InTransform.GetTranslation(), MirrorAxis);
+		const FQuat Q = FAnimationRuntime::MirrorQuat(InTransform.GetRotation(), MirrorAxis);
+		const FQuat QR = Q * FAnimationRuntime::MirrorQuat(ReferenceRotation, MirrorAxis).Inverse() * ReferenceRotation;
+		const FTransform Result = FTransform(QR, T, InTransform.GetScale3D());
+		return Result;
+	}
+	return InTransform;
+}
+
+void FMirrorDataCache::MirrorPose(FCompactPose& Pose) const
+{
+	if (MirrorDataTable != nullptr && MirrorDataTable != nullptr)
+	{
+		FAnimationRuntime::MirrorPose(
+			Pose,
+			MirrorDataTable->MirrorAxis,
+			CompactPoseMirrorBones,
+			ComponentSpaceRefRotations
+		);
+		// Note curves and attributes are not used during the indexing process and therefore don't need to be mirrored
+	}
+}
+
+} // namespace UE::PoseSearch
diff --git a/Engine/Plugins/Experimental/Animation/PoseSearch/Source/Runtime/Public/PoseSearch/PoseSearchAssetIndexer.h b/Engine/Plugins/Experimental/Animation/PoseSearch/Source/Runtime/Public/PoseSearch/PoseSearchAssetIndexer.h
--- a/Engine/Plugins/Experimental/Animation/PoseSearch/Source/Runtime/Public/PoseSearch/PoseSearchAssetIndexer.h
+++ b/Engine/Plugins/Experimental/Animation/PoseSearch/Source/Runtime/Public/PoseSearch/PoseSearchAssetIndexer.h
@@ -5,7 +5,7 @@
 #if WITH_EDITOR
 
 #include "BonePose.h"
-#include "CoreMinimal.h"
+#include "PoseSearch/PoseSearchMirrorDataCache.h"
 #include "PoseSearch/PoseSearchDefines.h"
 #include "PoseSearch/PoseSearchFeatureChannel.h"
 #include "PoseSearch/PoseSearchSchema.h"
@@ -19,23 +19,12 @@ struct FAnimationAssetSampler;
 struct FSearchIndexAsset;
 struct FPoseMetadata;
 
-struct FAssetSamplingContext
+struct FAssetSamplingContext : public FMirrorDataCache
 {
 	float BaseCostBias = 0.f;
 	float LoopingCostBias = 0.f;
 
-	// Mirror data table pointer copied from Schema for convenience
-	TObjectPtr<const UMirrorDataTable> MirrorDataTable;
-
-	// Compact pose format of Mirror Bone Map
-	TCustomBoneIndexArray<FCompactPoseBoneIndex, FCompactPoseBoneIndex> CompactPoseMirrorBones;
-
-	// Pre-calculated component space rotations of reference pose, which allows mirror to work with any joint orientation
-	// Only initialized and used when a mirroring table is specified
-	TCustomBoneIndexArray<FQuat, FCompactPoseBoneIndex> ComponentSpaceRefRotations;
-
 	FAssetSamplingContext(const UPoseSearchDatabase& Database, const FBoneContainer& BoneContainer);
-	FTransform MirrorTransform(const FTransform& Transform) const;
 };
 
 class FAssetIndexer
diff --git a/Engine/Plugins/Experimental/Animation/PoseSearch/Source/Runtime/Public/PoseSearch/PoseSearchContext.h b/Engine/Plugins/Experimental/Animation/PoseSearch/Source/Runtime/Public/PoseSearch/PoseSearchContext.h
--- a/Engine/Plugins/Experimental/Animation/PoseSearch/Source/Runtime/Public/PoseSearch/PoseSearchContext.h
+++ b/Engine/Plugins/Experimental/Animation/PoseSearch/Source/Runtime/Public/PoseSearch/PoseSearchContext.h
@@ -337,6 +337,4 @@ private:
 #endif // UE_POSE_SEARCH_TRACE_ENABLED
 };
 
-POSESEARCH_API FTransform MirrorTransform(const FTransform& InTransform, EAxis::Type MirrorAxis, const FQuat& ReferenceRotation);
-
 } // namespace UE::PoseSearch
diff --git a/Engine/Plugins/Experimental/Animation/PoseSearch/Source/Runtime/Public/PoseSearch/PoseSearchMirrorDataCache.h b/Engine/Plugins/Experimental/Animation/PoseSearch/Source/Runtime/Public/PoseSearch/PoseSearchMirrorDataCache.h
new file mode 100644
--- /dev/null
+++ b/Engine/Plugins/Experimental/Animation/PoseSearch/Source/Runtime/Public/PoseSearch/PoseSearchMirrorDataCache.h
@@ -0,0 +1,41 @@
+// Copyright Epic Games, Inc. All Rights Reserved.
+
+#pragma once
+
+#include "Animation/MirrorDataTable.h"
+
+class UMirrorDataTable;
+struct FBoneContainer;
+struct FCompactPose;
+
+namespace UE::PoseSearch
+{
+
+//POSESEARCH_API FTransform MirrorTransform(const FTransform& InTransform, EAxis::Type MirrorAxis, const FQuat& ReferenceRotation);
+
+struct POSESEARCH_API FMirrorDataCache
+{
+	FMirrorDataCache();
+	FMirrorDataCache(const UMirrorDataTable* InMirrorDataTable, const FBoneContainer& BoneContainer);
+	
+	void Init(const UMirrorDataTable* InMirrorDataTable, const FBoneContainer& BoneContainer);
+	void Reset();
+
+	FTransform MirrorTransform(const FTransform& InTransform) const;
+	void MirrorPose(FCompactPose& Pose) const;
+	const UMirrorDataTable* GetMirrorDataTable() const { return MirrorDataTable.Get(); }
+
+private:
+	// Mirror data table pointer copied from Schema for convenience
+	TWeakObjectPtr<const UMirrorDataTable> MirrorDataTable;
+
+	// Compact pose format of Mirror Bone Map
+	TCustomBoneIndexArray<FCompactPoseBoneIndex, FCompactPoseBoneIndex> CompactPoseMirrorBones;
+
+	// Pre-calculated component space rotations of reference pose, which allows mirror to work with any joint orientation
+	// Only initialized and used when a mirroring table is specified
+	TCustomBoneIndexArray<FQuat, FCompactPoseBoneIndex> ComponentSpaceRefRotations;
+};
+
+} // namespace UE::PoseSearch
+
