motion matching - removing redundant code

diff --git a/Engine/Plugins/Experimental/Animation/PoseSearch/Source/Runtime/Private/PoseSearchMirrorDataCache.cpp b/Engine/Plugins/Experimental/Animation/PoseSearch/Source/Runtime/Private/PoseSearchMirrorDataCache.cpp
--- a/Engine/Plugins/Experimental/Animation/PoseSearch/Source/Runtime/Private/PoseSearchMirrorDataCache.cpp
+++ b/Engine/Plugins/Experimental/Animation/PoseSearch/Source/Runtime/Private/PoseSearchMirrorDataCache.cpp
@@ -55,7 +55,7 @@ FTransform FMirrorDataCache::MirrorTransform(const FTransform& InTransform) cons
 
 void FMirrorDataCache::MirrorPose(FCompactPose& Pose) const
 {
-	if (MirrorDataTable != nullptr && MirrorDataTable != nullptr)
+	if (MirrorDataTable != nullptr)
 	{
 		FAnimationRuntime::MirrorPose(
 			Pose,
