motion matching - fix for exclusion ranges calculation from anim notity

diff --git a/Engine/Plugins/Experimental/Animation/PoseSearch/Source/Runtime/Private/PoseSearchAssetIndexer.cpp b/Engine/Plugins/Experimental/Animation/PoseSearch/Source/Runtime/Private/PoseSearchAssetIndexer.cpp
--- a/Engine/Plugins/Experimental/Animation/PoseSearch/Source/Runtime/Private/PoseSearchAssetIndexer.cpp
+++ b/Engine/Plugins/Experimental/Animation/PoseSearch/Source/Runtime/Private/PoseSearchAssetIndexer.cpp
@@ -466,7 +466,7 @@ bool FAssetIndexer::GetSampleRotation(FQuat& OutSampleRotation, float SampleTime
 							TempBoneReference.Initialize(BoneContainer.GetSkeletonAsset());
 							if (TempBoneReference.HasValidSetup())
 							{
-								const float SamplingAttributeTime = AnimNotifyEvent.GetTriggerTime();
+								const float SamplingAttributeTime = AnimNotifyEvent.GetTime();
 								const FTransform RootBoneTransform = GetTransform(OriginTime, bUnused, RootSchemaBoneIdx);
 								const FTransform SamplingAttributeBoneTransform = GetTransform(SamplingAttributeTime, bUnused, TempBoneReference);
 								OutSampleRotation = RootBoneTransform.InverseTransformRotation(SamplingAttributeBoneTransform.GetRotation());
@@ -540,7 +540,7 @@ bool FAssetIndexer::GetSamplePositionInternal(FVector& OutSamplePosition, float
 							TempBoneReference.Initialize(BoneContainer.GetSkeletonAsset());
 							if (TempBoneReference.HasValidSetup())
 							{
-								const float SamplingAttributeTime = AnimNotifyEvent.GetTriggerTime();
+								const float SamplingAttributeTime = AnimNotifyEvent.GetTime();
 								const FTransform RootBoneTransform = GetTransform(OriginTime, bUnused, RootSchemaBoneIdx);
 								const FTransform SamplingAttributeBoneTransform = GetTransform(SamplingAttributeTime, bClamped, TempBoneReference);
 								if (SchemaOriginBoneIdx == RootSchemaBoneIdx)
@@ -638,7 +638,7 @@ bool FAssetIndexer::GetSampleVelocity(FVector& OutSampleVelocity, float SampleTi
 							TempBoneReference.Initialize(BoneContainer.GetSkeletonAsset());
 							if (TempBoneReference.HasValidSetup())
 							{
-								const float SamplingAttributeTime = AnimNotifyEvent.GetTriggerTime();
+								const float SamplingAttributeTime = AnimNotifyEvent.GetTime();
 
 								if (GetSamplePositionInternal(BonePositionPast, SamplingAttributeTime - FiniteDelta, bUseCharacterSpaceVelocities ? OriginTime - FiniteDelta : OriginTime, bClampedPast, SchemaSampleBoneIdx, SchemaOriginBoneIdx, SamplingAttributeId) &&
 									GetSamplePositionInternal(BonePositionPresent, SamplingAttributeTime, OriginTime, bUnused, SchemaSampleBoneIdx, SchemaOriginBoneIdx, SamplingAttributeId))
diff --git a/Engine/Plugins/Experimental/Animation/PoseSearch/Source/Runtime/Private/PoseSearchAssetSampler.cpp b/Engine/Plugins/Experimental/Animation/PoseSearch/Source/Runtime/Private/PoseSearchAssetSampler.cpp
--- a/Engine/Plugins/Experimental/Animation/PoseSearch/Source/Runtime/Private/PoseSearchAssetSampler.cpp
+++ b/Engine/Plugins/Experimental/Animation/PoseSearch/Source/Runtime/Private/PoseSearchAssetSampler.cpp
@@ -641,7 +641,7 @@ void FAnimationAssetSampler::ExtractPoseSearchNotifyStates(float Time, TFunction
 		}
 
 		// @todo: is this condition necessary? can we just rely on the ExtractionInterval?
-		if (NotifyEvent->GetTriggerTime() > SampleTime || NotifyEvent->GetEndTriggerTime() < SampleTime)
+		if (NotifyEvent->GetTime() > SampleTime || (NotifyEvent->GetTime() + NotifyEvent->GetDuration()) < SampleTime)
 		{
 			continue;
 		}
diff --git a/Engine/Plugins/Experimental/Animation/PoseSearch/Source/Runtime/Private/PoseSearchDatabase.cpp b/Engine/Plugins/Experimental/Animation/PoseSearch/Source/Runtime/Private/PoseSearchDatabase.cpp
--- a/Engine/Plugins/Experimental/Animation/PoseSearch/Source/Runtime/Private/PoseSearchDatabase.cpp
+++ b/Engine/Plugins/Experimental/Animation/PoseSearch/Source/Runtime/Private/PoseSearchDatabase.cpp
@@ -718,7 +718,7 @@ void UPoseSearchDatabase::SynchronizeWithExternalDependencies(TConstArrayView<UA
 						{
 							FPoseSearchDatabaseSequence DatabaseSequence;
 							DatabaseSequence.Sequence = Sequence;
-							DatabaseSequence.SamplingRange = FFloatInterval(NotifyEvent.GetTriggerTime(), NotifyEvent.GetEndTriggerTime());
+							DatabaseSequence.SamplingRange = FFloatInterval(NotifyEvent.GetTime(), NotifyEvent.GetTime() + NotifyEvent.GetDuration());
 							DatabaseSequence.bSynchronizeWithExternalDependency = true;
 							AnimationAssets.Add(FInstancedStruct::Make(DatabaseSequence));
 							bModified = true;
@@ -727,7 +727,7 @@ void UPoseSearchDatabase::SynchronizeWithExternalDependencies(TConstArrayView<UA
 						{
 							FPoseSearchDatabaseAnimComposite DatabaseAnimComposite;
 							DatabaseAnimComposite.AnimComposite = AnimComposite;
-							DatabaseAnimComposite.SamplingRange = FFloatInterval(NotifyEvent.GetTriggerTime(), NotifyEvent.GetEndTriggerTime());
+							DatabaseAnimComposite.SamplingRange = FFloatInterval(NotifyEvent.GetTime(), NotifyEvent.GetTime() + NotifyEvent.GetDuration());
 							DatabaseAnimComposite.bSynchronizeWithExternalDependency = true;
 							AnimationAssets.Add(FInstancedStruct::Make(DatabaseAnimComposite));
 							bModified = true;
@@ -736,7 +736,7 @@ void UPoseSearchDatabase::SynchronizeWithExternalDependencies(TConstArrayView<UA
 						{
 							FPoseSearchDatabaseAnimMontage DatabaseAnimMontage;
 							DatabaseAnimMontage.AnimMontage = AnimMontage;
-							DatabaseAnimMontage.SamplingRange = FFloatInterval(NotifyEvent.GetTriggerTime(), NotifyEvent.GetEndTriggerTime());
+							DatabaseAnimMontage.SamplingRange = FFloatInterval(NotifyEvent.GetTime(), NotifyEvent.GetTime() + NotifyEvent.GetDuration());
 							DatabaseAnimMontage.bSynchronizeWithExternalDependency = true;
 							AnimationAssets.Add(FInstancedStruct::Make(DatabaseAnimMontage));
 							bModified = true;
diff --git a/Engine/Plugins/Experimental/Animation/PoseSearch/Source/Runtime/Private/PoseSearchDerivedData.cpp b/Engine/Plugins/Experimental/Animation/PoseSearch/Source/Runtime/Private/PoseSearchDerivedData.cpp
--- a/Engine/Plugins/Experimental/Animation/PoseSearch/Source/Runtime/Private/PoseSearchDerivedData.cpp
+++ b/Engine/Plugins/Experimental/Animation/PoseSearch/Source/Runtime/Private/PoseSearchDerivedData.cpp
@@ -304,7 +304,7 @@ static void FindValidSequenceIntervals(const UAnimSequenceBase* SequenceBase, FF
 		{
 			if (const UAnimNotifyState_PoseSearchExcludeFromDatabase* ExclusionNotifyState = Cast<const UAnimNotifyState_PoseSearchExcludeFromDatabase>(NotifyEvent->NotifyStateClass))
 			{
-				FFloatRange ExclusionRange = FFloatRange::Inclusive(NotifyEvent->GetTriggerTime(), NotifyEvent->GetEndTriggerTime());
+				FFloatRange ExclusionRange = FFloatRange::Inclusive(NotifyEvent->GetTime(), NotifyEvent->GetTime() + NotifyEvent->GetDuration());
 
 				// Split every valid range based on the exclusion range just found. Because this might increase the 
 				// number of ranges in ValidRanges, the algorithm iterates from end to start.
diff --git a/Engine/Plugins/Experimental/Animation/PoseSearch/Source/Runtime/Private/PoseSearchDerivedDataKey.cpp b/Engine/Plugins/Experimental/Animation/PoseSearch/Source/Runtime/Private/PoseSearchDerivedDataKey.cpp
--- a/Engine/Plugins/Experimental/Animation/PoseSearch/Source/Runtime/Private/PoseSearchDerivedDataKey.cpp
+++ b/Engine/Plugins/Experimental/Animation/PoseSearch/Source/Runtime/Private/PoseSearchDerivedDataKey.cpp
@@ -40,7 +40,7 @@ FKeyBuilder::FKeyBuilder(const UObject* Object, bool bUseDataVer, bool bPerformC
 	if (bUseDataVer)
 	{
 		// used to invalidate the key without having to change POSESEARCHDB_DERIVEDDATA_VER all the times
-		int32 POSESEARCHDB_DERIVEDDATA_VER_SMALL = 227;
+		int32 POSESEARCHDB_DERIVEDDATA_VER_SMALL = 229;
 		FGuid VersionGuid = FDevSystemGuids::GetSystemGuid(FDevSystemGuids::Get().POSESEARCHDB_DERIVEDDATA_VER);
 
 		*this << VersionGuid;
