-Added OnMotionMatchingStateUpdated anim node function that runs after selection
-Exposed future root motion velocity from motion matching

diff --git a/Engine/Plugins/Experimental/Animation/PoseSearch/Source/Editor/Private/AnimGraphNode_MotionMatching.cpp b/Engine/Plugins/Experimental/Animation/PoseSearch/Source/Editor/Private/AnimGraphNode_MotionMatching.cpp
--- a/Engine/Plugins/Experimental/Animation/PoseSearch/Source/Editor/Private/AnimGraphNode_MotionMatching.cpp
+++ b/Engine/Plugins/Experimental/Animation/PoseSearch/Source/Editor/Private/AnimGraphNode_MotionMatching.cpp
@@ -2,6 +2,7 @@
 
 #include "AnimGraphNode_MotionMatching.h"
 #include "Animation/AnimRootMotionProvider.h"
+#include "FindInBlueprintManager.h"
 
 #define LOCTEXT_NAMESPACE "AnimGraphNode_MotionMatching"
 
@@ -37,4 +38,65 @@ UScriptStruct* UAnimGraphNode_MotionMatching::GetTimePropertyStruct() const
 	return FAnimNode_MotionMatching::StaticStruct();
 }
 
+void UAnimGraphNode_MotionMatching::OnProcessDuringCompilation(IAnimBlueprintCompilationContext& InCompilationContext, IAnimBlueprintGeneratedClassCompiledData& OutCompiledData)
+{
+	Super::OnProcessDuringCompilation(InCompilationContext, OutCompiledData);
+	Node.OnMotionMatchingStateUpdated.SetFromFunction(OnMotionMatchingStateUpdatedFunction.ResolveMember<UFunction>(GetBlueprintClassFromNode()));
+}
+
+void UAnimGraphNode_MotionMatching::ValidateAnimNodeDuringCompilation(USkeleton* ForSkeleton, FCompilerResultsLog& MessageLog)
+{
+	Super::ValidateAnimNodeDuringCompilation(ForSkeleton, MessageLog);
+
+	ValidateFunctionRef(GET_MEMBER_NAME_CHECKED(UAnimGraphNode_MotionMatching, OnMotionMatchingStateUpdatedFunction), OnMotionMatchingStateUpdatedFunction, LOCTEXT("OnMotionMatchingStateUpdatedFunctionName", "On Motion Matching State Updated"), MessageLog);
+}
+
+void UAnimGraphNode_MotionMatching::PostEditChangeProperty(FPropertyChangedEvent& PropertyChangedEvent)
+{
+	const FName PropertyName = (PropertyChangedEvent.Property != nullptr) ? PropertyChangedEvent.Property->GetFName() : NAME_None;
+
+	if (PropertyName == GET_MEMBER_NAME_CHECKED(UAnimGraphNode_MotionMatching, OnMotionMatchingStateUpdatedFunction))
+	{
+		Node.OnMotionMatchingStateUpdated.SetFromFunction(OnMotionMatchingStateUpdatedFunction.ResolveMember<UFunction>(GetBlueprintClassFromNode()));
+		GetAnimBlueprint()->RequestRefreshExtensions();
+		GetSchema()->ReconstructNode(*this);
+	}
+
+	Super::PostEditChangeProperty(PropertyChangedEvent);
+}
+
+void UAnimGraphNode_MotionMatching::AddSearchMetaDataInfo(TArray<struct FSearchTagDataPair>& OutTaggedMetaData) const
+{
+	Super::AddSearchMetaDataInfo(OutTaggedMetaData);
+
+	const auto ConditionallyTagNodeFuncRef = [&OutTaggedMetaData](const FMemberReference& FuncMember, const FText& LocText)
+	{
+		if (IsPotentiallyBoundFunction(FuncMember))
+		{
+			const FText FunctionName = FText::FromName(FuncMember.GetMemberName());
+			OutTaggedMetaData.Add(FSearchTagDataPair(LocText, FunctionName));
+		}
+	};
+
+	// Conditionally include anim node function references as part of the node's search metadata
+	ConditionallyTagNodeFuncRef(OnMotionMatchingStateUpdatedFunction, LOCTEXT("OnMotionMatchingStateUpdatedFunctionName", "On Motion Matching State Updated"));
+}
+
+void UAnimGraphNode_MotionMatching::GetBoundFunctionsInfo(TArray<TPair<FName, FName>>& InOutBindingsInfo)
+{
+	Super::GetBoundFunctionsInfo(InOutBindingsInfo);
+	const FName CategoryName = TEXT("Functions|Motion Matching");
+
+	if (OnMotionMatchingStateUpdatedFunction.ResolveMember<UFunction>(GetBlueprintClassFromNode()) != nullptr)
+	{
+		InOutBindingsInfo.Emplace(CategoryName, GET_MEMBER_NAME_CHECKED(UAnimGraphNode_MotionMatching, OnMotionMatchingStateUpdatedFunction));
+	}
+}
+
+bool UAnimGraphNode_MotionMatching::ReferencesFunction(const FName& InFunctionName, const UStruct* InScope) const
+{
+	return Super::ReferencesFunction(InFunctionName, InScope)
+		|| Node.OnMotionMatchingStateUpdated.GetFunctionName() == InFunctionName;
+}
+
 #undef LOCTEXT_NAMESPACE
diff --git a/Engine/Plugins/Experimental/Animation/PoseSearch/Source/Editor/Public/AnimGraphNode_MotionMatching.h b/Engine/Plugins/Experimental/Animation/PoseSearch/Source/Editor/Public/AnimGraphNode_MotionMatching.h
--- a/Engine/Plugins/Experimental/Animation/PoseSearch/Source/Editor/Public/AnimGraphNode_MotionMatching.h
+++ b/Engine/Plugins/Experimental/Animation/PoseSearch/Source/Editor/Public/AnimGraphNode_MotionMatching.h
@@ -6,6 +6,7 @@
 #include "UObject/ObjectMacros.h"
 #include "AnimGraphNode_BlendStack.h"
 #include "PoseSearch/AnimNode_MotionMatching.h"
+#include "Animation/AnimNodeFunctionRef.h"
 #include "AnimGraphNode_MotionMatching.generated.h"
 
 
@@ -25,4 +26,14 @@ class UAnimGraphNode_MotionMatching : public UAnimGraphNode_BlendStack_Base
 	virtual void BakeDataDuringCompilation(class FCompilerResultsLog& MessageLog) override;
 
 	virtual FAnimNode_BlendStack_Standalone* GetBlendStackNode() const override { return (FAnimNode_BlendStack_Standalone*)(&Node); }
+
+	virtual void OnProcessDuringCompilation(IAnimBlueprintCompilationContext& InCompilationContext, IAnimBlueprintGeneratedClassCompiledData& OutCompiledData) override;
+	virtual void ValidateAnimNodeDuringCompilation(USkeleton* ForSkeleton, FCompilerResultsLog& MessageLog) override;
+	virtual void AddSearchMetaDataInfo(TArray<struct FSearchTagDataPair>& OutTaggedMetaData) const override;
+	virtual void GetBoundFunctionsInfo(TArray<TPair<FName, FName>>& InOutBindingsInfo) override;
+	virtual bool ReferencesFunction(const FName& InFunctionName, const UStruct* InScope) const override;
+	virtual void PostEditChangeProperty(FPropertyChangedEvent& PropertyChangedEvent) override;
+
+	UPROPERTY(EditAnywhere, Category = "Functions|Motion Matching", meta=(FunctionReference, AllowFunctionLibraries, PrototypeFunction="/Script/AnimGraphRuntime.AnimExecutionContextLibrary.Prototype_ThreadSafeAnimUpdateCall", DefaultBindingName="OnStateExit"), DisplayName="On Motion Matching State Updated")
+	FMemberReference OnMotionMatchingStateUpdatedFunction;
 };
diff --git a/Engine/Plugins/Experimental/Animation/PoseSearch/Source/Runtime/Private/AnimNode_MotionMatching.cpp b/Engine/Plugins/Experimental/Animation/PoseSearch/Source/Runtime/Private/AnimNode_MotionMatching.cpp
--- a/Engine/Plugins/Experimental/Animation/PoseSearch/Source/Runtime/Private/AnimNode_MotionMatching.cpp
+++ b/Engine/Plugins/Experimental/Animation/PoseSearch/Source/Runtime/Private/AnimNode_MotionMatching.cpp
@@ -189,6 +189,8 @@ void FAnimNode_MotionMatching::UpdateAssetPlayer(const FAnimationUpdateContext&
 		#endif // ENABLE_ANIM_DEBUG
 	);
 
+	UE::Anim::FNodeFunctionCaller::CallFunction(GetOnUpdateMotionMatchingStateFunction(), Context, *this);
+
 	// If a new pose is requested, blend into the new asset via BlendStackNode
 	if (MotionMatchingState.bJumpedToPose)
 	{
@@ -214,6 +216,11 @@ void FAnimNode_MotionMatching::UpdateAssetPlayer(const FAnimationUpdateContext&
 	bForceInterruptNextUpdate = false;
 }
 
+const FAnimNodeFunctionRef& FAnimNode_MotionMatching::GetOnUpdateMotionMatchingStateFunction() const
+{
+	return GET_ANIM_NODE_DATA(FAnimNodeFunctionRef, OnMotionMatchingStateUpdated);
+}
+
 void FAnimNode_MotionMatching::SetDatabaseToSearch(UPoseSearchDatabase* InDatabase, bool bForceInterruptIfNew)
 {
 	if (DatabasesToSearch.Num() == 1 && DatabasesToSearch[0] == InDatabase)
@@ -236,6 +243,11 @@ void FAnimNode_MotionMatching::SetDatabaseToSearch(UPoseSearchDatabase* InDataba
 	}
 }
 
+FVector FAnimNode_MotionMatching::GetEstimatedFutureRootMotionVelocity() const
+{
+	return MotionMatchingState.GetEstimatedFutureRootMotionVelocity();
+}
+
 void FAnimNode_MotionMatching::SetDatabasesToSearch(TConstArrayView<UPoseSearchDatabase*> InDatabases, bool bForceInterruptIfNew)
 {
 	// Check if InDatabases and DatabasesToSearch are the same.
diff --git a/Engine/Plugins/Experimental/Animation/PoseSearch/Source/Runtime/Private/PoseSearchFeatureChannel_Trajectory.cpp b/Engine/Plugins/Experimental/Animation/PoseSearch/Source/Runtime/Private/PoseSearchFeatureChannel_Trajectory.cpp
--- a/Engine/Plugins/Experimental/Animation/PoseSearch/Source/Runtime/Private/PoseSearchFeatureChannel_Trajectory.cpp
+++ b/Engine/Plugins/Experimental/Animation/PoseSearch/Source/Runtime/Private/PoseSearchFeatureChannel_Trajectory.cpp
@@ -215,4 +215,38 @@ float UPoseSearchFeatureChannel_Trajectory::GetEstimatedSpeedRatio(TConstArrayVi
 	return 1.f;
 }
 
+FVector UPoseSearchFeatureChannel_Trajectory::GetEstimatedFutureRootMotionVelocity(TConstArrayView<float> PoseVector) const
+{
+	using namespace UE::PoseSearch;
+
+	float LastSampleTimeOffset = 0.0f;
+	int32 LastSampleTimeOffsetIndex = INDEX_NONE;
+	FVector OutRootMotionTranslation = FVector::ZeroVector;
+
+	const int32 NumChannels = GetSubChannels().Num();
+	for (int32 ChannelIdx = 0; ChannelIdx < NumChannels; ++ChannelIdx)
+	{
+		const TObjectPtr<UPoseSearchFeatureChannel>& SubChannelPtr = GetSubChannels()[ChannelIdx];
+		if (const UPoseSearchFeatureChannel_Velocity* Velocity = Cast<UPoseSearchFeatureChannel_Velocity>(SubChannelPtr.Get()))
+		{
+			if (!Velocity->bNormalize && (Velocity->SampleTimeOffset > LastSampleTimeOffset))
+			{
+				LastSampleTimeOffset = Velocity->SampleTimeOffset;
+				LastSampleTimeOffsetIndex = ChannelIdx;
+			}
+		}
+	}
+
+	if (LastSampleTimeOffsetIndex != INDEX_NONE)
+	{
+		const TObjectPtr<UPoseSearchFeatureChannel>& SubChannelPtr = GetSubChannels()[LastSampleTimeOffsetIndex];
+		const UPoseSearchFeatureChannel_Velocity* Velocity = Cast<UPoseSearchFeatureChannel_Velocity>(SubChannelPtr.Get());
+		check(Velocity);
+		OutRootMotionTranslation = FFeatureVectorHelper::DecodeVector(PoseVector, Velocity->GetChannelDataOffset(), Velocity->ComponentStripping);
+	}
+
+	return OutRootMotionTranslation;
+}
+
+
 #undef LOCTEXT_NAMESPACE
\ No newline at end of file
diff --git a/Engine/Plugins/Experimental/Animation/PoseSearch/Source/Runtime/Private/PoseSearchFeatureChannel_Trajectory.h b/Engine/Plugins/Experimental/Animation/PoseSearch/Source/Runtime/Private/PoseSearchFeatureChannel_Trajectory.h
--- a/Engine/Plugins/Experimental/Animation/PoseSearch/Source/Runtime/Private/PoseSearchFeatureChannel_Trajectory.h
+++ b/Engine/Plugins/Experimental/Animation/PoseSearch/Source/Runtime/Private/PoseSearchFeatureChannel_Trajectory.h
@@ -80,5 +80,6 @@ public:
 #endif
 
 	float GetEstimatedSpeedRatio(TConstArrayView<float> QueryVector, TConstArrayView<float> PoseVector) const;
+	FVector GetEstimatedFutureRootMotionVelocity(TConstArrayView<float> PoseVector) const;
 };
 
diff --git a/Engine/Plugins/Experimental/Animation/PoseSearch/Source/Runtime/Private/PoseSearchLibrary.cpp b/Engine/Plugins/Experimental/Animation/PoseSearch/Source/Runtime/Private/PoseSearchLibrary.cpp
--- a/Engine/Plugins/Experimental/Animation/PoseSearch/Source/Runtime/Private/PoseSearchLibrary.cpp
+++ b/Engine/Plugins/Experimental/Animation/PoseSearch/Source/Runtime/Private/PoseSearchLibrary.cpp
@@ -66,6 +66,21 @@ void FMotionMatchingState::JumpToPose(const FAnimationUpdateContext& Context, co
 	bJumpedToPose = true;
 }
 
+FVector FMotionMatchingState::GetEstimatedFutureRootMotionVelocity() const
+{
+	using namespace UE::PoseSearch;
+	if (CurrentSearchResult.IsValid())
+	{
+		if (const UPoseSearchFeatureChannel_Trajectory* TrajectoryChannel = CurrentSearchResult.Database->Schema->FindFirstChannelOfType<UPoseSearchFeatureChannel_Trajectory>())
+		{
+			TConstArrayView<float> ResultData = CurrentSearchResult.Database->GetSearchIndex().GetPoseValues(CurrentSearchResult.PoseIdx);
+			return TrajectoryChannel->GetEstimatedFutureRootMotionVelocity(ResultData);
+		}
+	}
+
+	return FVector::ZeroVector;
+}
+
 void FMotionMatchingState::UpdateWantedPlayRate(const UE::PoseSearch::FSearchContext& SearchContext, const FFloatInterval& PlayRate, float TrajectorySpeedMultiplier)
 {
 	if (CurrentSearchResult.IsValid())
diff --git a/Engine/Plugins/Experimental/Animation/PoseSearch/Source/Runtime/Public/PoseSearch/AnimNode_MotionMatching.h b/Engine/Plugins/Experimental/Animation/PoseSearch/Source/Runtime/Public/PoseSearch/AnimNode_MotionMatching.h
--- a/Engine/Plugins/Experimental/Animation/PoseSearch/Source/Runtime/Public/PoseSearch/AnimNode_MotionMatching.h
+++ b/Engine/Plugins/Experimental/Animation/PoseSearch/Source/Runtime/Public/PoseSearch/AnimNode_MotionMatching.h
@@ -31,6 +31,10 @@ public:
 
 	const FMotionMatchingState& GetMotionMatchingState() const { return MotionMatchingState; }
 
+	FVector GetEstimatedFutureRootMotionVelocity() const;
+
+	const FAnimNodeFunctionRef& GetOnUpdateMotionMatchingStateFunction() const;
+
 private:
 	// FAnimNode_Base interface
 	// @todo: implement CacheBones_AnyThread to rebind the schema bones
@@ -130,6 +134,10 @@ private:
 	// If true, "Relevant anim" nodes that look for the highest weighted animation in a state will ignore this node
 	UPROPERTY(EditAnywhere, Category = Relevancy, meta = (FoldProperty, PinHiddenByDefault))
 	bool bIgnoreForRelevancyTest = false;
+
+	UPROPERTY(meta=(FoldProperty))
+	FAnimNodeFunctionRef OnMotionMatchingStateUpdated;
+
 #endif // WITH_EDITORONLY_DATA
 
 	friend class UAnimGraphNode_MotionMatching;
diff --git a/Engine/Plugins/Experimental/Animation/PoseSearch/Source/Runtime/Public/PoseSearch/PoseSearchLibrary.h b/Engine/Plugins/Experimental/Animation/PoseSearch/Source/Runtime/Public/PoseSearch/PoseSearchLibrary.h
--- a/Engine/Plugins/Experimental/Animation/PoseSearch/Source/Runtime/Public/PoseSearch/PoseSearchLibrary.h
+++ b/Engine/Plugins/Experimental/Animation/PoseSearch/Source/Runtime/Public/PoseSearch/PoseSearchLibrary.h
@@ -7,7 +7,6 @@
 #include "Animation/AnimExecutionContext.h"
 #include "Animation/AnimNodeReference.h"
 #include "Kismet/BlueprintFunctionLibrary.h"
-#include "Kismet/BlueprintFunctionLibrary.h"
 #include "PoseSearch/PoseSearchDefines.h"
 #include "PoseSearch/PoseSearchHistory.h"
 #include "PoseSearch/PoseSearchResult.h"
@@ -40,6 +39,8 @@ struct FMotionMatchingState
 
 	void UpdateRootBoneControl(const FAnimationUpdateContext& Context, float YawFromAnimationBlendRate);
 
+	FVector GetEstimatedFutureRootMotionVelocity() const;
+
 	UE::PoseSearch::FSearchResult CurrentSearchResult;
 
 	// Time since the last pose jump
diff --git a/Engine/Source/Runtime/Engine/Classes/Animation/AnimNodeFunctionRef.h b/Engine/Source/Runtime/Engine/Classes/Animation/AnimNodeFunctionRef.h
--- a/Engine/Source/Runtime/Engine/Classes/Animation/AnimNodeFunctionRef.h
+++ b/Engine/Source/Runtime/Engine/Classes/Animation/AnimNodeFunctionRef.h
@@ -100,9 +100,10 @@ private:
 
 	// Call the Update function of this node
 	static void Update(const FAnimationUpdateContext& InContext, FAnimNode_Base& InNode);
-	
+
+public:
 	// Call a generic function for this node
-	static void CallFunction(const FAnimNodeFunctionRef& InFunction, const FAnimationBaseContext& InContext, FAnimNode_Base& InNode);
+	ENGINE_API static void CallFunction(const FAnimNodeFunctionRef& InFunction, const FAnimationBaseContext& InContext, FAnimNode_Base& InNode);
 };
 
 }}
