motion matching - support for skeleton with no identity root bone transform

diff --git a/Engine/Plugins/Experimental/Animation/PoseSearch/Source/Editor/Private/PoseSearchDatabaseViewModel.cpp b/Engine/Plugins/Experimental/Animation/PoseSearch/Source/Editor/Private/PoseSearchDatabaseViewModel.cpp
--- a/Engine/Plugins/Experimental/Animation/PoseSearch/Source/Editor/Private/PoseSearchDatabaseViewModel.cpp
+++ b/Engine/Plugins/Experimental/Animation/PoseSearch/Source/Editor/Private/PoseSearchDatabaseViewModel.cpp
@@ -155,6 +155,17 @@ void FDatabasePreviewActor::UpdatePreviewActor(const UPoseSearchDatabase* PoseSe
 	}
 
 	Actor->SetActorTransform(RootTransformCurrent);
+
+	// @todo: optimize this bone container, since we only need the root bone here...
+	const FBoneContainer& BoneContainer = AnimInstance->GetRequiredBonesOnAnyThread();
+	if (BoneContainer.GetNumBones() > 0)
+	{
+		FMemMark Mark(FMemStack::Get());
+		FCompactPose Pose;
+		Pose.SetBoneContainer(&BoneContainer);
+		ExtractPose(QuantizedTime, Pose);
+		RootBoneTransformCurrentQuantizedTime = Pose[FCompactPoseBoneIndex(RootBoneIndexType)];
+	}
 }
 
 void FDatabasePreviewActor::Destroy()
@@ -176,7 +187,8 @@ bool FDatabasePreviewActor::DrawPreviewActor(const UPoseSearchDatabase* PoseSear
 		return false;
 	}
 
-	UE::PoseSearch::FDebugDrawParams DrawParams(Mesh->GetWorld(), Mesh, RootTransformCurrentQuantizedTime, PoseSearchDatabase);
+	const FTransform RootBoneTransformCurrentQuantizedTimeWorld = RootBoneTransformCurrentQuantizedTime * RootTransformCurrentQuantizedTime;
+	UE::PoseSearch::FDebugDrawParams DrawParams(Mesh->GetWorld(), Mesh, RootBoneTransformCurrentQuantizedTimeWorld, PoseSearchDatabase);
 	DrawParams.DrawFeatureVector(GetCurrentPoseIndex());
 
 	if (!QueryVector.IsEmpty())
diff --git a/Engine/Plugins/Experimental/Animation/PoseSearch/Source/Editor/Public/PoseSearchDatabaseViewModel.h b/Engine/Plugins/Experimental/Animation/PoseSearch/Source/Editor/Public/PoseSearchDatabaseViewModel.h
--- a/Engine/Plugins/Experimental/Animation/PoseSearch/Source/Editor/Public/PoseSearchDatabaseViewModel.h
+++ b/Engine/Plugins/Experimental/Animation/PoseSearch/Source/Editor/Public/PoseSearchDatabaseViewModel.h
@@ -53,10 +53,14 @@ namespace UE::PoseSearch
 		float PlayTimeOffset = 0.f;
 		float CurrentTime = 0.f;
 
+		// world space root transforms
 		FTransform RootTransformCurrentQuantizedTime = FTransform::Identity;
 		FTransform RootTransformCurrent = FTransform::Identity;
 		FTransform RootTransformOrigin = FTransform::Identity;
 
+		// local space root BONE transform
+		FTransform RootBoneTransformCurrentQuantizedTime = FTransform::Identity;
+
 		FAnimationAssetSampler Sampler;
 		FMirrorDataCache MirrorDataCache;
 
diff --git a/Engine/Plugins/Experimental/Animation/PoseSearch/Source/Runtime/Private/PoseSearchAssetIndexer.cpp b/Engine/Plugins/Experimental/Animation/PoseSearch/Source/Runtime/Private/PoseSearchAssetIndexer.cpp
--- a/Engine/Plugins/Experimental/Animation/PoseSearch/Source/Runtime/Private/PoseSearchAssetIndexer.cpp
+++ b/Engine/Plugins/Experimental/Animation/PoseSearch/Source/Runtime/Private/PoseSearchAssetIndexer.cpp
@@ -347,8 +347,6 @@ FAssetIndexer::CachedEntry& FAssetIndexer::GetEntry(float SampleTime)
 		}
 #endif // ENABLE_ANIM_DEBUG
 
-		Pose[FCompactPoseBoneIndex(RootBoneIndexType)].SetIdentity();
-
 		if (SearchIndexAsset.IsMirrored())
 		{
 			SamplingContext.MirrorPose(Pose);
@@ -373,14 +371,7 @@ FTransform FAssetIndexer::GetComponentSpaceTransform(float SampleTime, bool& bCl
 	CachedEntry& Entry = GetEntry(SampleTime);
 	bClamped = Entry.bClamped;
 
-	if (SchemaBoneIdx != RootSchemaBoneIdx)
-	{
-		return CalculateComponentSpaceTransform(Entry, SchemaBoneIdx);
-	}
-
-	// making sure the root bone transform is identity
-	check(CalculateComponentSpaceTransform(Entry, SchemaBoneIdx).Equals(FTransform::Identity));
-	return FTransform::Identity;
+	return CalculateComponentSpaceTransform(Entry, SchemaBoneIdx);
 }
 
 // returns the transform in animation space for the bone indexed by Schema->BoneReferences[SchemaBoneIdx] at SampleTime seconds
@@ -391,13 +382,7 @@ FTransform FAssetIndexer::GetTransform(float SampleTime, bool& bClamped, int8 Sc
 	CachedEntry& Entry = GetEntry(SampleTime);
 	bClamped = Entry.bClamped;
 
-	const FTransform MirroredRootTransform = MirrorTransform(Entry.RootTransform);
-	if (SchemaBoneIdx != RootSchemaBoneIdx)
-	{
-		return CalculateComponentSpaceTransform(Entry, SchemaBoneIdx) * MirroredRootTransform;
-	}
-
-	return MirroredRootTransform;
+	return CalculateComponentSpaceTransform(Entry, SchemaBoneIdx) * MirrorTransform(Entry.RootTransform);
 }
 
 // returns the transform in animation space for the BoneReference at SampleTime seconds
@@ -491,15 +476,9 @@ bool FAssetIndexer::GetSampleRotation(FQuat& OutSampleRotation, float SampleTime
 		return false;
 	}
 
-	if (SampleTime == OriginTime)
-	{
-		OutSampleRotation = GetComponentSpaceTransform(SampleTime, bUnused, SchemaSampleBoneIdx).GetRotation();
-		return true;
-	}
-
 	const FTransform RootBoneTransform = GetTransform(OriginTime, bUnused, RootSchemaBoneIdx);
-	const FTransform BoneTransform = GetTransform(SampleTime, bUnused, SchemaSampleBoneIdx);
-	OutSampleRotation = RootBoneTransform.InverseTransformRotation(BoneTransform.GetRotation());
+	const FTransform SampleBoneTransform = GetTransform(SampleTime, bUnused, SchemaSampleBoneIdx);
+	OutSampleRotation = RootBoneTransform.InverseTransformRotation(SampleBoneTransform.GetRotation());
 	return true;
 }
 
@@ -576,22 +555,6 @@ bool FAssetIndexer::GetSamplePositionInternal(FVector& OutSamplePosition, float
 		return false;
 	}
 
-	if (SampleTime == OriginTime)
-	{
-		if (SchemaOriginBoneIdx == RootSchemaBoneIdx)
-		{
-			OutSamplePosition = GetComponentSpaceTransform(SampleTime, bClamped, SchemaSampleBoneIdx).GetTranslation();
-			return true;
-		}
-
-		bool bOriginClamped;
-		const FVector SampleBonePosition = GetComponentSpaceTransform(SampleTime, bClamped, SchemaSampleBoneIdx).GetTranslation();
-		const FVector OriginBonePosition = GetComponentSpaceTransform(OriginTime, bOriginClamped, SchemaOriginBoneIdx).GetTranslation();
-		bClamped |= bOriginClamped;
-		OutSamplePosition = SampleBonePosition - OriginBonePosition;
-		return true;
-	}
-
 	const FTransform RootBoneTransform = GetTransform(OriginTime, bUnused, RootSchemaBoneIdx);
 	const FTransform SampleBoneTransform = GetTransform(SampleTime, bClamped, SchemaSampleBoneIdx);
 	if (SchemaOriginBoneIdx == RootSchemaBoneIdx)
diff --git a/Engine/Plugins/Experimental/Animation/PoseSearch/Source/Runtime/Private/PoseSearchContext.cpp b/Engine/Plugins/Experimental/Animation/PoseSearch/Source/Runtime/Private/PoseSearchContext.cpp
--- a/Engine/Plugins/Experimental/Animation/PoseSearch/Source/Runtime/Private/PoseSearchContext.cpp
+++ b/Engine/Plugins/Experimental/Animation/PoseSearch/Source/Runtime/Private/PoseSearchContext.cpp
@@ -429,71 +429,87 @@ FVector FSearchContext::GetSampleVelocity(float SampleTimeOffset, float OriginTi
 	return LinearVelocity;
 }
 
-FTransform FSearchContext::GetWorldRootBoneTransformAtTime(float SampleTime, bool bExtrapolate) const
+FTransform FSearchContext::GetWorldRootBoneTransformAtTime(float SampleTime) const
 {
-#if WITH_EDITOR
-	if (!Trajectory)
+	if (Trajectory)
 	{
-		UE_LOG(LogPoseSearch, Error, TEXT("FSearchContext::GetWorldRootBoneTransformAtTime - missing Trajectory is required!"));
-		return FTransform::Identity;
+		// Trajectory is already in root bone world space (transformed in UPoseSearchLibrary::ProcessTrajectory), so we just ask for a sample at the proper SampleTime
+		return Trajectory->GetSampleAtTime(SampleTime, true).GetTransform();
 	}
-#endif // WITH_EDITOR
 
-	check(Trajectory);
-	
-	// Trajectory is already in root bone world space (transformed in UPoseSearchLibrary::ProcessTrajectory), so we just ask for a sample at the proper SampleTime 
-	return Trajectory->GetSampleAtTime(SampleTime, bExtrapolate).GetTransform();
-}
-
-FTransform FSearchContext::GetWorldBoneTransformAtTime(float SampleTime, const UPoseSearchSchema* Schema, int8 SchemaBoneIdx)
-{
-	const FTransform WorldRootBoneTransform = GetWorldRootBoneTransformAtTime(SampleTime);
-
-	const FBoneIndexType BoneIndexType = Schema->GetBoneIndexType(SchemaBoneIdx);
-	if (BoneIndexType != RootBoneIndexType)
+	if (AnimInstance && AnimInstance->CurrentSkeleton)
 	{
-		const FTransform BoneTransform = GetLocalBoneTransformAtTime(SampleTime, Schema, SchemaBoneIdx);
-		return BoneTransform * WorldRootBoneTransform;
+		// @todo: perhaps pass in to FSearchContext the previous frame root bone transform as we should do in ProcessTrajectory
+		const FTransform& RootBoneTransform = AnimInstance->CurrentSkeleton->GetReferenceSkeleton().GetRefBonePose()[RootSchemaBoneIdx];
+		const FTransform& ComponentToWorldTransform = AnimInstance->GetSkelMeshComponent()->GetComponentTransform();
+
+		return RootBoneTransform * ComponentToWorldTransform;
 	}
 
-	return WorldRootBoneTransform;
+	return FTransform::Identity;
 }
 
-FTransform FSearchContext::GetLocalBoneTransformAtTime(float SampleTime, const UPoseSearchSchema* Schema, int8 SchemaBoneIdx)
+FTransform FSearchContext::GetWorldBoneTransformAtTime(float SampleTime, const UPoseSearchSchema* Schema, int8 SchemaBoneIdx)
 {
-	check(Schema);
+	const FBoneIndexType BoneIndexType = !Schema || SchemaBoneIdx == RootSchemaBoneIdx ? RootBoneIndexType : Schema->GetBoneIndexType(SchemaBoneIdx);
+	if (const FCachedTransform<FTransform>* CachedTransform = CachedTransforms.Find(SampleTime, BoneIndexType))
+	{
+		return CachedTransform->Transform;
+	}
 
-	const FBoneIndexType BoneIndexType = Schema->GetBoneIndexType(SchemaBoneIdx);
-	if (BoneIndexType != RootBoneIndexType)
+	FTransform WorldBoneTransform;
+	if (BoneIndexType == RootBoneIndexType)
+	{
+		// we already tried querying the CachedTransforms so, let's search in Trajectory
+		WorldBoneTransform = GetWorldRootBoneTransformAtTime(SampleTime);
+	}
+	else // if (BoneIndexType != RootBoneIndexType)
 	{
-		if (const FCachedTransform<FTransform>* CachedTransform = CachedTransforms.Find(SampleTime, BoneIndexType))
+		// searching for RootBoneIndexType in CachedTransforms
+		if (const FCachedTransform<FTransform>* CachedTransform = CachedTransforms.Find(SampleTime, RootBoneIndexType))
 		{
-			return CachedTransform->Transform;
+			WorldBoneTransform = CachedTransform->Transform;
 		}
-	
+		else
+		{
+			WorldBoneTransform = GetWorldRootBoneTransformAtTime(SampleTime);
+		}
+
+		check(Schema);
+
 		// collecting the local bone transforms from the IPoseHistory
-		check(History);
-		FTransform LocalBoneTransform;
-		if (!History->GetTransformAtTime(SampleTime, LocalBoneTransform, Schema->Skeleton, BoneIndexType, RootBoneIndexType))
+		#if WITH_EDITOR
+		if (!History)
 		{
-			if (const USkeleton* Skeleton = Schema->Skeleton)
+			UE_LOG(LogPoseSearch, Error, TEXT("FSearchContext::GetWorldBoneTransformAtTime - Couldn't search for bones requested by %s, because no IPoseHistory has been found!"), *Schema->GetName());
+		}
+		else
+		#endif // WITH_EDITOR
+		{
+			check(History);
+
+			FTransform LocalBoneTransform;
+			if (!History->GetTransformAtTime(SampleTime, LocalBoneTransform, Schema->Skeleton, BoneIndexType, RootBoneIndexType))
 			{
-				if (!History->IsEmpty())
+				if (const USkeleton* Skeleton = Schema->Skeleton)
+				{
+					if (!History->IsEmpty())
+					{
+						UE_LOG(LogPoseSearch, Warning, TEXT("FSearchContext::GetWorldBoneTransformAtTime - Couldn't find BoneIndexType %d (%s) requested by %s"), BoneIndexType, *Skeleton->GetReferenceSkeleton().GetBoneName(BoneIndexType).ToString(), *Schema->GetName());
+					}
+				}
+				else
 				{
-					UE_LOG(LogPoseSearch, Warning, TEXT("FSearchContext::GetLocalBoneTransformAtTime - Couldn't find BoneIndexType %d (%s) requested by %s"), BoneIndexType, *Skeleton->GetReferenceSkeleton().GetBoneName(BoneIndexType).ToString(), *Schema->GetName());
+					UE_LOG(LogPoseSearch, Warning, TEXT("FSearchContext::GetWorldBoneTransformAtTime - Schema '%s' Skeleton is not properly set"), *Schema->GetName());
 				}
 			}
-			else
-			{
-				UE_LOG(LogPoseSearch, Warning, TEXT("FSearchContext::GetLocalBoneTransformAtTime - Schema '%s' Skeleton is not properly set"), *Schema->GetName());
-			}
-		}
 
-		CachedTransforms.Add(SampleTime, BoneIndexType, LocalBoneTransform);
-		return LocalBoneTransform;
+			WorldBoneTransform = LocalBoneTransform * WorldBoneTransform;
+		}
 	}
 
-	return FTransform::Identity;
+	CachedTransforms.Add(SampleTime, BoneIndexType, WorldBoneTransform);
+	return WorldBoneTransform;
 }
 
 FVector FSearchContext::GetSamplePositionInternal(float SampleTime, float OriginTime, const UPoseSearchSchema* Schema, int8 SchemaSampleBoneIdx, int8 SchemaOriginBoneIdx, const FVector* SampleBonePositionWorldOverride)
@@ -506,23 +522,12 @@ FVector FSearchContext::GetSamplePositionInternal(float SampleTime, float Origin
 			return RootBoneTransform.InverseTransformPosition(*SampleBonePositionWorldOverride);
 		}
 
+		// @todo: validate this still works for when root bone is not Identity
 		const FTransform OriginBoneTransform = GetWorldBoneTransformAtTime(OriginTime, Schema, SchemaOriginBoneIdx);
 		const FVector DeltaBoneTranslation = *SampleBonePositionWorldOverride - OriginBoneTransform.GetTranslation();
 		return RootBoneTransform.InverseTransformVector(DeltaBoneTranslation);
 	}
 
-	if (SampleTime == OriginTime)
-	{
-		if (SchemaOriginBoneIdx == RootSchemaBoneIdx)
-		{
-			return GetLocalBoneTransformAtTime(SampleTime, Schema, SchemaSampleBoneIdx).GetTranslation();
-		}
-
-		const FVector SampleBonePosition = GetLocalBoneTransformAtTime(SampleTime, Schema, SchemaSampleBoneIdx).GetTranslation();
-		const FVector OriginBonePosition = GetLocalBoneTransformAtTime(OriginTime, Schema, SchemaOriginBoneIdx).GetTranslation();
-		return SampleBonePosition - OriginBonePosition;
-	}
-
 	const FTransform RootBoneTransform = GetWorldBoneTransformAtTime(OriginTime, Schema, RootSchemaBoneIdx);
 	const FTransform SampleBoneTransform = GetWorldBoneTransformAtTime(SampleTime, Schema, SchemaSampleBoneIdx);
 	if (SchemaOriginBoneIdx == RootSchemaBoneIdx)
@@ -550,28 +555,9 @@ FQuat FSearchContext::GetSampleRotationInternal(float SampleTime, float OriginTi
 		return RootBoneTransform.InverseTransformRotation(DeltaBoneRotation);
 	}
 
-	if (SampleTime == OriginTime)
-	{
-		if (SchemaOriginBoneIdx == RootSchemaBoneIdx)
-		{
-			return GetLocalBoneTransformAtTime(SampleTime, Schema, SchemaSampleBoneIdx).GetRotation();
-		}
-
-		const FTransform SampleBoneTransform = GetLocalBoneTransformAtTime(SampleTime, Schema, SchemaSampleBoneIdx);
-		const FTransform OriginBoneTransform = GetLocalBoneTransformAtTime(OriginTime, Schema, SchemaOriginBoneIdx);
-		return OriginBoneTransform.InverseTransformRotation(SampleBoneTransform.GetRotation());
-	}
-
 	const FTransform RootBoneTransform = GetWorldBoneTransformAtTime(OriginTime, Schema, RootSchemaBoneIdx);
 	const FTransform SampleBoneTransform = GetWorldBoneTransformAtTime(SampleTime, Schema, SchemaSampleBoneIdx);
-	if (SchemaOriginBoneIdx == RootSchemaBoneIdx)
-	{
-		return RootBoneTransform.InverseTransformRotation(SampleBoneTransform.GetRotation());
-	}
-
-	const FTransform OriginBoneTransform = GetWorldBoneTransformAtTime(OriginTime, Schema, SchemaOriginBoneIdx);
-	const FQuat DeltaBoneRotation = OriginBoneTransform.InverseTransformRotation(SampleBoneTransform.GetRotation());
-	return RootBoneTransform.InverseTransformRotation(DeltaBoneRotation);
+	return RootBoneTransform.InverseTransformRotation(SampleBoneTransform.GetRotation());
 }
 
 void FSearchContext::ClearCachedEntries()
diff --git a/Engine/Plugins/Experimental/Animation/PoseSearch/Source/Runtime/Private/PoseSearchDerivedDataKey.cpp b/Engine/Plugins/Experimental/Animation/PoseSearch/Source/Runtime/Private/PoseSearchDerivedDataKey.cpp
--- a/Engine/Plugins/Experimental/Animation/PoseSearch/Source/Runtime/Private/PoseSearchDerivedDataKey.cpp
+++ b/Engine/Plugins/Experimental/Animation/PoseSearch/Source/Runtime/Private/PoseSearchDerivedDataKey.cpp
@@ -40,7 +40,7 @@ FKeyBuilder::FKeyBuilder(const UObject* Object, bool bUseDataVer, bool bPerformC
 	if (bUseDataVer)
 	{
 		// used to invalidate the key without having to change POSESEARCHDB_DERIVEDDATA_VER all the times
-		int32 POSESEARCHDB_DERIVEDDATA_VER_SMALL = 229;
+		int32 POSESEARCHDB_DERIVEDDATA_VER_SMALL = 231;
 		FGuid VersionGuid = FDevSystemGuids::GetSystemGuid(FDevSystemGuids::Get().POSESEARCHDB_DERIVEDDATA_VER);
 
 		*this << VersionGuid;
diff --git a/Engine/Plugins/Experimental/Animation/PoseSearch/Source/Runtime/Private/PoseSearchFeatureChannel_Position.cpp b/Engine/Plugins/Experimental/Animation/PoseSearch/Source/Runtime/Private/PoseSearchFeatureChannel_Position.cpp
--- a/Engine/Plugins/Experimental/Animation/PoseSearch/Source/Runtime/Private/PoseSearchFeatureChannel_Position.cpp
+++ b/Engine/Plugins/Experimental/Animation/PoseSearch/Source/Runtime/Private/PoseSearchFeatureChannel_Position.cpp
@@ -114,7 +114,7 @@ void UPoseSearchFeatureChannel_Position::DebugDraw(const UE::PoseSearch::FDebugD
 #endif // WITH_EDITORONLY_DATA
 
 	const FVector FeaturesVector = FFeatureVectorHelper::DecodeVector(PoseVector, ChannelDataOffset, ComponentStripping);
-	if (SchemaOriginBoneIdx == RootSchemaBoneIdx)
+	if (SchemaOriginBoneIdx == RootSchemaBoneIdx && FMath::IsNearlyZero(OriginTimeOffset))
 	{
 		const FVector BonePos = DrawParams.GetRootTransform().TransformPosition(FeaturesVector);
 		DrawParams.DrawPoint(BonePos, Color);
@@ -122,7 +122,7 @@ void UPoseSearchFeatureChannel_Position::DebugDraw(const UE::PoseSearch::FDebugD
 	else
 	{
 		const EPermutationTimeType TimeType = PermutationTimeType == EPermutationTimeType::UsePermutationTime ? EPermutationTimeType::UseSampleToPermutationTime : EPermutationTimeType::UseSampleTime;
-		const FVector OriginBonePos = DrawParams.ExtractPosition(PoseVector, SampleTimeOffset, SchemaOriginBoneIdx, TimeType, SamplingAttributeId);
+		const FVector OriginBonePos = DrawParams.ExtractPosition(PoseVector, OriginTimeOffset, SchemaOriginBoneIdx, TimeType);
 		const FVector DeltaPos = DrawParams.GetRootTransform().TransformVector(FeaturesVector);
 		const FVector BonePos = OriginBonePos + DeltaPos;
 		DrawParams.DrawLine(OriginBonePos, BonePos, Color);
diff --git a/Engine/Plugins/Experimental/Animation/PoseSearch/Source/Runtime/Private/PoseSearchLibrary.cpp b/Engine/Plugins/Experimental/Animation/PoseSearch/Source/Runtime/Private/PoseSearchLibrary.cpp
--- a/Engine/Plugins/Experimental/Animation/PoseSearch/Source/Runtime/Private/PoseSearchLibrary.cpp
+++ b/Engine/Plugins/Experimental/Animation/PoseSearch/Source/Runtime/Private/PoseSearchLibrary.cpp
@@ -150,7 +150,9 @@ void FMotionMatchingState::UpdateRootBoneControl(const FAnimationUpdateContext&
 
 #if UE_POSE_SEARCH_TRACE_ENABLED
 void UPoseSearchLibrary::TraceMotionMatchingState(
-	const FPoseSearchQueryTrajectory& Trajectory,
+	// Trajectory is the trajectory in mesh component space prior ProcessTrajectory
+	const FPoseSearchQueryTrajectory& Trajectory, 
+	// SearchContext.Trajectory is the trajectory in root bone component space after ProcessTrajectory
 	UE::PoseSearch::FSearchContext& SearchContext,
 	const UE::PoseSearch::FSearchResult& CurrentResult,
 	float ElapsedPoseSearchTime,
@@ -210,8 +212,8 @@ void UPoseSearchLibrary::TraceMotionMatchingState(
 		// simulation
 		if (SearchContext.IsTrajectoryValid())
 		{
-			const FTransform PrevRoot = SearchContext.GetWorldRootBoneTransformAtTime(-DeltaTime);
-			const FTransform CurrRoot = SearchContext.GetWorldRootBoneTransformAtTime(0.f);
+			const FTransform PrevRoot = SearchContext.GetWorldBoneTransformAtTime(-DeltaTime);
+			const FTransform CurrRoot = SearchContext.GetWorldBoneTransformAtTime(0.f);
 			const FTransform SimDelta = CurrRoot.GetRelativeTransform(PrevRoot);
 
 			TraceState.SimLinearVelocity = SimDelta.GetTranslation().Size() / DeltaTime;
@@ -291,7 +293,9 @@ void UPoseSearchLibrary::UpdateMotionMatchingState(
 #endif // ENABLE_DRAW_DEBUG && ENABLE_ANIM_DEBUG
 	}
 
-	const FPoseSearchQueryTrajectory TrajectoryRootSpace = ProcessTrajectory(Trajectory, InOutMotionMatchingState.ComponentDeltaYaw, YawFromAnimationTrajectoryBlendTime, TrajectorySpeedMultiplier);
+	// @todo: perhaps store the previous frame root bone transform to use in ProcessTrajectory
+	const FTransform& RootBoneTransform = Context.AnimInstanceProxy->GetSkeleton()->GetReferenceSkeleton().GetRefBonePose()[RootSchemaBoneIdx];
+	const FPoseSearchQueryTrajectory TrajectoryRootSpace = ProcessTrajectory(Trajectory, RootBoneTransform, InOutMotionMatchingState.ComponentDeltaYaw, YawFromAnimationTrajectoryBlendTime, TrajectorySpeedMultiplier);
 
 	FMemMark Mark(FMemStack::Get());
 	const UAnimInstance* AnimInstance = Cast<const UAnimInstance>(Context.AnimInstanceProxy->GetAnimInstanceObject());
@@ -383,19 +387,16 @@ void UPoseSearchLibrary::UpdateMotionMatchingState(
 		else
 #endif // WITH_EDITOR
 		{
-			const FRotator DeltaRotation(0.f, InOutMotionMatchingState.ComponentDeltaYaw, 0.f);
-			const FTransform DeltaTransform(DeltaRotation);
-			const FTransform RootBoneTransform = DeltaTransform * Context.AnimInstanceProxy->GetComponentTransform();
-
 			if (bDebugDrawCurResult)
 			{
-				UE::PoseSearch::FDebugDrawParams DrawParams(Context.AnimInstanceProxy, RootBoneTransform, CurResultDatabase);
+				UE::PoseSearch::FDebugDrawParams DrawParams(Context.AnimInstanceProxy, SearchContext.GetWorldBoneTransformAtTime(0.f), CurResultDatabase);
 				DrawParams.DrawFeatureVector(CurResult.PoseIdx);
 			}
 
 			if (bDebugDrawQuery)
 			{
-				UE::PoseSearch::FDebugDrawParams DrawParams(Context.AnimInstanceProxy, RootBoneTransform, CurResultDatabase, EDebugDrawFlags::DrawQuery);
+				// @todo: use pose history to get the root bone!
+				UE::PoseSearch::FDebugDrawParams DrawParams(Context.AnimInstanceProxy, SearchContext.GetWorldBoneTransformAtTime(0.f), CurResultDatabase, EDebugDrawFlags::DrawQuery);
 				DrawParams.DrawFeatureVector(SearchContext.GetOrBuildQuery(CurResultDatabase->Schema).GetValues());
 			}
 		}
@@ -404,18 +405,34 @@ void UPoseSearchLibrary::UpdateMotionMatchingState(
 }
 
 // transforms Trajectory from SkeletalMeshComponent world space to root bone world space, and scale it by TrajectorySpeedMultiplier
-FPoseSearchQueryTrajectory UPoseSearchLibrary::ProcessTrajectory(const FPoseSearchQueryTrajectory& Trajectory, float RootBoneDeltaYaw, float YawFromAnimationTrajectoryBlendTime, float TrajectorySpeedMultiplier)
+FPoseSearchQueryTrajectory UPoseSearchLibrary::ProcessTrajectory(const FPoseSearchQueryTrajectory& Trajectory, const FTransform& RootBoneTransform, float RootBoneDeltaYaw, float YawFromAnimationTrajectoryBlendTime, float TrajectorySpeedMultiplier)
 {
 	QUICK_SCOPE_CYCLE_COUNTER(STAT_PoseSearch_ProcessTrajectory);
 
 	const float TrajectorySpeedMultiplierInv = FMath::IsNearlyZero(TrajectorySpeedMultiplier) ? 1.f : 1.f / TrajectorySpeedMultiplier;
 
 	FPoseSearchQueryTrajectory TrajectoryRootSpace = Trajectory;
-	for (FPoseSearchQueryTrajectorySample& Sample : TrajectoryRootSpace.Samples)
+
+	if (!RootBoneTransform.Equals(FTransform::Identity))
+	{
+		// @todo: optimize me
+		for (FPoseSearchQueryTrajectorySample& Sample : TrajectoryRootSpace.Samples)
+		{
+			Sample.SetTransform(RootBoneTransform * Sample.GetTransform());
+		}
+	}
+
+	if (!FMath::IsNearlyEqual(TrajectorySpeedMultiplierInv, 1.f))
 	{
-		Sample.AccumulatedSeconds *= TrajectorySpeedMultiplierInv;
+		for (FPoseSearchQueryTrajectorySample& Sample : TrajectoryRootSpace.Samples)
+		{
+			Sample.AccumulatedSeconds *= TrajectorySpeedMultiplierInv;
+		}
+	}
 
-		if (!FMath::IsNearlyZero(RootBoneDeltaYaw))
+	if (!FMath::IsNearlyZero(RootBoneDeltaYaw))
+	{
+		for (FPoseSearchQueryTrajectorySample& Sample : TrajectoryRootSpace.Samples)
 		{
 			const float BlendParam = YawFromAnimationTrajectoryBlendTime < UE_KINDA_SMALL_NUMBER ? 1 : FMath::Clamp(1.f - (Sample.AccumulatedSeconds - YawFromAnimationTrajectoryBlendTime) / YawFromAnimationTrajectoryBlendTime, 0.f, 1.f);
 			const FQuat RootBoneDelta(FRotator(0.f, RootBoneDeltaYaw * BlendParam, 0.f));
@@ -465,13 +482,15 @@ void UPoseSearchLibrary::MotionMatch(
 	Result.SearchCost = MAX_flt;
 
 	FMemMark Mark(FMemStack::Get());
-	if (Database && AnimInstance)
+	if (Database && AnimInstance && AnimInstance->CurrentSkeleton)
 	{
 #if ENABLE_DRAW_DEBUG && ENABLE_ANIM_DEBUG
 		FColor HistoryCollectorColor = FColor::Red;
 #endif // ENABLE_DRAW_DEBUG && ENABLE_ANIM_DEBUG
 
-		const FPoseSearchQueryTrajectory TrajectoryRootSpace = ProcessTrajectory(Trajectory, 0.f, 0.f, TrajectorySpeedMultiplier);
+		// @todo: perhaps store the previous frame root bone transform to use in ProcessTrajectory
+		const FTransform& RootBoneTransform = AnimInstance->CurrentSkeleton->GetReferenceSkeleton().GetRefBonePose()[RootSchemaBoneIdx];
+		const FPoseSearchQueryTrajectory TrajectoryRootSpace = ProcessTrajectory(Trajectory, RootBoneTransform, 0.f, 0.f, TrajectorySpeedMultiplier);
 
 		// ExtendedPoseHistory will hold future poses to match AssetSamplerBase (at FutureAnimationStartTime) TimeToFutureAnimationStart seconds in the future
 		FExtendedPoseHistory ExtendedPoseHistory;
@@ -574,13 +593,13 @@ void UPoseSearchLibrary::MotionMatch(
 			FAnimInstanceProxy* AnimInstanceProxy = UAnimInstanceProxyProvider::GetAnimInstanceProxy(AnimInstance);
 			if (CVarAnimMotionMatchDrawMatchEnable.GetValueOnAnyThread())
 			{
-				FDebugDrawParams DrawParams(AnimInstanceProxy, AnimInstanceProxy->GetComponentTransform(), SearchResult.Database.Get());
+				FDebugDrawParams DrawParams(AnimInstanceProxy, SearchContext.GetWorldBoneTransformAtTime(0.f), SearchResult.Database.Get());
 				DrawParams.DrawFeatureVector(SearchResult.PoseIdx);
 			}
 
 			if (CVarAnimMotionMatchDrawQueryEnable.GetValueOnAnyThread())
 			{
-				FDebugDrawParams DrawParams(AnimInstanceProxy, AnimInstanceProxy->GetComponentTransform(), SearchResult.Database.Get(), EDebugDrawFlags::DrawQuery);
+				FDebugDrawParams DrawParams(AnimInstanceProxy, SearchContext.GetWorldBoneTransformAtTime(0.f), SearchResult.Database.Get(), EDebugDrawFlags::DrawQuery);
 				DrawParams.DrawFeatureVector(SearchContext.GetOrBuildQuery(SearchResult.Database->Schema).GetValues());
 			}
 		}
@@ -664,13 +683,13 @@ UE::PoseSearch::FSearchResult UPoseSearchLibrary::MotionMatch(const FAnimationBa
 		{
 			if (CVarAnimMotionMatchDrawMatchEnable.GetValueOnAnyThread())
 			{
-				FDebugDrawParams DrawParams(Context.AnimInstanceProxy, Context.AnimInstanceProxy->GetComponentTransform(), SearchResult.Database.Get());
+				FDebugDrawParams DrawParams(Context.AnimInstanceProxy, SearchContext.GetWorldBoneTransformAtTime(0.f), SearchResult.Database.Get());
 				DrawParams.DrawFeatureVector(SearchResult.PoseIdx);
 			}
 
 			if (CVarAnimMotionMatchDrawQueryEnable.GetValueOnAnyThread())
 			{
-				FDebugDrawParams DrawParams(Context.AnimInstanceProxy, Context.AnimInstanceProxy->GetComponentTransform(), SearchResult.Database.Get(), EDebugDrawFlags::DrawQuery);
+				FDebugDrawParams DrawParams(Context.AnimInstanceProxy, SearchContext.GetWorldBoneTransformAtTime(0.f), SearchResult.Database.Get(), EDebugDrawFlags::DrawQuery);
 				DrawParams.DrawFeatureVector(SearchContext.GetOrBuildQuery(SearchResult.Database->Schema).GetValues());
 			}
 		}
diff --git a/Engine/Plugins/Experimental/Animation/PoseSearch/Source/Runtime/Public/PoseSearch/PoseSearchContext.h b/Engine/Plugins/Experimental/Animation/PoseSearch/Source/Runtime/Public/PoseSearch/PoseSearchContext.h
--- a/Engine/Plugins/Experimental/Animation/PoseSearch/Source/Runtime/Public/PoseSearch/PoseSearchContext.h
+++ b/Engine/Plugins/Experimental/Animation/PoseSearch/Source/Runtime/Public/PoseSearch/PoseSearchContext.h
@@ -183,21 +183,18 @@ struct POSESEARCH_API FSearchContext
 	bool IsTrajectoryValid() const { return Trajectory != nullptr; }
 	float GetDesiredPermutationTimeOffset() const { return DesiredPermutationTimeOffset; }
 	bool IsForceInterrupt() const { return bForceInterrupt; }
-	FTransform GetWorldRootBoneTransformAtTime(float SampleTime, bool bExtrapolate = true) const;
 	const UAnimInstance* GetAnimInstance() const { return AnimInstance; }
 
 	void SetAnimationsToConsider(TConstArrayView<const UAnimationAsset*> InAnimationsToConsider) { AnimationsToConsider = InAnimationsToConsider; }
 	TConstArrayView<const UAnimationAsset*> GetAnimationsToConsider() const { return AnimationsToConsider; }
 	
-private:
 	// returns the world space transform of the bone SchemaBoneIdx at time SampleTime
-	FTransform GetWorldBoneTransformAtTime(float SampleTime, const UPoseSearchSchema* Schema, int8 SchemaBoneIdx = RootSchemaBoneIdx);
-	
-	// returns the local space transform relative to the root bone of the bone SchemaBoneIdx at time SampleTime
-	FTransform GetLocalBoneTransformAtTime(float SampleTime, const UPoseSearchSchema* Schema, int8 SchemaSampleBoneIdx = RootSchemaBoneIdx);
+	FTransform GetWorldBoneTransformAtTime(float SampleTime, const UPoseSearchSchema* Schema = nullptr, int8 SchemaBoneIdx = RootSchemaBoneIdx);
 	
+private:
 	FVector GetSamplePositionInternal(float SampleTime, float OriginTime, const UPoseSearchSchema* Schema, int8 SchemaSampleBoneIdx = RootSchemaBoneIdx, int8 SchemaOriginBoneIdx = RootSchemaBoneIdx, const FVector* SampleBonePositionWorldOverride = nullptr);
 	FQuat GetSampleRotationInternal(float SampleTime, float OriginTime, const UPoseSearchSchema* Schema, int8 SchemaSampleBoneIdx = RootSchemaBoneIdx, int8 SchemaOriginBoneIdx = RootSchemaBoneIdx, const FQuat* SampleBoneRotationWorldOverride = nullptr);
+	FTransform GetWorldRootBoneTransformAtTime(float SampleTime) const;
 
 	const UAnimInstance* AnimInstance = nullptr;
 	const IPoseHistory* History = nullptr;
@@ -216,7 +213,7 @@ private:
 	TConstArrayView<float> CurrentResultPoseVector;
 	TStackAlignedArray<float> CurrentResultPoseVectorData;
 
-	// transforms cached in component space
+	// transforms cached in world space
 	FCachedTransforms<FTransform> CachedTransforms;
 	TArray<FFeatureVectorBuilder, TInlineAllocator<PreallocatedCachedQueriesNum>> CachedQueries;
 
diff --git a/Engine/Plugins/Experimental/Animation/PoseSearch/Source/Runtime/Public/PoseSearch/PoseSearchLibrary.h b/Engine/Plugins/Experimental/Animation/PoseSearch/Source/Runtime/Public/PoseSearch/PoseSearchLibrary.h
--- a/Engine/Plugins/Experimental/Animation/PoseSearch/Source/Runtime/Public/PoseSearch/PoseSearchLibrary.h
+++ b/Engine/Plugins/Experimental/Animation/PoseSearch/Source/Runtime/Public/PoseSearch/PoseSearchLibrary.h
@@ -90,6 +90,7 @@ class POSESEARCH_API UPoseSearchLibrary : public UBlueprintFunctionLibrary
 
 	static FPoseSearchQueryTrajectory ProcessTrajectory(
 		const FPoseSearchQueryTrajectory& Trajectory,
+		const FTransform& RootBoneTransform,
 		float RootBoneDeltaYaw,
 		float YawFromAnimationTrajectoryBlendTime,
 		float TrajectorySpeedMultiplier);
