motion matching - fix for database editor player buttons

diff --git a/Engine/Plugins/Experimental/Animation/PoseSearch/Source/Editor/Private/PoseSearchDatabaseViewModel.cpp b/Engine/Plugins/Experimental/Animation/PoseSearch/Source/Editor/Private/PoseSearchDatabaseViewModel.cpp
--- a/Engine/Plugins/Experimental/Animation/PoseSearch/Source/Editor/Private/PoseSearchDatabaseViewModel.cpp
+++ b/Engine/Plugins/Experimental/Animation/PoseSearch/Source/Editor/Private/PoseSearchDatabaseViewModel.cpp
@@ -347,39 +347,39 @@ void FDatabaseViewModel::BuildSearchIndex()
 
 void FDatabaseViewModel::PreviewBackwardEnd()
 {
-	PlayTime = MinPreviewPlayLength;
+	SetPlayTime(MinPreviewPlayLength, false);
 }
 
 void FDatabaseViewModel::PreviewBackwardStep()
 {
-	PlayTime = FMath::Clamp(PlayTime - StepDeltaTime, MinPreviewPlayLength, MaxPreviewPlayLength);
-	DeltaTimeMultiplier = 0.0f;
+	const float NewPlayTime = FMath::Clamp(PlayTime - StepDeltaTime, MinPreviewPlayLength, MaxPreviewPlayLength);
+	SetPlayTime(NewPlayTime, false);
 }
 
 void FDatabaseViewModel::PreviewBackward()
 {
-	DeltaTimeMultiplier = -1.0f;
+	DeltaTimeMultiplier = -1.f;
 }
 
 void FDatabaseViewModel::PreviewPause()
 {
-	DeltaTimeMultiplier = 0.0f;
+	DeltaTimeMultiplier = 0.f;
 }
 
 void FDatabaseViewModel::PreviewForward()
 {
-	DeltaTimeMultiplier = 1.0f;
+	DeltaTimeMultiplier = 1.f;
 }
 
 void FDatabaseViewModel::PreviewForwardStep()
 {
-	PlayTime = FMath::Clamp(PlayTime + StepDeltaTime, MinPreviewPlayLength, MaxPreviewPlayLength);
-	DeltaTimeMultiplier = 0.0f;
+	const float NewPlayTime = FMath::Clamp(PlayTime + StepDeltaTime, MinPreviewPlayLength, MaxPreviewPlayLength);
+	SetPlayTime(NewPlayTime, false);
 }
 
 void FDatabaseViewModel::PreviewForwardEnd()
 {
-	PlayTime = MaxPreviewPlayLength;
+	SetPlayTime(MaxPreviewPlayLength, false);
 }
 
 UWorld* FDatabaseViewModel::GetWorld()
