motion matching - cook size estimator support for pose search database WIP

diff --git a/Engine/Plugins/Experimental/Animation/PoseSearch/Source/Editor/Private/PoseSearchDatabaseEditorReflection.cpp b/Engine/Plugins/Experimental/Animation/PoseSearch/Source/Editor/Private/PoseSearchDatabaseEditorReflection.cpp
--- a/Engine/Plugins/Experimental/Animation/PoseSearch/Source/Editor/Private/PoseSearchDatabaseEditorReflection.cpp
+++ b/Engine/Plugins/Experimental/Animation/PoseSearch/Source/Editor/Private/PoseSearchDatabaseEditorReflection.cpp
@@ -168,6 +168,22 @@ void UPoseSearchDatabaseStatistics::Initialize(const UPoseSearchDatabase* PoseSe
 		// Memory Information
 			
 		{
+			uint32 SourceAnimAssetsSizeCookedEstimateInBytes = 0;
+			TSet<const UAnimationAsset*> Analyzed;
+			Analyzed.Reserve(PoseSearchDatabase->AnimationAssets.Num());
+			for (const FInstancedStruct& AnimAsset : PoseSearchDatabase->AnimationAssets)
+			{
+				if (const FPoseSearchDatabaseAnimationAssetBase* DatabaseAnimationAssetBase = AnimAsset.GetPtr<FPoseSearchDatabaseAnimationAssetBase>())
+				{
+					bool bAlreadyAnalyzed = false;
+					Analyzed.Add(DatabaseAnimationAssetBase->GetAnimationAsset(), &bAlreadyAnalyzed);
+					if (!bAlreadyAnalyzed && DatabaseAnimationAssetBase->GetAnimationAsset())
+					{
+						SourceAnimAssetsSizeCookedEstimateInBytes += DatabaseAnimationAssetBase->GetApproxCookedSize();
+					}
+				}
+			}
+
 			const uint32 ValuesBytesSize = SearchIndex.Values.GetAllocatedSize();
 			const uint32 PCAValuesBytesSize = SearchIndex.PCAValues.GetAllocatedSize();
 			const uint32 KDTreeBytesSize = SearchIndex.KDTree.GetAllocatedSize();
@@ -178,7 +194,7 @@ void UPoseSearchDatabaseStatistics::Initialize(const UPoseSearchDatabase* PoseSe
 			const uint32 PoseMetadataBytesSize = SearchIndex.PoseMetadata.GetAllocatedSize();
 			const uint32 AssetsBytesSize = SearchIndex.Assets.GetAllocatedSize();
 			const uint32 OtherBytesSize = SearchIndex.PCAProjectionMatrix.GetAllocatedSize() + SearchIndex.Mean.GetAllocatedSize() + SearchIndex.WeightsSqrt.GetAllocatedSize();
-			const uint32 EstimatedDatabaseBytesSize = ValuesBytesSize + PCAValuesBytesSize + KDTreeBytesSize + VPTreeBytesSize + ValuesVectorToPoseIndexesBytesSize + PCAValuesVectorToPoseIndexesBytesSize + PoseMetadataBytesSize + AssetsBytesSize + OtherBytesSize;
+			const uint32 EstimatedDatabaseBytesSize = ValuesBytesSize + PCAValuesBytesSize + KDTreeBytesSize + VPTreeBytesSize + ValuesVectorToPoseIndexesBytesSize + PCAValuesVectorToPoseIndexesBytesSize + PoseMetadataBytesSize + AssetsBytesSize + OtherBytesSize + SourceAnimAssetsSizeCookedEstimateInBytes;
 				
 			ValuesSize = FText::AsMemory(ValuesBytesSize);
 			PCAValuesSize = FText::AsMemory(PCAValuesBytesSize);
@@ -187,6 +203,7 @@ void UPoseSearchDatabaseStatistics::Initialize(const UPoseSearchDatabase* PoseSe
 			PoseMetadataSize = FText::AsMemory(PoseMetadataBytesSize);
 			AssetsSize = FText::AsMemory(AssetsBytesSize);
 			EstimatedDatabaseSize = FText::AsMemory(EstimatedDatabaseBytesSize);
+			SourceAnimAssetsSizeCookedEstimate = FText::AsMemory(SourceAnimAssetsSizeCookedEstimateInBytes);
 		}
 	}
 }
diff --git a/Engine/Plugins/Experimental/Animation/PoseSearch/Source/Editor/Public/PoseSearchDatabaseEditorReflection.h b/Engine/Plugins/Experimental/Animation/PoseSearch/Source/Editor/Public/PoseSearchDatabaseEditorReflection.h
--- a/Engine/Plugins/Experimental/Animation/PoseSearch/Source/Editor/Public/PoseSearchDatabaseEditorReflection.h
+++ b/Engine/Plugins/Experimental/Animation/PoseSearch/Source/Editor/Public/PoseSearchDatabaseEditorReflection.h
@@ -231,7 +231,7 @@ public:
 	UPROPERTY(VisibleAnywhere, Category = "Principal Component Analysis Information", meta = (Units = "Percent"))
 	float ExplainedVariance;
 
-	// aggregated total memory used by this database.
+	// aggregated total memory used by this database (including animation assets).
 	UPROPERTY(VisibleAnywhere, Category = "Memory Information")
 	FText EstimatedDatabaseSize;
 
@@ -259,6 +259,10 @@ public:
 	UPROPERTY(VisibleAnywhere, Category = "Memory Information")
 	FText AssetsSize;
 
+	// partial memory size used by the cooked animation assets.
+	UPROPERTY(VisibleAnywhere, Category = "Memory Information")
+ 	FText SourceAnimAssetsSizeCookedEstimate;
+
 	/** Initialize statistics given a database */
 	void Initialize(const UPoseSearchDatabase* PoseSearchDatabase);
 };
\ No newline at end of file
diff --git a/Engine/Plugins/Experimental/Animation/PoseSearch/Source/Runtime/Private/PoseSearchDatabase.cpp b/Engine/Plugins/Experimental/Animation/PoseSearch/Source/Runtime/Private/PoseSearchDatabase.cpp
--- a/Engine/Plugins/Experimental/Animation/PoseSearch/Source/Runtime/Private/PoseSearchDatabase.cpp
+++ b/Engine/Plugins/Experimental/Animation/PoseSearch/Source/Runtime/Private/PoseSearchDatabase.cpp
@@ -14,6 +14,7 @@
 #include "PoseSearch/PoseSearchFeatureChannel.h"
 #include "PoseSearch/PoseSearchHistory.h"
 #include "PoseSearch/PoseSearchSchema.h"
+#include "Serialization/ArchiveCountMem.h"
 #include "UObject/ObjectSaveContext.h"
 
 #if WITH_EDITOR
@@ -305,11 +306,21 @@ private:
 
 } // namespace UE::PoseSearch
 
+//////////////////////////////////////////////////////////////////////////
+// FPoseSearchDatabaseAnimationAssetBase
+#if WITH_EDITORONLY_DATA
+int64 FPoseSearchDatabaseAnimationAssetBase::GetEditorMemSize() const
+{
+	FArchiveCountMem EditorMemCount(GetAnimationAsset());
+	return EditorMemCount.GetNum();
+}
+#endif // WITH_EDITORONLY_DATA
+
 //////////////////////////////////////////////////////////////////////////
 // FPoseSearchDatabaseSequence
 UAnimationAsset* FPoseSearchDatabaseSequence::GetAnimationAsset() const
 {
-	return Sequence;
+	return Sequence.Get();
 }
 
 #if WITH_EDITORONLY_DATA
@@ -453,7 +464,7 @@ FVector FPoseSearchDatabaseBlendSpace::BlendParameterForSampleRanges(int32 Horiz
 // FPoseSearchDatabaseAnimComposite
 UAnimationAsset* FPoseSearchDatabaseAnimComposite::GetAnimationAsset() const
 {
-	return AnimComposite;
+	return AnimComposite.Get();
 }
 
 #if WITH_EDITORONLY_DATA
@@ -485,7 +496,7 @@ bool FPoseSearchDatabaseAnimComposite::IsRootMotionEnabled() const
 // FPoseSearchDatabaseAnimMontage
 UAnimationAsset* FPoseSearchDatabaseAnimMontage::GetAnimationAsset() const
 {
-	return AnimMontage;
+	return AnimMontage.Get();
 }
 
 #if WITH_EDITORONLY_DATA
diff --git a/Engine/Plugins/Experimental/Animation/PoseSearch/Source/Runtime/Public/PoseSearch/PoseSearchDatabase.h b/Engine/Plugins/Experimental/Animation/PoseSearch/Source/Runtime/Public/PoseSearch/PoseSearchDatabase.h
--- a/Engine/Plugins/Experimental/Animation/PoseSearch/Source/Runtime/Public/PoseSearch/PoseSearchDatabase.h
+++ b/Engine/Plugins/Experimental/Animation/PoseSearch/Source/Runtime/Public/PoseSearch/PoseSearchDatabase.h
@@ -91,6 +91,9 @@ struct POSESEARCH_API FPoseSearchDatabaseAnimationAssetBase
 	// [0, 0] represents the entire frame range of the original animation.
 	virtual FFloatInterval GetSamplingRange() const { return FFloatInterval(0.f, 0.f); }
 
+	virtual int64 GetEditorMemSize() const;
+	virtual int64 GetApproxCookedSize() const { return GetEditorMemSize(); }
+
 	// This allows users to enable or exclude animations from this database. Useful for debugging.
 	UPROPERTY(EditAnywhere, Category = "Settings", meta = (DisplayPriority = 1))
 	bool bEnabled = true;
