motion matching - improved GC compliance

diff --git a/Engine/Plugins/Experimental/Animation/PoseSearch/Source/Editor/Private/PoseSearchDatabaseAssetListItem.cpp b/Engine/Plugins/Experimental/Animation/PoseSearch/Source/Editor/Private/PoseSearchDatabaseAssetListItem.cpp
--- a/Engine/Plugins/Experimental/Animation/PoseSearch/Source/Editor/Private/PoseSearchDatabaseAssetListItem.cpp
+++ b/Engine/Plugins/Experimental/Animation/PoseSearch/Source/Editor/Private/PoseSearchDatabaseAssetListItem.cpp
@@ -153,11 +153,11 @@ namespace UE::PoseSearch
 
 	FText SDatabaseAssetListItem::GetName() const
 	{
-		TSharedPtr<FDatabaseAssetTreeNode> Node = WeakAssetTreeNode.Pin();
 		TSharedPtr<FDatabaseViewModel> ViewModel = EditorViewModel.Pin();
 
 		if (const UPoseSearchDatabase* Database = ViewModel->GetPoseSearchDatabase())
 		{
+			TSharedPtr<FDatabaseAssetTreeNode> Node = WeakAssetTreeNode.Pin();
 			if (const FPoseSearchDatabaseAnimationAssetBase* DatabaseAnimationAsset = Database->GetAnimationAssetBase(Node->SourceAssetIdx))
 			{
 				return FText::FromString(DatabaseAnimationAsset->GetName());
@@ -172,8 +172,6 @@ namespace UE::PoseSearch
 	TSharedRef<SWidget> SDatabaseAssetListItem::GenerateItemWidget()
 	{
 		TSharedPtr<FDatabaseAssetTreeNode> Node = WeakAssetTreeNode.Pin();
-		TSharedPtr<FDatabaseViewModel> ViewModel = EditorViewModel.Pin();
-		UPoseSearchDatabase* Database = ViewModel->GetPoseSearchDatabase();
 		
 		TSharedPtr<SWidget> ItemWidget;
 		const FDetailColumnSizeData& ColumnSizeData = SkeletonView.Pin()->GetColumnSizeData();
@@ -209,10 +207,14 @@ namespace UE::PoseSearch
 		{
 			// Item Icon
 			TSharedPtr<SImage> ItemIconWidget;
-			if (const FPoseSearchDatabaseAnimationAssetBase* DatabaseAnimationAsset = Database->GetAnimationAssetBase(Node->SourceAssetIdx))
+			TSharedPtr<FDatabaseViewModel> ViewModel = EditorViewModel.Pin();
+			if (UPoseSearchDatabase* Database = ViewModel->GetPoseSearchDatabase())
 			{
-				SAssignNew(ItemIconWidget, SImage)
-					.Image(FSlateIconFinder::FindIconBrushForClass(DatabaseAnimationAsset->GetAnimationAssetStaticClass()));
+				if (const FPoseSearchDatabaseAnimationAssetBase* DatabaseAnimationAsset = Database->GetAnimationAssetBase(Node->SourceAssetIdx))
+				{
+					SAssignNew(ItemIconWidget, SImage)
+						.Image(FSlateIconFinder::FindIconBrushForClass(DatabaseAnimationAsset->GetAnimationAssetStaticClass()));
+				}
 			}
 
 			// Setup table row to display 
@@ -405,14 +407,15 @@ namespace UE::PoseSearch
 	ECheckBoxState SDatabaseAssetListItem::GetDisableReselectionChecked() const
 	{
 		TSharedPtr<FDatabaseViewModel> ViewModelPtr = EditorViewModel.Pin();
-		TSharedPtr<FDatabaseAssetTreeNode> TreeNodePtr = WeakAssetTreeNode.Pin();
-		const UPoseSearchDatabase* Database = ViewModelPtr->GetPoseSearchDatabase();
-
-		if (Database->AnimationAssets.IsValidIndex(TreeNodePtr->SourceAssetIdx))
+		if (const UPoseSearchDatabase* Database = ViewModelPtr->GetPoseSearchDatabase())
 		{
-			if (ViewModelPtr->IsDisableReselection(TreeNodePtr->SourceAssetIdx))
+			TSharedPtr<FDatabaseAssetTreeNode> TreeNodePtr = WeakAssetTreeNode.Pin();
+			if (Database->AnimationAssets.IsValidIndex(TreeNodePtr->SourceAssetIdx))
 			{
-				return ECheckBoxState::Checked;
+				if (ViewModelPtr->IsDisableReselection(TreeNodePtr->SourceAssetIdx))
+				{
+					return ECheckBoxState::Checked;
+				}
 			}
 		}
 
@@ -421,34 +424,37 @@ namespace UE::PoseSearch
 
 	void SDatabaseAssetListItem::OnDisableReselectionChanged(ECheckBoxState NewCheckboxState)
 	{
-		const FScopedTransaction Transaction(LOCTEXT("EnableChangedForAssetInPoseSearchDatabase", "Update enabled flag for item from Pose Search Database"));
-
 		const TSharedPtr<FDatabaseViewModel> ViewModelPtr = EditorViewModel.Pin();
-		const TSharedPtr<FDatabaseAssetTreeNode> TreeNodePtr = WeakAssetTreeNode.Pin();
+		if (UPoseSearchDatabase* PoseSearchDatabase = ViewModelPtr->GetPoseSearchDatabase())
+		{
+			const FScopedTransaction Transaction(LOCTEXT("EnableChangedForAssetInPoseSearchDatabase", "Update enabled flag for item from Pose Search Database"));
 
-		ViewModelPtr->GetPoseSearchDatabase()->Modify();
-		
-		ViewModelPtr->SetDisableReselection(TreeNodePtr->SourceAssetIdx, NewCheckboxState == ECheckBoxState::Checked ? true : false);
+			const TSharedPtr<FDatabaseAssetTreeNode> TreeNodePtr = WeakAssetTreeNode.Pin();
 
-		SkeletonView.Pin()->RefreshTreeView(false, true);
-		
-		// no need to rebuild the SearchIndex (ViewModelPtr->BuildSearchIndex()), since bDisableReselection is a runtime only parameter
+			PoseSearchDatabase->Modify();
+
+			ViewModelPtr->SetDisableReselection(TreeNodePtr->SourceAssetIdx, NewCheckboxState == ECheckBoxState::Checked ? true : false);
+
+			SkeletonView.Pin()->RefreshTreeView(false, true);
+
+			// no need to rebuild the SearchIndex (ViewModelPtr->BuildSearchIndex()), since bDisableReselection is a runtime only parameter
+		}
 	}
 
 	ECheckBoxState SDatabaseAssetListItem::GetAssetEnabledChecked() const
 	{
 		TSharedPtr<FDatabaseViewModel> ViewModelPtr = EditorViewModel.Pin();
-		TSharedPtr<FDatabaseAssetTreeNode> TreeNodePtr = WeakAssetTreeNode.Pin();
-		const UPoseSearchDatabase* Database = ViewModelPtr->GetPoseSearchDatabase();
-
-		if (Database->AnimationAssets.IsValidIndex(TreeNodePtr->SourceAssetIdx))
+		if (const UPoseSearchDatabase* Database = ViewModelPtr->GetPoseSearchDatabase())
 		{
-			if (ViewModelPtr->IsEnabled(TreeNodePtr->SourceAssetIdx))
+			TSharedPtr<FDatabaseAssetTreeNode> TreeNodePtr = WeakAssetTreeNode.Pin();
+			if (Database->AnimationAssets.IsValidIndex(TreeNodePtr->SourceAssetIdx))
 			{
-				return ECheckBoxState::Checked;
+				if (ViewModelPtr->IsEnabled(TreeNodePtr->SourceAssetIdx))
+				{
+					return ECheckBoxState::Checked;
+				}
 			}
 		}
-
 		return ECheckBoxState::Unchecked;
 	}
 
@@ -468,22 +474,22 @@ namespace UE::PoseSearch
 	FSlateColor SDatabaseAssetListItem::GetNameTextColorAndOpacity() const
 	{
 		TSharedPtr<FDatabaseViewModel> ViewModelPtr = EditorViewModel.Pin();
-		TSharedPtr<FDatabaseAssetTreeNode> TreeNodePtr = WeakAssetTreeNode.Pin();
-		const UPoseSearchDatabase* Database = ViewModelPtr->GetPoseSearchDatabase();
-
-		if (const FPoseSearchDatabaseAnimationAssetBase* DatabaseAnimationAssetBase = Database->GetAnimationAssetBase(TreeNodePtr->SourceAssetIdx))
+		if (const UPoseSearchDatabase* Database = ViewModelPtr->GetPoseSearchDatabase())
 		{
-			if (DatabaseAnimationAssetBase->IsEnabled())
+			TSharedPtr<FDatabaseAssetTreeNode> TreeNodePtr = WeakAssetTreeNode.Pin();
+			if (const FPoseSearchDatabaseAnimationAssetBase* DatabaseAnimationAssetBase = Database->GetAnimationAssetBase(TreeNodePtr->SourceAssetIdx))
 			{
-				if (DatabaseAnimationAssetBase->bSynchronizeWithExternalDependency)
+				if (DatabaseAnimationAssetBase->IsEnabled())
 				{
-					return FColor::Turquoise;
-				}
+					if (DatabaseAnimationAssetBase->bSynchronizeWithExternalDependency)
+					{
+						return FColor::Turquoise;
+					}
 
-				return FLinearColor::White;
+					return FLinearColor::White;
+				}
 			}
 		}
-
 		return DisabledColor;
 	}
 
diff --git a/Engine/Plugins/Experimental/Animation/PoseSearch/Source/Editor/Private/PoseSearchDatabaseAssetTree.cpp b/Engine/Plugins/Experimental/Animation/PoseSearch/Source/Editor/Private/PoseSearchDatabaseAssetTree.cpp
--- a/Engine/Plugins/Experimental/Animation/PoseSearch/Source/Editor/Private/PoseSearchDatabaseAssetTree.cpp
+++ b/Engine/Plugins/Experimental/Animation/PoseSearch/Source/Editor/Private/PoseSearchDatabaseAssetTree.cpp
@@ -170,7 +170,7 @@ namespace UE::PoseSearch
 		// Ensure that we only react to modifications to the UPosesSearchDatabase.
 		if (const TSharedPtr<FDatabaseViewModel> ViewModel = EditorViewModel.Pin())
 		{
-			if (const UPoseSearchDatabase * Database = ViewModel->GetPoseSearchDatabase())
+			if (const UPoseSearchDatabase* Database = ViewModel->GetPoseSearchDatabase())
 			{
 				for (const TPair<UObject*, FTransactionObjectEvent>& TransactionObjectPair : TransactionObjectContexts)
 				{
@@ -388,51 +388,50 @@ namespace UE::PoseSearch
 		const int32 NumAssets = DroppedAssetData.Num();
 
 		int32 AddedAssets = 0;
-		if (NumAssets > 0)
+		UPoseSearchDatabase* PoseSearchDatabase = ViewModel->GetPoseSearchDatabase();
+		if (PoseSearchDatabase && NumAssets > 0)
 		{
 			GWarn->BeginSlowTask(LOCTEXT("LoadingAssets", "Loading Asset(s)"), true);
 
+			const FScopedTransaction Transaction(LOCTEXT("AddSequencesOrBlendspaces", "Add Sequence(s) and/or Blendspace(s) to Pose Search Database"));
+			
+			PoseSearchDatabase->Modify();
+			for (int32 DroppedAssetIdx = 0; DroppedAssetIdx < NumAssets; ++DroppedAssetIdx)
 			{
-				const FScopedTransaction Transaction(LOCTEXT("AddSequencesOrBlendspaces", "Add Sequence(s) and/or Blendspace(s) to Pose Search Database"));
-				ViewModel->GetPoseSearchDatabase()->Modify();
-				
-				for (int32 DroppedAssetIdx = 0; DroppedAssetIdx < NumAssets; ++DroppedAssetIdx)
-				{
-					const FAssetData& AssetData = DroppedAssetData[DroppedAssetIdx];
+				const FAssetData& AssetData = DroppedAssetData[DroppedAssetIdx];
 
-					if (!AssetData.IsAssetLoaded())
-					{
-						GWarn->StatusUpdate(
-							DroppedAssetIdx,
-							NumAssets,
-							FText::Format(
-								LOCTEXT("LoadingAsset", "Loading Asset {0}"),
-								FText::FromName(AssetData.AssetName)));
-					}
+				if (!AssetData.IsAssetLoaded())
+				{
+					GWarn->StatusUpdate(
+						DroppedAssetIdx,
+						NumAssets,
+						FText::Format(
+							LOCTEXT("LoadingAsset", "Loading Asset {0}"),
+							FText::FromName(AssetData.AssetName)));
+				}
 
-					UClass* AssetClass = AssetData.GetClass();
-					UObject* Asset = AssetData.GetAsset();
+				UClass* AssetClass = AssetData.GetClass();
+				UObject* Asset = AssetData.GetAsset();
 					
-					if (AssetClass->IsChildOf(UAnimSequence::StaticClass()))
-					{
-						ViewModel->AddSequenceToDatabase(Cast<UAnimSequence>(Asset));
-						++AddedAssets;
-					}
-					if (AssetClass->IsChildOf(UAnimComposite::StaticClass()))
-					{
-						ViewModel->AddAnimCompositeToDatabase(Cast<UAnimComposite>(Asset));
-						++AddedAssets;
-					}
-					else if (AssetClass->IsChildOf(UBlendSpace::StaticClass()))
-					{
-						ViewModel->AddBlendSpaceToDatabase(Cast<UBlendSpace>(Asset));
-						++AddedAssets;
-					}
-					else if (AssetClass->IsChildOf(UAnimMontage::StaticClass()))
-					{
-						ViewModel->AddAnimMontageToDatabase(Cast<UAnimMontage>(Asset));
-						++AddedAssets;
-					}
+				if (AssetClass->IsChildOf(UAnimSequence::StaticClass()))
+				{
+					ViewModel->AddSequenceToDatabase(Cast<UAnimSequence>(Asset));
+					++AddedAssets;
+				}
+				if (AssetClass->IsChildOf(UAnimComposite::StaticClass()))
+				{
+					ViewModel->AddAnimCompositeToDatabase(Cast<UAnimComposite>(Asset));
+					++AddedAssets;
+				}
+				else if (AssetClass->IsChildOf(UBlendSpace::StaticClass()))
+				{
+					ViewModel->AddBlendSpaceToDatabase(Cast<UBlendSpace>(Asset));
+					++AddedAssets;
+				}
+				else if (AssetClass->IsChildOf(UAnimMontage::StaticClass()))
+				{
+					ViewModel->AddAnimMontageToDatabase(Cast<UAnimMontage>(Asset));
+					++AddedAssets;
 				}
 			}
 			
@@ -559,31 +558,37 @@ namespace UE::PoseSearch
 
 	void SDatabaseAssetTree::OnAddSequence(bool bFinalizeChanges)
 	{
-		FScopedTransaction Transaction(LOCTEXT("AddSequence", "Add Sequence"));
 		const TSharedPtr<FDatabaseViewModel> ViewModel = EditorViewModel.Pin();
-		
-		ViewModel->GetPoseSearchDatabase()->Modify();
-		
-		ViewModel->AddSequenceToDatabase(nullptr);
-
-		if (bFinalizeChanges)
+		if (UPoseSearchDatabase* PoseSearchDatabase = ViewModel->GetPoseSearchDatabase())
 		{
-			FinalizeTreeChanges();
+			FScopedTransaction Transaction(LOCTEXT("AddSequence", "Add Sequence"));
+
+			PoseSearchDatabase->Modify();
+
+			ViewModel->AddSequenceToDatabase(nullptr);
+
+			if (bFinalizeChanges)
+			{
+				FinalizeTreeChanges();
+			}
 		}
 	}
 
 	void SDatabaseAssetTree::OnAddBlendSpace(bool bFinalizeChanges)
 	{
-		FScopedTransaction Transaction(LOCTEXT("AddBlendSpaceTransaction", "Add Blend Space"));
 		const TSharedPtr<FDatabaseViewModel> ViewModel = EditorViewModel.Pin();
-		
-		ViewModel->GetPoseSearchDatabase()->Modify();
-		
-		ViewModel->AddBlendSpaceToDatabase(nullptr);
-
-		if (bFinalizeChanges)
+		if (UPoseSearchDatabase* PoseSearchDatabase = ViewModel->GetPoseSearchDatabase())
 		{
-			FinalizeTreeChanges();
+			FScopedTransaction Transaction(LOCTEXT("AddBlendSpaceTransaction", "Add Blend Space"));
+
+			PoseSearchDatabase->Modify();
+
+			ViewModel->AddBlendSpaceToDatabase(nullptr);
+
+			if (bFinalizeChanges)
+			{
+				FinalizeTreeChanges();
+			}
 		}
 	}
 
@@ -701,21 +706,24 @@ namespace UE::PoseSearch
 
 	void SDatabaseAssetTree::EnableSelectedNodes(bool bIsEnabled)
 	{
-		TArray<TSharedPtr<FDatabaseAssetTreeNode>> SelectedNodes = TreeView->GetSelectedItems();
-		if (!SelectedNodes.IsEmpty())
+		const TSharedPtr<FDatabaseViewModel> ViewModel = EditorViewModel.Pin();
+		if (UPoseSearchDatabase* PoseSearchDatabase = ViewModel->GetPoseSearchDatabase())
 		{
-			const FText TransactionName = bIsEnabled ? LOCTEXT("EnablePoseSearchDatabaseNodes", "Enable selected items from Pose Search Database") : LOCTEXT("DisablePoseSearchDatabaseNodes", "Disable selected items from Pose Search Database");
-			const FScopedTransaction Transaction(TransactionName);
-			const TSharedPtr<FDatabaseViewModel> ViewModel = EditorViewModel.Pin();
+			TArray<TSharedPtr<FDatabaseAssetTreeNode>> SelectedNodes = TreeView->GetSelectedItems();
+			if (!SelectedNodes.IsEmpty())
+			{
+				const FText TransactionName = bIsEnabled ? LOCTEXT("EnablePoseSearchDatabaseNodes", "Enable selected items from Pose Search Database") : LOCTEXT("DisablePoseSearchDatabaseNodes", "Disable selected items from Pose Search Database");
+				const FScopedTransaction Transaction(TransactionName);
 
-			ViewModel->GetPoseSearchDatabase()->Modify();
+				PoseSearchDatabase->Modify();
 
-			for (TSharedPtr<FDatabaseAssetTreeNode> SelectedNode : SelectedNodes)
-			{
-				ViewModel->SetIsEnabled(SelectedNode->SourceAssetIdx, bIsEnabled);
+				for (TSharedPtr<FDatabaseAssetTreeNode> SelectedNode : SelectedNodes)
+				{
+					ViewModel->SetIsEnabled(SelectedNode->SourceAssetIdx, bIsEnabled);
+				}
+
+				FinalizeTreeChanges();
 			}
-		
-			FinalizeTreeChanges();
 		}
 	}
 
diff --git a/Engine/Plugins/Experimental/Animation/PoseSearch/Source/Editor/Private/PoseSearchDatabaseDataDetails.cpp b/Engine/Plugins/Experimental/Animation/PoseSearch/Source/Editor/Private/PoseSearchDatabaseDataDetails.cpp
--- a/Engine/Plugins/Experimental/Animation/PoseSearch/Source/Editor/Private/PoseSearchDatabaseDataDetails.cpp
+++ b/Engine/Plugins/Experimental/Animation/PoseSearch/Source/Editor/Private/PoseSearchDatabaseDataDetails.cpp
@@ -224,7 +224,7 @@ public:
 					{
 						if (const FDatabasePreviewActor* FoundPreviewActor = ViewModel->GetPreviewActors().FindByPredicate(
 							[ColumnName](const FDatabasePreviewActor& PreviewActor)
-							{ return *PreviewActor.GetActor()->GetName() == ColumnName; }))
+							{ return *GetNameSafe(PreviewActor.GetActor()) == ColumnName; }))
 						{
 							const int32 PoseIdx = FoundPreviewActor->GetCurrentPoseIndex();
 							const TArray<float> PoseValues = PoseSearchDatabase->GetSearchIndex().GetPoseValuesSafe(PoseIdx);
@@ -357,7 +357,7 @@ void SDatabaseDataDetails::Reconstruct(int32 MaxPreviewActors)
 	for (int32 PreviewActorIdx = 0; PreviewActorIdx < PreviewActorNum; ++PreviewActorIdx)
 	{
 		const FDatabasePreviewActor& PreviewActor = PreviewActors[PreviewActorIdx];
-		HeaderRow->AddColumn(SHeaderRow::Column(*PreviewActor.GetActor()->GetName())
+		HeaderRow->AddColumn(SHeaderRow::Column(*GetNameSafe(PreviewActor.GetActor()))
 			.DefaultLabel(FText::FromString(PreviewActor.GetSampler().GetAsset()->GetName())));
 	}
 
diff --git a/Engine/Plugins/Experimental/Animation/PoseSearch/Source/Editor/Private/PoseSearchDatabaseEditor.cpp b/Engine/Plugins/Experimental/Animation/PoseSearch/Source/Editor/Private/PoseSearchDatabaseEditor.cpp
--- a/Engine/Plugins/Experimental/Animation/PoseSearch/Source/Editor/Private/PoseSearchDatabaseEditor.cpp
+++ b/Engine/Plugins/Experimental/Animation/PoseSearch/Source/Editor/Private/PoseSearchDatabaseEditor.cpp
@@ -387,7 +387,7 @@ namespace UE::PoseSearch
 	FText FDatabaseEditor::GetToolkitName() const
 	{
 		FFormatNamedArguments Args;
-		Args.Add(TEXT("AssetName"), FText::FromString(GetPoseSearchDatabase()->GetName()));
+		Args.Add(TEXT("AssetName"), FText::FromString(GetNameSafe(GetPoseSearchDatabase())));
 		return FText::Format(LOCTEXT("PoseSearchDatabaseEditorToolkitName", "{AssetName}"), Args);
 	}
 
@@ -512,17 +512,17 @@ namespace UE::PoseSearch
 			SelectionWidgetPair.Value.SelectedReflections.Reset();
 		}
 
-		if (SelectedItems.Num() > 0)
+		const UPoseSearchDatabase* PoseSearchDatabase = GetPoseSearchDatabase();
+		if (PoseSearchDatabase && SelectedItems.Num() > 0)
 		{
 			for (TSharedPtr<FDatabaseAssetTreeNode>& SelectedItem : SelectedItems)
 			{
-				if (!SelectedItem.IsValid() ||
-					!GetPoseSearchDatabase()->AnimationAssets.IsValidIndex(SelectedItem->SourceAssetIdx))
+				if (!SelectedItem.IsValid() || !PoseSearchDatabase->AnimationAssets.IsValidIndex(SelectedItem->SourceAssetIdx))
 				{
 					continue;
 				}
 
-				const FInstancedStruct& DatabaseAsset = GetPoseSearchDatabase()->AnimationAssets[SelectedItem->SourceAssetIdx];
+				const FInstancedStruct& DatabaseAsset = PoseSearchDatabase->AnimationAssets[SelectedItem->SourceAssetIdx];
 				const UScriptStruct* ScriptStruct = DatabaseAsset.GetScriptStruct();
 				FSelectionWidget& SelectionWidget = FindOrAddSelectionWidget(ScriptStruct);
 
diff --git a/Engine/Plugins/Experimental/Animation/PoseSearch/Source/Editor/Private/PoseSearchDatabaseEditorReflection.cpp b/Engine/Plugins/Experimental/Animation/PoseSearch/Source/Editor/Private/PoseSearchDatabaseEditorReflection.cpp
--- a/Engine/Plugins/Experimental/Animation/PoseSearch/Source/Editor/Private/PoseSearchDatabaseEditorReflection.cpp
+++ b/Engine/Plugins/Experimental/Animation/PoseSearch/Source/Editor/Private/PoseSearchDatabaseEditorReflection.cpp
@@ -109,7 +109,7 @@ void UPoseSearchDatabaseStatistics::Initialize(const UPoseSearchDatabase* PoseSe
 {
 	static FText TimeFormat = LOCTEXT("TimeFormat", "{0} {0}|plural(one=Second,other=Seconds)");
 	
-	if (PoseSearchDatabase)
+	if (IsValid(PoseSearchDatabase))
 	{
 		const UE::PoseSearch::FSearchIndex& SearchIndex = PoseSearchDatabase->GetSearchIndex();
 		// General Information
diff --git a/Engine/Plugins/Experimental/Animation/PoseSearch/Source/Editor/Private/PoseSearchDatabaseViewModel.cpp b/Engine/Plugins/Experimental/Animation/PoseSearch/Source/Editor/Private/PoseSearchDatabaseViewModel.cpp
--- a/Engine/Plugins/Experimental/Animation/PoseSearch/Source/Editor/Private/PoseSearchDatabaseViewModel.cpp
+++ b/Engine/Plugins/Experimental/Animation/PoseSearch/Source/Editor/Private/PoseSearchDatabaseViewModel.cpp
@@ -66,10 +66,10 @@ bool FDatabasePreviewActor::SpawnPreviewActor(UWorld* World, const UPoseSearchDa
 
 	FActorSpawnParameters Params;
 	Params.SpawnCollisionHandlingOverride = ESpawnActorCollisionHandlingMethod::AdjustIfPossibleButAlwaysSpawn;
-	Actor = World->SpawnActor<AActor>(AActor::StaticClass(), FTransform::Identity, Params);
-	Actor->SetFlags(RF_Transient);
+	ActorPtr = World->SpawnActor<AActor>(AActor::StaticClass(), FTransform::Identity, Params);
+	ActorPtr->SetFlags(RF_Transient);
 
-	UDebugSkelMeshComponent* Mesh = NewObject<UDebugSkelMeshComponent>(Actor.Get());
+	UDebugSkelMeshComponent* Mesh = NewObject<UDebugSkelMeshComponent>(ActorPtr.Get());
 	Mesh->RegisterComponentWithWorld(World);
 
 	UAnimPreviewInstance* AnimInstance = NewObject<UAnimPreviewInstance>(Mesh);
@@ -92,14 +92,14 @@ bool FDatabasePreviewActor::SpawnPreviewActor(UWorld* World, const UPoseSearchDa
 
 	AnimInstance->PlayAnim(false, 0.0f);
 
-	if (!Actor->GetRootComponent())
+	if (!ActorPtr->GetRootComponent())
 	{
-		Actor->SetRootComponent(Mesh);
+		ActorPtr->SetRootComponent(Mesh);
 	}
 
 	AnimInstance->SetPlayRate(0.f);
 
-	UE_LOG(LogPoseSearchEditor, Log, TEXT("Spawned preview Actor: %s"), *GetNameSafe(Actor.Get()));
+	UE_LOG(LogPoseSearchEditor, Log, TEXT("Spawned preview Actor: %s"), *GetNameSafe(ActorPtr.Get()));
 	return true;
 }
 
@@ -154,7 +154,8 @@ void FDatabasePreviewActor::UpdatePreviewActor(const UPoseSearchDatabase* PoseSe
 		RootTransformCurrent.SetToRelativeTransform(RootTransformOrigin);
 	}
 
-	Actor->SetActorTransform(RootTransformCurrent);
+	check(ActorPtr != nullptr);
+	ActorPtr->SetActorTransform(RootTransformCurrent);
 
 	// @todo: optimize this bone container, since we only need the root bone here...
 	const FBoneContainer& BoneContainer = AnimInstance->GetRequiredBonesOnAnyThread();
@@ -170,8 +171,10 @@ void FDatabasePreviewActor::UpdatePreviewActor(const UPoseSearchDatabase* PoseSe
 
 void FDatabasePreviewActor::Destroy()
 {
-	check(Actor);
-	Actor->Destroy();
+	if (ActorPtr != nullptr)
+	{
+		ActorPtr->Destroy();
+	}
 }
 
 bool FDatabasePreviewActor::DrawPreviewActor(const UPoseSearchDatabase* PoseSearchDatabase, bool bDisplayRootMotionSpeed, TConstArrayView<float> QueryVector)
@@ -291,9 +294,9 @@ void FDatabasePreviewActor::ExtractPose(float Time, FCompactPose& OutPose) const
 
 const UDebugSkelMeshComponent* FDatabasePreviewActor::GetDebugSkelMeshComponent() const
 {
-	if (Actor)
+	if (ActorPtr != nullptr)
 	{
-		return Cast<UDebugSkelMeshComponent>(Actor->GetRootComponent());
+		return Cast<UDebugSkelMeshComponent>(ActorPtr->GetRootComponent());
 	}
 	return nullptr;
 }
@@ -309,9 +312,9 @@ const UAnimPreviewInstance* FDatabasePreviewActor::GetAnimPreviewInstance() cons
 
 UAnimPreviewInstance* FDatabasePreviewActor::GetAnimPreviewInstanceInternal()
 {
-	if (Actor)
+	if (ActorPtr != nullptr)
 	{
-		if (UDebugSkelMeshComponent* Mesh = Cast<UDebugSkelMeshComponent>(Actor->GetRootComponent()))
+		if (UDebugSkelMeshComponent* Mesh = Cast<UDebugSkelMeshComponent>(ActorPtr->GetRootComponent()))
 		{
 			return Mesh->PreviewInstance.Get();
 		}
@@ -320,19 +323,14 @@ UAnimPreviewInstance* FDatabasePreviewActor::GetAnimPreviewInstanceInternal()
 }
 
 // FDatabaseViewModel
-FDatabaseViewModel::FDatabaseViewModel()
-	: PoseSearchDatabase(nullptr)
-{
-}
-
 void FDatabaseViewModel::AddReferencedObjects(FReferenceCollector& Collector)
 {
-	Collector.AddReferencedObject(PoseSearchDatabase);
+	Collector.AddReferencedObject(PoseSearchDatabasePtr);
 }
 
 void FDatabaseViewModel::Initialize(UPoseSearchDatabase* InPoseSearchDatabase, const TSharedRef<FDatabasePreviewScene>& InPreviewScene, const TSharedRef<SDatabaseDataDetails>& InDatabaseDataDetails)
 {
-	PoseSearchDatabase = InPoseSearchDatabase;
+	PoseSearchDatabasePtr = InPoseSearchDatabase;
 	PreviewScenePtr = InPreviewScene;
 	DatabaseDataDetails = InDatabaseDataDetails;
 
@@ -342,7 +340,7 @@ void FDatabaseViewModel::Initialize(UPoseSearchDatabase* InPoseSearchDatabase, c
 void FDatabaseViewModel::BuildSearchIndex()
 {
 	using namespace UE::PoseSearch;
-	FAsyncPoseSearchDatabasesManagement::RequestAsyncBuildIndex(PoseSearchDatabase, ERequestAsyncBuildFlag::NewRequest);
+	FAsyncPoseSearchDatabasesManagement::RequestAsyncBuildIndex(GetPoseSearchDatabase(), ERequestAsyncBuildFlag::NewRequest);
 }
 
 void FDatabaseViewModel::PreviewBackwardEnd()
@@ -401,11 +399,14 @@ void FDatabaseViewModel::Tick(float DeltaSeconds)
 		PlayTime += DeltaPlayTime;
 		PlayTime = FMath::Clamp(PlayTime, MinPreviewPlayLength, MaxPreviewPlayLength);
 
-		if (FAsyncPoseSearchDatabasesManagement::RequestAsyncBuildIndex(PoseSearchDatabase, ERequestAsyncBuildFlag::ContinueRequest))
+		if (const UPoseSearchDatabase* Database = GetPoseSearchDatabase())
 		{
-			for (FDatabasePreviewActor& PreviewActor : GetPreviewActors())
+			if (FAsyncPoseSearchDatabasesManagement::RequestAsyncBuildIndex(Database, ERequestAsyncBuildFlag::ContinueRequest))
 			{
-				PreviewActor.UpdatePreviewActor(PoseSearchDatabase, PlayTime, bQuantizeAnimationToPoseData);
+				for (FDatabasePreviewActor& PreviewActor : GetPreviewActors())
+				{
+					PreviewActor.UpdatePreviewActor(Database, PlayTime, bQuantizeAnimationToPoseData);
+				}
 			}
 		}
 	}
@@ -430,42 +431,57 @@ void FDatabaseViewModel::RemovePreviewActors()
 
 void FDatabaseViewModel::AddSequenceToDatabase(UAnimSequence* AnimSequence)
 {
-	FPoseSearchDatabaseSequence NewAsset;
-	NewAsset.Sequence = AnimSequence;
-	PoseSearchDatabase->AnimationAssets.Add(FInstancedStruct::Make(NewAsset));
+	if (UPoseSearchDatabase* Database = GetPoseSearchDatabase())
+	{
+		FPoseSearchDatabaseSequence NewAsset;
+		NewAsset.Sequence = AnimSequence;
+		Database->AnimationAssets.Add(FInstancedStruct::Make(NewAsset));
+	}
 }
 
 void FDatabaseViewModel::AddBlendSpaceToDatabase(UBlendSpace* BlendSpace)
 {
-	FPoseSearchDatabaseBlendSpace NewAsset;
-	NewAsset.BlendSpace = BlendSpace;
-	PoseSearchDatabase->AnimationAssets.Add(FInstancedStruct::Make(NewAsset));
+	if (UPoseSearchDatabase* Database = GetPoseSearchDatabase())
+	{
+		FPoseSearchDatabaseBlendSpace NewAsset;
+		NewAsset.BlendSpace = BlendSpace;
+		Database->AnimationAssets.Add(FInstancedStruct::Make(NewAsset));
+	}
 }
 
 void FDatabaseViewModel::AddAnimCompositeToDatabase(UAnimComposite* AnimComposite)
 {
-	FPoseSearchDatabaseAnimComposite NewAsset;
-	NewAsset.AnimComposite = AnimComposite;
-	PoseSearchDatabase->AnimationAssets.Add(FInstancedStruct::Make(NewAsset));
+	if (UPoseSearchDatabase* Database = GetPoseSearchDatabase())
+	{
+		FPoseSearchDatabaseAnimComposite NewAsset;
+		NewAsset.AnimComposite = AnimComposite;
+		Database->AnimationAssets.Add(FInstancedStruct::Make(NewAsset));
+	}
 }
 
 void FDatabaseViewModel::AddAnimMontageToDatabase(UAnimMontage* AnimMontage)
 {
-	FPoseSearchDatabaseAnimMontage NewAsset;
-	NewAsset.AnimMontage = AnimMontage;
-	PoseSearchDatabase->AnimationAssets.Add(FInstancedStruct::Make(NewAsset));
+	if (UPoseSearchDatabase* Database = GetPoseSearchDatabase())
+	{
+		FPoseSearchDatabaseAnimMontage NewAsset;
+		NewAsset.AnimMontage = AnimMontage;
+		Database->AnimationAssets.Add(FInstancedStruct::Make(NewAsset));
+	}
 }
 
 bool FDatabaseViewModel::DeleteFromDatabase(int32 AnimationAssetIndex)
 {
-	if (const FPoseSearchDatabaseAnimationAssetBase* DatabaseAnimationAssetBase = PoseSearchDatabase->GetAnimationAssetBase(AnimationAssetIndex))
+	if (UPoseSearchDatabase* Database = GetPoseSearchDatabase())
 	{
-		if (!DatabaseAnimationAssetBase->bSynchronizeWithExternalDependency)
+		if (const FPoseSearchDatabaseAnimationAssetBase* DatabaseAnimationAssetBase = Database->GetAnimationAssetBase(AnimationAssetIndex))
 		{
-			PoseSearchDatabase->AnimationAssets.RemoveAt(AnimationAssetIndex);
-			PoseSearchDatabase->Modify();
-		
-			return true;
+			if (!DatabaseAnimationAssetBase->bSynchronizeWithExternalDependency)
+			{
+				Database->AnimationAssets.RemoveAt(AnimationAssetIndex);
+				Database->Modify();
+
+				return true;
+			}
 		}
 	}
 
@@ -474,17 +490,23 @@ bool FDatabaseViewModel::DeleteFromDatabase(int32 AnimationAssetIndex)
 
 void FDatabaseViewModel::SetDisableReselection(int32 AnimationAssetIndex, bool bEnabled)
 {
-	if (FPoseSearchDatabaseAnimationAssetBase* DatabaseAnimationAsset = PoseSearchDatabase->GetMutableAnimationAssetBase(AnimationAssetIndex))
+	if (UPoseSearchDatabase* Database = GetPoseSearchDatabase())
 	{
-		DatabaseAnimationAsset->SetDisableReselection(bEnabled);
+		if (FPoseSearchDatabaseAnimationAssetBase* DatabaseAnimationAsset = Database->GetMutableAnimationAssetBase(AnimationAssetIndex))
+		{
+			DatabaseAnimationAsset->SetDisableReselection(bEnabled);
+		}
 	}
 }
 
 bool FDatabaseViewModel::IsDisableReselection(int32 AnimationAssetIndex) const
 {
-	if (const FPoseSearchDatabaseAnimationAssetBase* DatabaseAnimationAsset = PoseSearchDatabase->GetAnimationAssetBase(AnimationAssetIndex))
+	if (const UPoseSearchDatabase* Database = GetPoseSearchDatabase())
 	{
-		return DatabaseAnimationAsset->IsDisableReselection();
+		if (const FPoseSearchDatabaseAnimationAssetBase* DatabaseAnimationAsset = Database->GetAnimationAssetBase(AnimationAssetIndex))
+		{
+			return DatabaseAnimationAsset->IsDisableReselection();
+		}
 	}
 
 	return false;
@@ -492,19 +514,25 @@ bool FDatabaseViewModel::IsDisableReselection(int32 AnimationAssetIndex) const
 
 void FDatabaseViewModel::SetIsEnabled(int32 AnimationAssetIndex, bool bEnabled)
 {
-	if (FPoseSearchDatabaseAnimationAssetBase* DatabaseAnimationAsset = PoseSearchDatabase->GetMutableAnimationAssetBase(AnimationAssetIndex))
+	if (UPoseSearchDatabase* Database = GetPoseSearchDatabase())
 	{
-		GetPoseSearchDatabase()->Modify();
+		if (FPoseSearchDatabaseAnimationAssetBase* DatabaseAnimationAsset = Database->GetMutableAnimationAssetBase(AnimationAssetIndex))
+		{
+			Database->Modify();
 
-		DatabaseAnimationAsset->SetIsEnabled(bEnabled);
+			DatabaseAnimationAsset->SetIsEnabled(bEnabled);
+		}
 	}
 }
 
 bool FDatabaseViewModel::IsEnabled(int32 AnimationAssetIndex) const
 {
-	if (const FPoseSearchDatabaseAnimationAssetBase* DatabaseAnimationAsset = PoseSearchDatabase->GetAnimationAssetBase(AnimationAssetIndex))
+	if (const UPoseSearchDatabase* Database = GetPoseSearchDatabase())
 	{
-		return DatabaseAnimationAsset->IsEnabled();
+		if (const FPoseSearchDatabaseAnimationAssetBase* DatabaseAnimationAsset = Database->GetAnimationAssetBase(AnimationAssetIndex))
+		{
+			return DatabaseAnimationAsset->IsEnabled();
+		}
 	}
 
 	return false;
@@ -523,34 +551,37 @@ int32 FDatabaseViewModel::SetSelectedNode(int32 PoseIdx, bool bClearSelection, b
 	bDrawQueryVector = bDrawQuery;
 	QueryVector = InQueryVector;
 
-	if (FAsyncPoseSearchDatabasesManagement::RequestAsyncBuildIndex(PoseSearchDatabase, ERequestAsyncBuildFlag::ContinueRequest))
+	if (const UPoseSearchDatabase* Database = GetPoseSearchDatabase())
 	{
-		const FSearchIndex& SearchIndex = PoseSearchDatabase->GetSearchIndex();
-		if (SearchIndex.PoseMetadata.IsValidIndex(PoseIdx))
+		if (FAsyncPoseSearchDatabasesManagement::RequestAsyncBuildIndex(Database, ERequestAsyncBuildFlag::ContinueRequest))
 		{
-			const uint32 IndexAssetIndex = SearchIndex.PoseMetadata[PoseIdx].GetAssetIndex();
-			if (SearchIndex.Assets.IsValidIndex(IndexAssetIndex))
+			const FSearchIndex& SearchIndex = Database->GetSearchIndex();
+			if (SearchIndex.PoseMetadata.IsValidIndex(PoseIdx))
 			{
-				FDatabasePreviewActor PreviewActor;
-				if (PreviewActor.SpawnPreviewActor(GetWorld(), PoseSearchDatabase, IndexAssetIndex, PoseIdx))
+				const uint32 IndexAssetIndex = SearchIndex.PoseMetadata[PoseIdx].GetAssetIndex();
+				if (SearchIndex.Assets.IsValidIndex(IndexAssetIndex))
 				{
-					const FSearchIndexAsset& IndexAsset = SearchIndex.Assets[IndexAssetIndex];
-					MaxPreviewPlayLength = FMath::Max(MaxPreviewPlayLength, IndexAsset.GetLastSampleTime(PoseSearchDatabase->Schema->SampleRate) - PreviewActor.GetPlayTimeOffset());
-					MinPreviewPlayLength = FMath::Min(MinPreviewPlayLength, IndexAsset.GetFirstSampleTime(PoseSearchDatabase->Schema->SampleRate) - PreviewActor.GetPlayTimeOffset());
-					PreviewActors.Add(PreviewActor);
-					SelectedSourceAssetIdx = IndexAsset.GetSourceAssetIdx();
+					FDatabasePreviewActor PreviewActor;
+					if (PreviewActor.SpawnPreviewActor(GetWorld(), Database, IndexAssetIndex, PoseIdx))
+					{
+						const FSearchIndexAsset& IndexAsset = SearchIndex.Assets[IndexAssetIndex];
+						MaxPreviewPlayLength = FMath::Max(MaxPreviewPlayLength, IndexAsset.GetLastSampleTime(Database->Schema->SampleRate) - PreviewActor.GetPlayTimeOffset());
+						MinPreviewPlayLength = FMath::Min(MinPreviewPlayLength, IndexAsset.GetFirstSampleTime(Database->Schema->SampleRate) - PreviewActor.GetPlayTimeOffset());
+						PreviewActors.Add(PreviewActor);
+						SelectedSourceAssetIdx = IndexAsset.GetSourceAssetIdx();
+					}
 				}
 			}
-		}
 
-		DatabaseDataDetails.Pin()->Reconstruct();
+			DatabaseDataDetails.Pin()->Reconstruct();
 
-		for (FDatabasePreviewActor& PreviewActor : GetPreviewActors())
-		{
-			PreviewActor.UpdatePreviewActor(PoseSearchDatabase, PlayTime, bQuantizeAnimationToPoseData);
-		}
+			for (FDatabasePreviewActor& PreviewActor : GetPreviewActors())
+			{
+				PreviewActor.UpdatePreviewActor(Database, PlayTime, bQuantizeAnimationToPoseData);
+			}
 
-		SetPlayTime(0.f, false);
+			SetPlayTime(0.f, false);
+		}
 	}
 
 	ProcessSelectedActor(nullptr);
@@ -562,37 +593,40 @@ void FDatabaseViewModel::SetSelectedNodes(const TArrayView<TSharedPtr<FDatabaseA
 {
 	RemovePreviewActors();
 
-	if (FAsyncPoseSearchDatabasesManagement::RequestAsyncBuildIndex(PoseSearchDatabase, ERequestAsyncBuildFlag::ContinueRequest))
+	if (const UPoseSearchDatabase* Database = GetPoseSearchDatabase())
 	{
-		TMap<int32, int32> AssociatedAssetIndices;
-		for (int32 i = 0; i < InSelectedNodes.Num(); ++i)
+		if (FAsyncPoseSearchDatabasesManagement::RequestAsyncBuildIndex(Database, ERequestAsyncBuildFlag::ContinueRequest))
 		{
-			AssociatedAssetIndices.FindOrAdd(InSelectedNodes[i]->SourceAssetIdx) = i;
-		}
+			TMap<int32, int32> AssociatedAssetIndices;
+			for (int32 i = 0; i < InSelectedNodes.Num(); ++i)
+			{
+				AssociatedAssetIndices.FindOrAdd(InSelectedNodes[i]->SourceAssetIdx) = i;
+			}
 
-		const FSearchIndex& SearchIndex = PoseSearchDatabase->GetSearchIndex();
-		for (int32 IndexAssetIndex = 0; IndexAssetIndex < SearchIndex.Assets.Num(); ++IndexAssetIndex)
-		{
-			const FSearchIndexAsset& IndexAsset = SearchIndex.Assets[IndexAssetIndex];
-			if (const int32* SelectedNodesIndex = AssociatedAssetIndices.Find(IndexAsset.GetSourceAssetIdx()))
+			const FSearchIndex& SearchIndex = Database->GetSearchIndex();
+			for (int32 IndexAssetIndex = 0; IndexAssetIndex < SearchIndex.Assets.Num(); ++IndexAssetIndex)
 			{
-				FDatabasePreviewActor PreviewActor;
-				if (PreviewActor.SpawnPreviewActor(GetWorld(), PoseSearchDatabase, IndexAssetIndex))
+				const FSearchIndexAsset& IndexAsset = SearchIndex.Assets[IndexAssetIndex];
+				if (const int32* SelectedNodesIndex = AssociatedAssetIndices.Find(IndexAsset.GetSourceAssetIdx()))
 				{
-					MaxPreviewPlayLength = FMath::Max(MaxPreviewPlayLength, IndexAsset.GetLastSampleTime(PoseSearchDatabase->Schema->SampleRate) - IndexAsset.GetFirstSampleTime(PoseSearchDatabase->Schema->SampleRate));
-					PreviewActors.Add(PreviewActor);
+					FDatabasePreviewActor PreviewActor;
+					if (PreviewActor.SpawnPreviewActor(GetWorld(), Database, IndexAssetIndex))
+					{
+						MaxPreviewPlayLength = FMath::Max(MaxPreviewPlayLength, IndexAsset.GetLastSampleTime(Database->Schema->SampleRate) - IndexAsset.GetFirstSampleTime(Database->Schema->SampleRate));
+						PreviewActors.Add(PreviewActor);
+					}
 				}
 			}
-		}
 
-		DatabaseDataDetails.Pin()->Reconstruct();
-		for (FDatabasePreviewActor& PreviewActor : GetPreviewActors())
-		{
-			PreviewActor.UpdatePreviewActor(PoseSearchDatabase, PlayTime, bQuantizeAnimationToPoseData);
+			DatabaseDataDetails.Pin()->Reconstruct();
+			for (FDatabasePreviewActor& PreviewActor : GetPreviewActors())
+			{
+				PreviewActor.UpdatePreviewActor(Database, PlayTime, bQuantizeAnimationToPoseData);
+			}
 		}
-	}
 
-	ProcessSelectedActor(nullptr);
+		ProcessSelectedActor(nullptr);
+	}
 }
 
 void FDatabaseViewModel::ProcessSelectedActor(AActor* Actor)
@@ -615,12 +649,15 @@ void FDatabaseViewModel::SetDrawQueryVector(bool bValue)
 
 const FSearchIndexAsset* FDatabaseViewModel::GetSelectedActorIndexAsset() const
 {
-	if (FAsyncPoseSearchDatabasesManagement::RequestAsyncBuildIndex(PoseSearchDatabase, ERequestAsyncBuildFlag::ContinueRequest))
+	if (const UPoseSearchDatabase* Database = GetPoseSearchDatabase())
 	{
-		const FSearchIndex& SearchIndex = PoseSearchDatabase->GetSearchIndex();
-		if (SearchIndex.Assets.IsValidIndex(SelectedActorIndexAssetIndex))
+		if (FAsyncPoseSearchDatabasesManagement::RequestAsyncBuildIndex(Database, ERequestAsyncBuildFlag::ContinueRequest))
 		{
-			return &SearchIndex.Assets[SelectedActorIndexAssetIndex];
+			const FSearchIndex& SearchIndex = Database->GetSearchIndex();
+			if (SearchIndex.Assets.IsValidIndex(SelectedActorIndexAssetIndex))
+			{
+				return &SearchIndex.Assets[SelectedActorIndexAssetIndex];
+			}
 		}
 	}
 	return nullptr;
@@ -645,11 +682,15 @@ void FDatabaseViewModel::SetPlayTime(float NewPlayTime, bool bInTickPlayTime)
 	if (!FMath::IsNearlyEqual(PlayTime, NewPlayTime))
 	{
 		PlayTime = NewPlayTime;
-		if (FAsyncPoseSearchDatabasesManagement::RequestAsyncBuildIndex(PoseSearchDatabase, ERequestAsyncBuildFlag::ContinueRequest))
+		
+		if (const UPoseSearchDatabase* Database = GetPoseSearchDatabase())
 		{
-			for (FDatabasePreviewActor& PreviewActor : GetPreviewActors())
+			if (FAsyncPoseSearchDatabasesManagement::RequestAsyncBuildIndex(Database, ERequestAsyncBuildFlag::ContinueRequest))
 			{
-				PreviewActor.UpdatePreviewActor(PoseSearchDatabase, PlayTime, bQuantizeAnimationToPoseData);
+				for (FDatabasePreviewActor& PreviewActor : GetPreviewActors())
+				{
+					PreviewActor.UpdatePreviewActor(Database, PlayTime, bQuantizeAnimationToPoseData);
+				}
 			}
 		}
 	}
@@ -657,36 +698,39 @@ void FDatabaseViewModel::SetPlayTime(float NewPlayTime, bool bInTickPlayTime)
 
 bool FDatabaseViewModel::GetAnimationTime(int32 SourceAssetIdx, float& CurrentPlayTime, FVector& BlendParameters) const
 {
-	if (FAsyncPoseSearchDatabasesManagement::RequestAsyncBuildIndex(PoseSearchDatabase, ERequestAsyncBuildFlag::ContinueRequest))
+	if (const UPoseSearchDatabase* Database = GetPoseSearchDatabase())
 	{
-		const FSearchIndex& SearchIndex = PoseSearchDatabase->GetSearchIndex();
-		for (const FDatabasePreviewActor& PreviewActor : GetPreviewActors())
+		if (FAsyncPoseSearchDatabasesManagement::RequestAsyncBuildIndex(Database, ERequestAsyncBuildFlag::ContinueRequest))
 		{
-			if (PreviewActor.GetIndexAssetIndex() >= 0 && PreviewActor.GetIndexAssetIndex() < SearchIndex.Assets.Num())
+			const FSearchIndex& SearchIndex = Database->GetSearchIndex();
+			for (const FDatabasePreviewActor& PreviewActor : GetPreviewActors())
 			{
-				const FSearchIndexAsset& IndexAsset = SearchIndex.Assets[PreviewActor.GetIndexAssetIndex()];
-				if (IndexAsset.GetSourceAssetIdx() == SourceAssetIdx)
+				if (PreviewActor.GetIndexAssetIndex() >= 0 && PreviewActor.GetIndexAssetIndex() < SearchIndex.Assets.Num())
 				{
-					CurrentPlayTime = PreviewActor.GetSampler().ToNormalizedTime(PlayTime + IndexAsset.GetFirstSampleTime(PoseSearchDatabase->Schema->SampleRate) + PreviewActor.GetPlayTimeOffset());
-					BlendParameters = IndexAsset.GetBlendParameters();
-					return true;
+					const FSearchIndexAsset& IndexAsset = SearchIndex.Assets[PreviewActor.GetIndexAssetIndex()];
+					if (IndexAsset.GetSourceAssetIdx() == SourceAssetIdx)
+					{
+						CurrentPlayTime = PreviewActor.GetSampler().ToNormalizedTime(PlayTime + IndexAsset.GetFirstSampleTime(Database->Schema->SampleRate) + PreviewActor.GetPlayTimeOffset());
+						BlendParameters = IndexAsset.GetBlendParameters();
+						return true;
+					}
 				}
 			}
-		}
 
-		for (const FSearchIndexAsset& IndexAsset : SearchIndex.Assets)
-		{
-			if (IndexAsset.GetSourceAssetIdx() == SourceAssetIdx)
+			for (const FSearchIndexAsset& IndexAsset : SearchIndex.Assets)
 			{
-				CurrentPlayTime = PlayTime + IndexAsset.GetFirstSampleTime(PoseSearchDatabase->Schema->SampleRate);
-				BlendParameters = IndexAsset.GetBlendParameters();
-
-				const bool bIsBlendSpace = PoseSearchDatabase->GetAnimationAssetStruct(IndexAsset).GetPtr<FPoseSearchDatabaseBlendSpace>() != nullptr;
-				if (bIsBlendSpace && !FMath::IsNearlyEqual(MaxPreviewPlayLength, MinPreviewPlayLength))
+				if (IndexAsset.GetSourceAssetIdx() == SourceAssetIdx)
 				{
-					CurrentPlayTime = (CurrentPlayTime - MaxPreviewPlayLength) / (MaxPreviewPlayLength - MinPreviewPlayLength);
+					CurrentPlayTime = PlayTime + IndexAsset.GetFirstSampleTime(Database->Schema->SampleRate);
+					BlendParameters = IndexAsset.GetBlendParameters();
+
+					const bool bIsBlendSpace = Database->GetAnimationAssetStruct(IndexAsset).GetPtr<FPoseSearchDatabaseBlendSpace>() != nullptr;
+					if (bIsBlendSpace && !FMath::IsNearlyEqual(MaxPreviewPlayLength, MinPreviewPlayLength))
+					{
+						CurrentPlayTime = (CurrentPlayTime - MaxPreviewPlayLength) / (MaxPreviewPlayLength - MinPreviewPlayLength);
+					}
+					return true;
 				}
-				return true;
 			}
 		}
 	}
diff --git a/Engine/Plugins/Experimental/Animation/PoseSearch/Source/Editor/Private/PoseSearchDebuggerView.h b/Engine/Plugins/Experimental/Animation/PoseSearch/Source/Editor/Private/PoseSearchDebuggerView.h
--- a/Engine/Plugins/Experimental/Animation/PoseSearch/Source/Editor/Private/PoseSearchDebuggerView.h
+++ b/Engine/Plugins/Experimental/Animation/PoseSearch/Source/Editor/Private/PoseSearchDebuggerView.h
@@ -47,8 +47,8 @@ private:
 	/** Details widget constructed for the MM node */
 	TSharedPtr<IDetailsView> Details;
 
-	/** Last updated reflection data relative to MM state */
-	TObjectPtr<UPoseSearchDebuggerReflection> Reflection = nullptr;
+	/** Last updated reflection data relative to MM state. This is not gonna be garbage collected because of AddToRoot during Construct */
+	TObjectPtr<UPoseSearchDebuggerReflection> Reflection;
 };
 
 /** Callback to relay closing of the view to destroy the debugger instance */
diff --git a/Engine/Plugins/Experimental/Animation/PoseSearch/Source/Editor/Public/PoseSearchDatabaseViewModel.h b/Engine/Plugins/Experimental/Animation/PoseSearch/Source/Editor/Public/PoseSearchDatabaseViewModel.h
--- a/Engine/Plugins/Experimental/Animation/PoseSearch/Source/Editor/Public/PoseSearchDatabaseViewModel.h
+++ b/Engine/Plugins/Experimental/Animation/PoseSearch/Source/Editor/Public/PoseSearchDatabaseViewModel.h
@@ -39,7 +39,7 @@ namespace UE::PoseSearch
 		FTransform ExtractRootTransform(float Time) const;
 		void ExtractPose(float Time, FCompactPose& OutPose) const;
 
-		const AActor* GetActor() const { return Actor.Get(); }
+		const AActor* GetActor() const { return ActorPtr.Get(); }
 		int32 GetIndexAssetIndex() const { return IndexAssetIndex; }
 		int32 GetCurrentPoseIndex() const { return CurrentPoseIndex; }
 		float GetPlayTimeOffset() const { return PlayTimeOffset; }
@@ -47,7 +47,7 @@ namespace UE::PoseSearch
 	private:
 		UAnimPreviewInstance* GetAnimPreviewInstanceInternal();
 
-		TObjectPtr<AActor> Actor;
+		TWeakObjectPtr<AActor> ActorPtr;
 		int32 IndexAssetIndex = INDEX_NONE;
 		int32 CurrentPoseIndex = INDEX_NONE;
 		float PlayTimeOffset = 0.f;
@@ -71,8 +71,6 @@ namespace UE::PoseSearch
 	class FDatabaseViewModel : public TSharedFromThis<FDatabaseViewModel>, public FGCObject
 	{
 	public:
-		FDatabaseViewModel();
-
 		// ~ FGCObject interface
 		virtual void AddReferencedObjects(FReferenceCollector& Collector) override;
 		virtual FString GetReferencerName() const override { return TEXT("FPoseSearchDatabaseViewModel"); }
@@ -90,7 +88,8 @@ namespace UE::PoseSearch
 		void PreviewForwardStep();
 		void PreviewForwardEnd();
 
-		UPoseSearchDatabase* GetPoseSearchDatabase() { return PoseSearchDatabase; }
+		UPoseSearchDatabase* GetPoseSearchDatabase() { return PoseSearchDatabasePtr.Get(); }
+		const UPoseSearchDatabase* GetPoseSearchDatabase() const { return PoseSearchDatabasePtr.Get(); }
 		void OnPreviewActorClassChanged();
 
 		void Tick(float DeltaSeconds);
@@ -143,7 +142,7 @@ namespace UE::PoseSearch
 		float DeltaTimeMultiplier = 1.f;
 
 		/** Scene asset being viewed and edited by this view model. */
-		TObjectPtr<UPoseSearchDatabase> PoseSearchDatabase;
+		TWeakObjectPtr<UPoseSearchDatabase> PoseSearchDatabasePtr;
 
 		/** Weak pointer to the PreviewScene */
 		TWeakPtr<FDatabasePreviewScene> PreviewScenePtr;
diff --git a/Engine/Plugins/Experimental/Animation/PoseSearch/Source/Runtime/Private/AnimNode_MotionMatching.cpp b/Engine/Plugins/Experimental/Animation/PoseSearch/Source/Runtime/Private/AnimNode_MotionMatching.cpp
--- a/Engine/Plugins/Experimental/Animation/PoseSearch/Source/Runtime/Private/AnimNode_MotionMatching.cpp
+++ b/Engine/Plugins/Experimental/Animation/PoseSearch/Source/Runtime/Private/AnimNode_MotionMatching.cpp
@@ -142,7 +142,7 @@ void FAnimNode_MotionMatching::UpdateAssetPlayer(const FAnimationUpdateContext&
 #if ENABLE_ANIM_DEBUG
 	if (CVarAnimNodeMotionMatchingDrawInfo.GetValueOnAnyThread())
 	{
-		const UPoseSearchDatabase* CurrentDatabase = MotionMatchingState.CurrentSearchResult.Database;
+		const UPoseSearchDatabase* CurrentDatabase = MotionMatchingState.CurrentSearchResult.Database.Get();
 		const UAnimationAsset* CurrentAnimationAsset = AnimPlayers.IsEmpty() ? nullptr : AnimPlayers.First().GetAnimationAsset();
 
 		FString DebugInfo = FString::Printf(TEXT("bForceInterruptNextUpdate(%d)\n"), bForceInterruptNextUpdate);
diff --git a/Engine/Plugins/Experimental/Animation/PoseSearch/Source/Runtime/Private/PoseSearchAssetSampler.cpp b/Engine/Plugins/Experimental/Animation/PoseSearch/Source/Runtime/Private/PoseSearchAssetSampler.cpp
--- a/Engine/Plugins/Experimental/Animation/PoseSearch/Source/Runtime/Private/PoseSearchAssetSampler.cpp
+++ b/Engine/Plugins/Experimental/Animation/PoseSearch/Source/Runtime/Private/PoseSearchAssetSampler.cpp
@@ -247,15 +247,15 @@ FAnimationAssetSampler::FAnimationAssetSampler(TObjectPtr<const UAnimationAsset>
 
 void FAnimationAssetSampler::Init(TObjectPtr<const UAnimationAsset> InAnimationAsset, const FVector& InBlendParameters, int32 InRootTransformSamplingRate)
 {
-	AnimationAsset = InAnimationAsset;
+	AnimationAssetPtr = InAnimationAsset;
 	BlendParameters = InBlendParameters;
 	RootTransformSamplingRate = InRootTransformSamplingRate;
-	CachedPlayLength = GetPlayLength(AnimationAsset.Get(), BlendParameters);
+	CachedPlayLength = GetPlayLength(AnimationAssetPtr.Get(), BlendParameters);
 }
 
 bool FAnimationAssetSampler::IsInitialized() const
 {
-	return AnimationAsset != nullptr;
+	return AnimationAssetPtr != nullptr;
 }
 
 float FAnimationAssetSampler::GetPlayLength(const UAnimationAsset* AnimAsset, const FVector& BlendParameters)
@@ -291,13 +291,13 @@ float FAnimationAssetSampler::GetPlayLength(const UAnimationAsset* AnimAsset, co
 
 const UAnimationAsset* FAnimationAssetSampler::GetAsset() const
 {
-	return AnimationAsset.Get();
+	return AnimationAssetPtr.Get();
 }
 
 float FAnimationAssetSampler::ToRealTime(float NormalizedTime) const
 {
 	// Asset player time for blend spaces is normalized [0, 1] so we convert the sampling / animation time to asset time by multiplying it by CachedPlayLength
-	if (CachedPlayLength > UE_KINDA_SMALL_NUMBER && Cast<UBlendSpace>(AnimationAsset.Get()))
+	if (CachedPlayLength > UE_KINDA_SMALL_NUMBER && Cast<UBlendSpace>(AnimationAssetPtr.Get()))
 	{
 		check(NormalizedTime >= 0.f && NormalizedTime <= 1.f);
 		const float RealTime = NormalizedTime * CachedPlayLength;
@@ -310,7 +310,7 @@ float FAnimationAssetSampler::ToRealTime(float NormalizedTime) const
 float FAnimationAssetSampler::ToNormalizedTime(float RealTime) const
 {
 	// Asset player time for blend spaces is normalized [0, 1] so we convert the sampling / animation time to asset time by dividing it by CachedPlayLength
-	if (CachedPlayLength > UE_KINDA_SMALL_NUMBER && Cast<UBlendSpace>(AnimationAsset.Get()))
+	if (CachedPlayLength > UE_KINDA_SMALL_NUMBER && Cast<UBlendSpace>(AnimationAssetPtr.Get()))
 	{
 		const float NormalizedTime = RealTime / CachedPlayLength;
 
@@ -319,7 +319,7 @@ float FAnimationAssetSampler::ToNormalizedTime(float RealTime) const
 			return NormalizedTime;
 		}
 
-		UE_LOG(LogPoseSearch, Error, TEXT("FAnimationAssetSampler::ToNormalizedTime: requested RealTime %f is greater than CachedPlayLength %f for UBlendSpace %s!"), RealTime, CachedPlayLength, *AnimationAsset->GetName());
+		UE_LOG(LogPoseSearch, Error, TEXT("FAnimationAssetSampler::ToNormalizedTime: requested RealTime %f is greater than CachedPlayLength %f for UBlendSpace %s!"), RealTime, CachedPlayLength, *AnimationAssetPtr->GetName());
 		return FMath::Clamp(NormalizedTime, 0.f, 1.f);
 	}
 
@@ -333,23 +333,22 @@ float FAnimationAssetSampler::GetPlayLength() const
 
 bool FAnimationAssetSampler::IsLoopable() const
 {
-	if (const UAnimSequenceBase* SequenceBase = Cast<UAnimSequenceBase>(AnimationAsset.Get()))
+	if (const UAnimSequenceBase* SequenceBase = Cast<UAnimSequenceBase>(AnimationAssetPtr.Get()))
 	{
 		return SequenceBase->bLoop;
 	}
 
-	if (const UBlendSpace* BlendSpace = Cast<UBlendSpace>(AnimationAsset.Get()))
+	if (const UBlendSpace* BlendSpace = Cast<UBlendSpace>(AnimationAssetPtr.Get()))
 	{
 		return BlendSpace->bLoop;
 	}
 
-	checkNoEntry();
 	return false;
 }
 
 FTransform FAnimationAssetSampler::GetTotalRootTransform() const
 {
-	if (Cast<UBlendSpace>(AnimationAsset.Get()))
+	if (Cast<UBlendSpace>(AnimationAssetPtr.Get()))
 	{
 		const FTransform InitialRootTransform = ExtractBlendSpaceRootTrackTransform(0.f, AccumulatedRootTransform, RootTransformSamplingRate);
 		const FTransform LastRootTransform = ExtractBlendSpaceRootTrackTransform(CachedPlayLength, AccumulatedRootTransform, RootTransformSamplingRate);
@@ -357,7 +356,7 @@ FTransform FAnimationAssetSampler::GetTotalRootTransform() const
 		return TotalRootTransform;
 	}
 
-	if (const UAnimMontage* AnimMontage = Cast<UAnimMontage>(AnimationAsset.Get()))
+	if (const UAnimMontage* AnimMontage = Cast<UAnimMontage>(AnimationAssetPtr.Get()))
 	{
 		// @todo: add support for SlotName / multiple SlotAnimTracks
 		if (AnimMontage->SlotAnimTracks.Num() != 1)
@@ -381,7 +380,7 @@ FTransform FAnimationAssetSampler::GetTotalRootTransform() const
 
 void FAnimationAssetSampler::ExtractPose(const FAnimExtractContext& ExtractionCtx, FAnimationPoseData& OutAnimPoseData) const
 {
-	if (const UBlendSpace* BlendSpace = Cast<UBlendSpace>(AnimationAsset.Get()))
+	if (const UBlendSpace* BlendSpace = Cast<UBlendSpace>(AnimationAssetPtr.Get()))
 	{
 		TArray<FBlendSampleData> BlendSamples;
 		int32 TriangulationIndex = 0;
@@ -403,25 +402,26 @@ void FAnimationAssetSampler::ExtractPose(const FAnimExtractContext& ExtractionCt
 			BlendSpace->GetAnimationPose(BlendSamples, ExtractionCtx, OutAnimPoseData);
 		}
 	}
-	else if (const UAnimMontage* AnimMontage = Cast<UAnimMontage>(AnimationAsset.Get()))
+	else if (const UAnimMontage* AnimMontage = Cast<UAnimMontage>(AnimationAssetPtr.Get()))
 	{
 		// @todo: add support for SlotName / multiple SlotAnimTracks
 		if (AnimMontage->SlotAnimTracks.Num() != 1)
 		{
 			UE_LOG(LogPoseSearch, Error, TEXT("FAnimMontageSampler::ExtractPose: so far we support only montages with one SlotAnimTracks. %s has %d"), *AnimMontage->GetName(), AnimMontage->SlotAnimTracks.Num());
 			OutAnimPoseData.GetPose().ResetToRefPose();
-			return;
 		}
-
-		AnimMontage->SlotAnimTracks[0].AnimTrack.GetAnimationPose(OutAnimPoseData, ExtractionCtx);
+		else
+		{
+			AnimMontage->SlotAnimTracks[0].AnimTrack.GetAnimationPose(OutAnimPoseData, ExtractionCtx);
+		}
 	}
-	else if (const UAnimSequenceBase* SequenceBase = Cast<UAnimSequenceBase>(AnimationAsset.Get()))
+	else if (const UAnimSequenceBase* SequenceBase = Cast<UAnimSequenceBase>(AnimationAssetPtr.Get()))
 	{
 		SequenceBase->GetAnimationPose(OutAnimPoseData, ExtractionCtx);
 	}
 	else
 	{
-		checkNoEntry();
+		OutAnimPoseData.GetPose().ResetToRefPose();
 	}
 }
 
@@ -443,7 +443,7 @@ FTransform FAnimationAssetSampler::ExtractRootTransform(float Time) const
 {
 	FTransform RootTransform = FTransform::Identity;
 	
-	if (const UBlendSpace* BlendSpace = Cast<UBlendSpace>(AnimationAsset.Get()))
+	if (const UBlendSpace* BlendSpace = Cast<UBlendSpace>(AnimationAssetPtr.Get()))
 	{
 		if (IsLoopable())
 		{
@@ -485,7 +485,7 @@ FTransform FAnimationAssetSampler::ExtractRootTransform(float Time) const
 			}
 		}
 	}
-	else if (const UAnimMontage* AnimMontage = Cast<UAnimMontage>(AnimationAsset.Get()))
+	else if (const UAnimMontage* AnimMontage = Cast<UAnimMontage>(AnimationAssetPtr.Get()))
 	{
 		if (IsLoopable())
 		{
@@ -528,7 +528,7 @@ FTransform FAnimationAssetSampler::ExtractRootTransform(float Time) const
 			}
 		}
 	}
-	else if (const UAnimSequenceBase* SequenceBase = Cast<UAnimSequenceBase>(AnimationAsset.Get()))
+	else if (const UAnimSequenceBase* SequenceBase = Cast<UAnimSequenceBase>(AnimationAssetPtr.Get()))
 	{
 		if (IsLoopable())
 		{
@@ -571,17 +571,13 @@ FTransform FAnimationAssetSampler::ExtractRootTransform(float Time) const
 			}
 		}
 	}
-	else
-	{
-		checkNoEntry();
-	}
 
 	return RootTransform;
 }
 
 void FAnimationAssetSampler::Process(const FBoneContainer& BoneContainer)
 {
-	if (const UBlendSpace* BlendSpace = Cast<UBlendSpace>(AnimationAsset.Get()))
+	if (const UBlendSpace* BlendSpace = Cast<UBlendSpace>(AnimationAssetPtr.Get()))
 	{
 		ProcessRootTransform(BlendSpace, BlendParameters, CachedPlayLength, BoneContainer, RootTransformSamplingRate, IsLoopable(), AccumulatedRootTransform);
 	}
@@ -591,7 +587,7 @@ void FAnimationAssetSampler::ExtractPoseSearchNotifyStates(float Time, TFunction
 {
 	float SampleTime = Time;
 	FAnimNotifyContext NotifyContext;
-	if (const UBlendSpace* BlendSpace = Cast<UBlendSpace>(AnimationAsset.Get()))
+	if (const UBlendSpace* BlendSpace = Cast<UBlendSpace>(AnimationAssetPtr.Get()))
 	{
 		if (BlendSpace->NotifyTriggerMode == ENotifyTriggerMode::HighestWeightedAnimation)
 		{
@@ -621,7 +617,7 @@ void FAnimationAssetSampler::ExtractPoseSearchNotifyStates(float Time, TFunction
 			UE_LOG(LogPoseSearch, Error, TEXT("FAnimationAssetSampler::ExtractPoseSearchNotifyStates: Unsupported BlendSpace NotifyTriggerMode for '%s'"), *BlendSpace->GetName());
 		}
 	}
-	else if (const UAnimSequenceBase* SequenceBase = Cast<UAnimSequenceBase>(AnimationAsset.Get()))
+	else if (const UAnimSequenceBase* SequenceBase = Cast<UAnimSequenceBase>(AnimationAssetPtr.Get()))
 	{
 		// getting pose search notifies in an interval of size ExtractionInterval, centered on Time
 		SequenceBase->GetAnimNotifies(Time - (ExtractionInterval * 0.5f), ExtractionInterval, NotifyContext);
@@ -658,7 +654,7 @@ void FAnimationAssetSampler::ExtractPoseSearchNotifyStates(float Time, TFunction
 
 TConstArrayView<FAnimNotifyEvent> FAnimationAssetSampler::GetAllAnimNotifyEvents() const
 {
-	if (const UBlendSpace* BlendSpace = Cast<UBlendSpace>(AnimationAsset.Get()))
+	if (const UBlendSpace* BlendSpace = Cast<UBlendSpace>(AnimationAssetPtr.Get()))
 	{
 		if (BlendSpace->NotifyTriggerMode == ENotifyTriggerMode::HighestWeightedAnimation)
 		{
@@ -680,14 +676,10 @@ TConstArrayView<FAnimNotifyEvent> FAnimationAssetSampler::GetAllAnimNotifyEvents
 			UE_LOG(LogPoseSearch, Error, TEXT("FAnimationAssetSampler::ExtractPoseSearchNotifyStates: Unsupported BlendSpace NotifyTriggerMode for '%s'"), *BlendSpace->GetName());
 		}
 	}
-	else if (const UAnimSequenceBase* SequenceBase = Cast<UAnimSequenceBase>(AnimationAsset.Get()))
+	else if (const UAnimSequenceBase* SequenceBase = Cast<UAnimSequenceBase>(AnimationAssetPtr.Get()))
 	{
 		return SequenceBase->Notifies;
 	}
-	else
-	{
-		checkNoEntry();
-	}
 
 	return TConstArrayView<FAnimNotifyEvent>(); 
 }
diff --git a/Engine/Plugins/Experimental/Animation/PoseSearch/Source/Runtime/Private/PoseSearchDatabase.cpp b/Engine/Plugins/Experimental/Animation/PoseSearch/Source/Runtime/Private/PoseSearchDatabase.cpp
--- a/Engine/Plugins/Experimental/Animation/PoseSearch/Source/Runtime/Private/PoseSearchDatabase.cpp
+++ b/Engine/Plugins/Experimental/Animation/PoseSearch/Source/Runtime/Private/PoseSearchDatabase.cpp
@@ -379,8 +379,7 @@ bool FPoseSearchDatabaseBlendSpace::IsRootMotionEnabled() const
 	{
 		BlendSpace->ForEachImmutableSample([&bIsRootMotionUsedInBlendSpace](const FBlendSample& Sample)
 			{
-				const TObjectPtr<UAnimSequence> Sequence = Sample.Animation;
-
+				const UAnimSequence* Sequence = Sample.Animation.Get();
 				if (IsValid(Sequence) && Sequence->HasRootMotion())
 				{
 					bIsRootMotionUsedInBlendSpace = true;
diff --git a/Engine/Plugins/Experimental/Animation/PoseSearch/Source/Runtime/Private/PoseSearchDerivedData.cpp b/Engine/Plugins/Experimental/Animation/PoseSearch/Source/Runtime/Private/PoseSearchDerivedData.cpp
--- a/Engine/Plugins/Experimental/Animation/PoseSearch/Source/Runtime/Private/PoseSearchDerivedData.cpp
+++ b/Engine/Plugins/Experimental/Animation/PoseSearch/Source/Runtime/Private/PoseSearchDerivedData.cpp
@@ -1398,7 +1398,7 @@ void FPoseSearchDatabaseAsyncCacheTask::OnGetComplete(UE::DerivedData::FCacheGet
 				COOK_STAT(auto Timer = UsageStats.TimeSyncWork());
 
 				// collecting all the databases that need to be built to gather their FSearchIndexBase
-				TArray<TObjectPtr<const UPoseSearchDatabase>> IndexBaseDatabases;
+				TArray<const UPoseSearchDatabase*> IndexBaseDatabases;
 				IndexBaseDatabases.Add(Database.Get()); // the first one is always this Database
 				if (!IndexBaseDatabases[0])
 				{
@@ -1422,7 +1422,7 @@ void FPoseSearchDatabaseAsyncCacheTask::OnGetComplete(UE::DerivedData::FCacheGet
 				Schemas.AddDefaulted(IndexBaseDatabases.Num());
 				for (int32 IndexBaseIdx = 0; IndexBaseIdx < IndexBaseDatabases.Num(); ++IndexBaseIdx)
 				{
-					const UPoseSearchDatabase* IndexBaseDatabase = IndexBaseDatabases[IndexBaseIdx].Get();
+					const UPoseSearchDatabase* IndexBaseDatabase = IndexBaseDatabases[IndexBaseIdx];
 					check(IndexBaseDatabase);
 					FSearchIndexBase& SearchIndexBase = SearchIndexBases[IndexBaseIdx];
 					Schemas[IndexBaseIdx] = IndexBaseDatabase->Schema;
@@ -1633,7 +1633,7 @@ void FPoseSearchDatabaseAsyncCacheTask::OnGetComplete(UE::DerivedData::FCacheGet
 					return;
 				}
 
-				PreprocessSearchIndexKDTree(SearchIndex, IndexBaseDatabases[0].Get());
+				PreprocessSearchIndexKDTree(SearchIndex, IndexBaseDatabases[0]);
 				if (Owner.IsCanceled())
 				{
 					UE_LOG(LogPoseSearch, Log, TEXT("%s - %s BuildIndex Cancelled"), *LexToString(FullIndexKey.Hash), *IndexBaseDatabases[0]->GetName());
@@ -1651,7 +1651,7 @@ void FPoseSearchDatabaseAsyncCacheTask::OnGetComplete(UE::DerivedData::FCacheGet
 				}
 
 				const int32 RandomSeed = GetTypeHash(FullIndexKey.Hash);
-				PreprocessSearchIndexVPTree(SearchIndex, IndexBaseDatabases[0].Get(), RandomSeed);
+				PreprocessSearchIndexVPTree(SearchIndex, IndexBaseDatabases[0], RandomSeed);
 				if (Owner.IsCanceled())
 				{
 					UE_LOG(LogPoseSearch, Log, TEXT("%s - %s BuildIndex Cancelled"), *LexToString(FullIndexKey.Hash), *IndexBaseDatabases[0]->GetName());
@@ -1744,7 +1744,7 @@ void FPoseSearchDatabaseAsyncCacheTask::AddReferencedObjects(FReferenceCollector
 		// keeping around the assets for starting or in progress tasks
 		Collector.AddReferencedObject(Database);
 
-		for (TWeakObjectPtr<const UObject>& Dependency : DatabaseDependencies )
+		for (TWeakObjectPtr<const UObject>& Dependency : DatabaseDependencies)
 		{
 			Collector.AddReferencedObject(Dependency);
 		}
@@ -1982,7 +1982,7 @@ void FAsyncPoseSearchDatabasesManagement::AddReferencedObjects(FReferenceCollect
 // returns true if the index has been built and the Database updated correctly  
 bool FAsyncPoseSearchDatabasesManagement::RequestAsyncBuildIndex(const UPoseSearchDatabase* Database, ERequestAsyncBuildFlag Flag)
 {
-	if (!Database)
+	if (!IsValid(Database))
 	{
 		return false;
 	}
diff --git a/Engine/Plugins/Experimental/Animation/PoseSearch/Source/Runtime/Private/PoseSearchFeatureChannel_Velocity.cpp b/Engine/Plugins/Experimental/Animation/PoseSearch/Source/Runtime/Private/PoseSearchFeatureChannel_Velocity.cpp
--- a/Engine/Plugins/Experimental/Animation/PoseSearch/Source/Runtime/Private/PoseSearchFeatureChannel_Velocity.cpp
+++ b/Engine/Plugins/Experimental/Animation/PoseSearch/Source/Runtime/Private/PoseSearchFeatureChannel_Velocity.cpp
@@ -37,7 +37,6 @@ void UPoseSearchFeatureChannel_Velocity::BuildQuery(UE::PoseSearch::FSearchConte
 {
 	using namespace UE::PoseSearch;
 
-	check(InOutQuery.GetSchema());
 	const bool bIsRootBone = SchemaBoneIdx == RootSchemaBoneIdx;
 	if (bUseBlueprintQueryOverride)
 	{
diff --git a/Engine/Plugins/Experimental/Animation/PoseSearch/Source/Runtime/Private/PoseSearchNormalizationSet.cpp b/Engine/Plugins/Experimental/Animation/PoseSearch/Source/Runtime/Private/PoseSearchNormalizationSet.cpp
--- a/Engine/Plugins/Experimental/Animation/PoseSearch/Source/Runtime/Private/PoseSearchNormalizationSet.cpp
+++ b/Engine/Plugins/Experimental/Animation/PoseSearch/Source/Runtime/Private/PoseSearchNormalizationSet.cpp
@@ -5,13 +5,13 @@
 
 #if WITH_EDITORONLY_DATA
 
-void UPoseSearchNormalizationSet::AddUniqueDatabases(TArray<TObjectPtr<const UPoseSearchDatabase>>& UniqueDatabases) const
+void UPoseSearchNormalizationSet::AddUniqueDatabases(TArray<const UPoseSearchDatabase*>& UniqueDatabases) const
 {
 	for (const TObjectPtr<const UPoseSearchDatabase>& Database : Databases)
 	{
 		if (Database)
 		{
-			UniqueDatabases.AddUnique(Database);
+			UniqueDatabases.AddUnique(Database.Get());
 		}
 	}
 }
diff --git a/Engine/Plugins/Experimental/Animation/PoseSearch/Source/Runtime/Private/PoseSearchResult.cpp b/Engine/Plugins/Experimental/Animation/PoseSearch/Source/Runtime/Private/PoseSearchResult.cpp
--- a/Engine/Plugins/Experimental/Animation/PoseSearch/Source/Runtime/Private/PoseSearchResult.cpp
+++ b/Engine/Plugins/Experimental/Animation/PoseSearch/Source/Runtime/Private/PoseSearchResult.cpp
@@ -14,8 +14,8 @@ namespace UE::PoseSearch
 FFeatureVectorBuilder::FFeatureVectorBuilder(const UPoseSearchSchema* InSchema)
 {
 	check(InSchema && InSchema->IsValid());
-	Schema = InSchema;
-	Values.SetNumZeroed(Schema->SchemaCardinality);
+	SchemaPtr = InSchema;
+	Values.SetNumZeroed(SchemaPtr->SchemaCardinality);
 }
 
 //////////////////////////////////////////////////////////////////////////
diff --git a/Engine/Plugins/Experimental/Animation/PoseSearch/Source/Runtime/Public/PoseSearch/PoseSearchAssetSampler.h b/Engine/Plugins/Experimental/Animation/PoseSearch/Source/Runtime/Public/PoseSearch/PoseSearchAssetSampler.h
--- a/Engine/Plugins/Experimental/Animation/PoseSearch/Source/Runtime/Public/PoseSearch/PoseSearchAssetSampler.h
+++ b/Engine/Plugins/Experimental/Animation/PoseSearch/Source/Runtime/Public/PoseSearch/PoseSearchAssetSampler.h
@@ -52,7 +52,7 @@ struct POSESEARCH_API FAnimationAssetSampler
 	static float GetPlayLength(const UAnimationAsset* AnimAsset, const FVector& BlendParameters = FVector::ZeroVector);
 
 protected:
-	TObjectPtr<const UAnimationAsset> AnimationAsset;
+	TWeakObjectPtr<const UAnimationAsset> AnimationAssetPtr;
 
 	// members used to sample blend spaces only!
 	FVector BlendParameters = FVector::ZeroVector;
diff --git a/Engine/Plugins/Experimental/Animation/PoseSearch/Source/Runtime/Public/PoseSearch/PoseSearchNormalizationSet.h b/Engine/Plugins/Experimental/Animation/PoseSearch/Source/Runtime/Public/PoseSearch/PoseSearchNormalizationSet.h
--- a/Engine/Plugins/Experimental/Animation/PoseSearch/Source/Runtime/Public/PoseSearch/PoseSearchNormalizationSet.h
+++ b/Engine/Plugins/Experimental/Animation/PoseSearch/Source/Runtime/Public/PoseSearch/PoseSearchNormalizationSet.h
@@ -17,5 +17,5 @@ public:
 	UPROPERTY(EditAnywhere, BlueprintReadOnly, Category = "NormalizationSet")
 	TArray<TObjectPtr<const UPoseSearchDatabase>> Databases;
 
-	void AddUniqueDatabases(TArray<TObjectPtr<const UPoseSearchDatabase>>& UniqueDatabases) const;
+	void AddUniqueDatabases(TArray<const UPoseSearchDatabase*>& UniqueDatabases) const;
 };
\ No newline at end of file
diff --git a/Engine/Plugins/Experimental/Animation/PoseSearch/Source/Runtime/Public/PoseSearch/PoseSearchResult.h b/Engine/Plugins/Experimental/Animation/PoseSearch/Source/Runtime/Public/PoseSearch/PoseSearchResult.h
--- a/Engine/Plugins/Experimental/Animation/PoseSearch/Source/Runtime/Public/PoseSearch/PoseSearchResult.h
+++ b/Engine/Plugins/Experimental/Animation/PoseSearch/Source/Runtime/Public/PoseSearch/PoseSearchResult.h
@@ -21,14 +21,14 @@ struct FFeatureVectorBuilder
 {
 public:
 	explicit FFeatureVectorBuilder(const UPoseSearchSchema* Schema);
-	const UPoseSearchSchema* GetSchema() const { return Schema.Get(); }
+	const UPoseSearchSchema* GetSchema() const { return SchemaPtr.Get(); }
 
 	TArrayView<float> EditValues() { return Values; }
 	TConstArrayView<float> GetValues() const { return Values; }
 
 private:
 	TStackAlignedArray<float> Values;
-	TObjectPtr<const UPoseSearchSchema> Schema;
+	TWeakObjectPtr<const UPoseSearchSchema> SchemaPtr;
 };
 	
 struct FSearchResult
@@ -37,7 +37,7 @@ struct FSearchResult
 	FPoseSearchCost PoseCost;
 	int32 PoseIdx = INDEX_NONE;
 
-	TObjectPtr<const UPoseSearchDatabase> Database;
+	TWeakObjectPtr<const UPoseSearchDatabase> Database;
 
 	float AssetTime = 0.0f;
 
