motion matching
- added database editor Display Block Transition
- DDC fix for editing Block Transition notifies

diff --git a/Engine/Plugins/Experimental/Animation/PoseSearch/Source/Editor/Private/PoseSearchDatabaseEditorCommands.cpp b/Engine/Plugins/Experimental/Animation/PoseSearch/Source/Editor/Private/PoseSearchDatabaseEditorCommands.cpp
--- a/Engine/Plugins/Experimental/Animation/PoseSearch/Source/Editor/Private/PoseSearchDatabaseEditorCommands.cpp
+++ b/Engine/Plugins/Experimental/Animation/PoseSearch/Source/Editor/Private/PoseSearchDatabaseEditorCommands.cpp
@@ -11,6 +11,7 @@ namespace UE::PoseSearch
 		UI_COMMAND(ShowDisplayRootMotionSpeed, "Display Root Motion Speed", "Show root motion speed for the selected item", EUserInterfaceActionType::ToggleButton, FInputChord());
 		UI_COMMAND(ShowQuantizeAnimationToPoseData, "Quantize Animation To Pose Data", "Animations will not playing continuously, only using the associated pose quantized time", EUserInterfaceActionType::ToggleButton, FInputChord());
 		UI_COMMAND(ShowBones, "Show All Bones", "Debug Draw All Bones", EUserInterfaceActionType::ToggleButton, FInputChord());
+		UI_COMMAND(ShowDisplayBlockTransition, "Display Block Transition", "Show root motion samples (green) with block transition (red)", EUserInterfaceActionType::ToggleButton, FInputChord());
 	}
 }
 
diff --git a/Engine/Plugins/Experimental/Animation/PoseSearch/Source/Editor/Private/PoseSearchDatabasePreviewScene.cpp b/Engine/Plugins/Experimental/Animation/PoseSearch/Source/Editor/Private/PoseSearchDatabasePreviewScene.cpp
--- a/Engine/Plugins/Experimental/Animation/PoseSearch/Source/Editor/Private/PoseSearchDatabasePreviewScene.cpp
+++ b/Engine/Plugins/Experimental/Animation/PoseSearch/Source/Editor/Private/PoseSearchDatabasePreviewScene.cpp
@@ -63,11 +63,12 @@ namespace UE::PoseSearch
 		if (!ViewModel->GetPreviewActors().IsEmpty() && FAsyncPoseSearchDatabasesManagement::RequestAsyncBuildIndex(Database, ERequestAsyncBuildFlag::ContinueRequest))
 		{
 			const bool bDisplayRootMotionSpeed = ViewModel->IsDisplayRootMotionSpeedChecked();
+			const bool bDisplayBlockTransition = ViewModel->IsDisplayBlockTransitionChecked();
 			bool bDrawQueryVector = ViewModel->ShouldDrawQueryVector();
 
 			for (FDatabasePreviewActor& PreviewActor : ViewModel->GetPreviewActors())
 			{
-				bDrawQueryVector &= !PreviewActor.DrawPreviewActor(Database, bDisplayRootMotionSpeed, bDrawQueryVector ? ViewModel->GetQueryVector() : TConstArrayView<float>());
+				bDrawQueryVector &= !PreviewActor.DrawPreviewActor(Database, bDisplayRootMotionSpeed, bDisplayBlockTransition, bDrawQueryVector ? ViewModel->GetQueryVector() : TConstArrayView<float>());
 			}
 		}
 	}
diff --git a/Engine/Plugins/Experimental/Animation/PoseSearch/Source/Editor/Private/PoseSearchDatabaseViewModel.cpp b/Engine/Plugins/Experimental/Animation/PoseSearch/Source/Editor/Private/PoseSearchDatabaseViewModel.cpp
--- a/Engine/Plugins/Experimental/Animation/PoseSearch/Source/Editor/Private/PoseSearchDatabaseViewModel.cpp
+++ b/Engine/Plugins/Experimental/Animation/PoseSearch/Source/Editor/Private/PoseSearchDatabaseViewModel.cpp
@@ -177,7 +177,7 @@ void FDatabasePreviewActor::Destroy()
 	}
 }
 
-bool FDatabasePreviewActor::DrawPreviewActor(const UPoseSearchDatabase* PoseSearchDatabase, bool bDisplayRootMotionSpeed, TConstArrayView<float> QueryVector)
+bool FDatabasePreviewActor::DrawPreviewActor(const UPoseSearchDatabase* PoseSearchDatabase, bool bDisplayRootMotionSpeed, bool bDisplayBlockTransition, TConstArrayView<float> QueryVector)
 {
 	if (!PoseSearchDatabase->GetSearchIndex().IsValidPoseIndex(GetCurrentPoseIndex()))
 	{
@@ -199,7 +199,7 @@ bool FDatabasePreviewActor::DrawPreviewActor(const UPoseSearchDatabase* PoseSear
 		DrawParams.DrawFeatureVector(QueryVector);
 	}
 
-	if (bDisplayRootMotionSpeed)
+	if (bDisplayRootMotionSpeed || bDisplayBlockTransition)
 	{
 		// initializing SampledRootMotion if required
 		if (SampledRootMotion.IsEmpty())
@@ -232,7 +232,10 @@ bool FDatabasePreviewActor::DrawPreviewActor(const UPoseSearchDatabase* PoseSear
 				SampledRootMotionSpeed[0] = SampledRootMotionSpeed[1];
 			}
 		}
+	}
 
+	if (bDisplayRootMotionSpeed)
+	{
 		// drawing PreviewActor.SampledRootMotion
 		const int32 SampledRootMotionNum = SampledRootMotion.Num();
 		if (SampledRootMotionNum > 1)
@@ -255,6 +258,30 @@ bool FDatabasePreviewActor::DrawPreviewActor(const UPoseSearchDatabase* PoseSear
 		}
 	}
 
+	if (bDisplayBlockTransition)
+	{
+		const FSearchIndex& SearchIndex = PoseSearchDatabase->GetSearchIndex();
+		const FSearchIndexAsset& IndexAsset = SearchIndex.Assets[IndexAssetIndex];
+
+		const int NumPoses = IndexAsset.GetNumPoses();
+		if (NumPoses == SampledRootMotion.Num())
+		{
+			for (int32 Index = 0; Index < NumPoses; ++Index)
+			{
+				const int32 IndexAssetPoseIdx = Index + IndexAsset.GetFirstPoseIdx();
+				if (SearchIndex.PoseMetadata[IndexAssetPoseIdx].IsBlockTransition())
+				{
+					DrawParams.DrawPoint(SampledRootMotion[Index], FColor::Red);
+				}
+				else
+				{
+					DrawParams.DrawPoint(SampledRootMotion[Index], FColor::Green);
+				}
+			}
+		}
+	}
+
+
 #if ENABLE_ANIM_DEBUG
 	const float DebugDrawSamplerSize = CVarDatabasePreviewDebugDrawSamplerSize.GetValueOnAnyThread();
 	if (DebugDrawSamplerSize > UE_KINDA_SMALL_NUMBER)
diff --git a/Engine/Plugins/Experimental/Animation/PoseSearch/Source/Editor/Private/SPoseSearchDatabaseViewport.cpp b/Engine/Plugins/Experimental/Animation/PoseSearch/Source/Editor/Private/SPoseSearchDatabaseViewport.cpp
--- a/Engine/Plugins/Experimental/Animation/PoseSearch/Source/Editor/Private/SPoseSearchDatabaseViewport.cpp
+++ b/Engine/Plugins/Experimental/Animation/PoseSearch/Source/Editor/Private/SPoseSearchDatabaseViewport.cpp
@@ -61,6 +61,17 @@ namespace UE::PoseSearch
 		return AssetEditorPtr.Pin()->GetViewModelSharedPtr().ToSharedRef()->IsShowBonesChecked();
 	}
 
+	void SDatabaseViewport::OnToggleDisplayBlockTransition()
+	{
+		AssetEditorPtr.Pin()->GetViewModelSharedPtr().ToSharedRef()->ToggleDisplayBlockTransition();
+		Invalidate();
+	}
+
+	bool SDatabaseViewport::IsDisplayBlockTransitionChecked() const
+	{
+		return AssetEditorPtr.Pin()->GetViewModelSharedPtr().ToSharedRef()->IsDisplayBlockTransitionChecked();
+	}
+
 	void SDatabaseViewport::BindCommands()
 	{
 		SEditorViewport::BindCommands();
@@ -85,6 +96,11 @@ namespace UE::PoseSearch
 			FCanExecuteAction(),
 			FIsActionChecked::CreateSP(this, &SDatabaseViewport::IsShowBonesChecked));
 		
+		CommandList->MapAction(
+			Commands.ShowDisplayBlockTransition,
+			FExecuteAction::CreateSP(this, &SDatabaseViewport::OnToggleDisplayBlockTransition),
+			FCanExecuteAction(),
+			FIsActionChecked::CreateSP(this, &SDatabaseViewport::IsDisplayBlockTransitionChecked));
 	}
 
 	TSharedRef<FEditorViewportClient> SDatabaseViewport::MakeEditorViewportClient()
diff --git a/Engine/Plugins/Experimental/Animation/PoseSearch/Source/Editor/Private/SPoseSearchDatabaseViewportToolbar.cpp b/Engine/Plugins/Experimental/Animation/PoseSearch/Source/Editor/Private/SPoseSearchDatabaseViewportToolbar.cpp
--- a/Engine/Plugins/Experimental/Animation/PoseSearch/Source/Editor/Private/SPoseSearchDatabaseViewportToolbar.cpp
+++ b/Engine/Plugins/Experimental/Animation/PoseSearch/Source/Editor/Private/SPoseSearchDatabaseViewportToolbar.cpp
@@ -34,6 +34,7 @@ namespace UE::PoseSearch
 			ShowMenuBuilder.AddMenuEntry(Commands.ShowDisplayRootMotionSpeed);
 			ShowMenuBuilder.AddMenuEntry(Commands.ShowQuantizeAnimationToPoseData);
 			ShowMenuBuilder.AddMenuEntry(Commands.ShowBones);
+			ShowMenuBuilder.AddMenuEntry(Commands.ShowDisplayBlockTransition);
 			ShowMenuBuilder.EndSection();
 		}
 
diff --git a/Engine/Plugins/Experimental/Animation/PoseSearch/Source/Editor/Public/PoseSearchDatabaseEditorCommands.h b/Engine/Plugins/Experimental/Animation/PoseSearch/Source/Editor/Public/PoseSearchDatabaseEditorCommands.h
--- a/Engine/Plugins/Experimental/Animation/PoseSearch/Source/Editor/Public/PoseSearchDatabaseEditorCommands.h
+++ b/Engine/Plugins/Experimental/Animation/PoseSearch/Source/Editor/Public/PoseSearchDatabaseEditorCommands.h
@@ -27,6 +27,7 @@ namespace UE::PoseSearch
 		TSharedPtr<FUICommandInfo> ShowDisplayRootMotionSpeed;
 		TSharedPtr<FUICommandInfo> ShowQuantizeAnimationToPoseData;
 		TSharedPtr<FUICommandInfo> ShowBones;
+		TSharedPtr<FUICommandInfo> ShowDisplayBlockTransition;
 	};
 }
 
diff --git a/Engine/Plugins/Experimental/Animation/PoseSearch/Source/Editor/Public/PoseSearchDatabaseViewModel.h b/Engine/Plugins/Experimental/Animation/PoseSearch/Source/Editor/Public/PoseSearchDatabaseViewModel.h
--- a/Engine/Plugins/Experimental/Animation/PoseSearch/Source/Editor/Public/PoseSearchDatabaseViewModel.h
+++ b/Engine/Plugins/Experimental/Animation/PoseSearch/Source/Editor/Public/PoseSearchDatabaseViewModel.h
@@ -30,7 +30,7 @@ namespace UE::PoseSearch
 	public:
 		bool SpawnPreviewActor(UWorld* World, const UPoseSearchDatabase* PoseSearchDatabase, int32 IndexAssetIdx, int32 PoseIdxForTimeOffset = INDEX_NONE);
 		void UpdatePreviewActor(const UPoseSearchDatabase* PoseSearchDatabase, float PlayTime, bool bQuantizeAnimationToPoseData);
-		bool DrawPreviewActor(const UPoseSearchDatabase* PoseSearchDatabase, bool bDisplayRootMotionSpeed, TConstArrayView<float> QueryVector);
+		bool DrawPreviewActor(const UPoseSearchDatabase* PoseSearchDatabase, bool bDisplayRootMotionSpeed, bool bDisplayBlockTransition, TConstArrayView<float> QueryVector);
 		void Destroy();
 
 		const UDebugSkelMeshComponent* GetDebugSkelMeshComponent() const;
@@ -106,6 +106,9 @@ namespace UE::PoseSearch
 		void ToggleShowBones() { bShowBones = !bShowBones; }
 		bool IsShowBonesChecked() const { return bShowBones; }
 
+		void ToggleDisplayBlockTransition() { bDisplayBlockTransition = !bDisplayBlockTransition; }
+		bool IsDisplayBlockTransitionChecked() const { return bDisplayBlockTransition; }
+
 		void AddSequenceToDatabase(UAnimSequence* AnimSequence);
 		void AddBlendSpaceToDatabase(UBlendSpace* BlendSpace);
 		void AddAnimCompositeToDatabase(UAnimComposite* AnimComposite);
@@ -165,6 +168,7 @@ namespace UE::PoseSearch
 		bool bDisplayRootMotionSpeed = false;
 		bool bQuantizeAnimationToPoseData = false;
 		bool bShowBones = false;
+		bool bDisplayBlockTransition = false;
 
 		int32 SelectedActorIndexAssetIndex = INDEX_NONE;
 	};
diff --git a/Engine/Plugins/Experimental/Animation/PoseSearch/Source/Editor/Public/SPoseSearchDatabaseViewport.h b/Engine/Plugins/Experimental/Animation/PoseSearch/Source/Editor/Public/SPoseSearchDatabaseViewport.h
--- a/Engine/Plugins/Experimental/Animation/PoseSearch/Source/Editor/Public/SPoseSearchDatabaseViewport.h
+++ b/Engine/Plugins/Experimental/Animation/PoseSearch/Source/Editor/Public/SPoseSearchDatabaseViewport.h
@@ -58,6 +58,8 @@ namespace UE::PoseSearch
 		bool IsQuantizeAnimationToPoseDataChecked() const;
 		void OnToggleShowBones();
 		bool IsShowBonesChecked() const;
+		void OnToggleDisplayBlockTransition();
+		bool IsDisplayBlockTransitionChecked() const;
 
 		/** The viewport toolbar */
 		TSharedPtr<SPoseSearchDatabaseViewportToolBar> ViewportToolbar;
diff --git a/Engine/Plugins/Experimental/Animation/PoseSearch/Source/Runtime/Private/PoseSearchDerivedData.cpp b/Engine/Plugins/Experimental/Animation/PoseSearch/Source/Runtime/Private/PoseSearchDerivedData.cpp
--- a/Engine/Plugins/Experimental/Animation/PoseSearch/Source/Runtime/Private/PoseSearchDerivedData.cpp
+++ b/Engine/Plugins/Experimental/Animation/PoseSearch/Source/Runtime/Private/PoseSearchDerivedData.cpp
@@ -1098,16 +1098,17 @@ struct FPoseSearchDatabaseAsyncCacheTask
 	{
 		Notstarted,	// key generation failed (not all the asset has been post loaded). It'll be retried to StartNewRequestIfNeeded the next Update
 		Prestarted,	// key has been successfully generated and we kicked the DDC get
-		Cancelled,	// the task has been cancelled
+		PreCancelled,	// the task has been requested to be cancelled
+		Cancelled,		// the task cancellation has been finalized
 		Ended,		// the task has ended successfully
 		Failed		// the task has ended unsuccessfully
 	};
 
 	FPoseSearchDatabaseAsyncCacheTask(UPoseSearchDatabase* InDatabase, bool bPerformConditionalPostLoadIfRequired);
 	void StartNewRequestIfNeeded(bool bPerformConditionalPostLoadIfRequired);
-	bool CancelIfDependsOn(const UObject* Object);
 	void Update(FCriticalSection& OuterMutex);
 	void Wait(FCriticalSection& OuterMutex);
+	void PreCancelIfDependsOn(const UObject* Object);
 	void Cancel();
 	bool Poll() const;
 	void AddReferencedObjects(FReferenceCollector& Collector);
@@ -1181,6 +1182,7 @@ void FPoseSearchDatabaseAsyncCacheTask::StartNewRequestIfNeeded(bool bPerformCon
 	using namespace UE::DerivedData;
 
 	check(IsInGameThread());
+	check(GetState() != EState::PreCancelled);
 
 	// making sure there are no active requests
 	// Owner.Cancel must be performed before SearchIndex.Reset() in case any task is flying (launched by Owner.LaunchTask)
@@ -1226,6 +1228,28 @@ void FPoseSearchDatabaseAsyncCacheTask::StartNewRequestIfNeeded(bool bPerformCon
 	}
 }
 
+// it cancels and waits for the task to be done and set the state to PreCancelled, so no other new requests can start until the task gets cancelled
+void FPoseSearchDatabaseAsyncCacheTask::PreCancelIfDependsOn(const UObject* Object)
+{
+	check(IsInGameThread());
+
+	if (Object)
+	{
+		// DatabaseDependencies is updated only in StartNewRequestIfNeeded when there are no active requests, so it's thread safe to access it 
+		if (DatabaseDependencies.Contains(Object))
+		{
+			// Database can be null if the task was Ended/Failed and Database was already garbage collected, but Tick hasn't been called yet
+			FString DatabaseName = IsValid() ? *Database->GetName() : TEXT("Garbage Collected Database");
+			UE_LOG(LogPoseSearch, Log, TEXT("%s - %s PreCancelled because of %s"), *LexToString(DerivedDataKey), *DatabaseName, *Object->GetName());
+
+			// Owner.Cancel must be performed before SearchIndex.Reset() in case any task is flying (launched by Owner.LaunchTask)
+			Owner.Cancel();
+
+			SetState(EState::PreCancelled);
+		}
+	}
+}
+
 // it cancels and waits for the task to be done and reset the local SearchIndex. SetState to Cancelled
 void FPoseSearchDatabaseAsyncCacheTask::Cancel()
 {
@@ -1243,24 +1267,6 @@ void FPoseSearchDatabaseAsyncCacheTask::Cancel()
 	SetState(EState::Cancelled);
 }
 
-bool FPoseSearchDatabaseAsyncCacheTask::CancelIfDependsOn(const UObject* Object)
-{
-	if (Object)
-	{
-		// DatabaseDependencies is updated only in StartNewRequestIfNeeded when there are no active requests, so it's thread safe to access it 
-		if (DatabaseDependencies.Contains(Object))
-		{
-			// Database can be null if the task was Ended/Failed and Database was already garbage collected, but Tick hasn't been called yet
-			FString DatabaseName = IsValid() ? *Database->GetName() : TEXT("Garbage Collected Database");
-			UE_LOG(LogPoseSearch, Log, TEXT("%s - %s Cancelled because of %s"), *LexToString(DerivedDataKey), *DatabaseName, *Object->GetName());
-
-			Cancel();
-			return true;
-		}
-	}
-	return false;
-}
-
 void FPoseSearchDatabaseAsyncCacheTask::Update(FCriticalSection& OuterMutex)
 {
 	check(IsInGameThread());
@@ -1278,12 +1284,15 @@ void FPoseSearchDatabaseAsyncCacheTask::Update(FCriticalSection& OuterMutex)
 		Wait(OuterMutex);
 	}
 
+	if (GetState() != EState::PreCancelled)
+	{
 	if (bBroadcastOnDerivedDataRebuild)
 	{
 		Database->NotifyDerivedDataRebuild();
 		bBroadcastOnDerivedDataRebuild = false;
 	}
 }
+}
 
 // it waits for the task to be done and SetSearchIndex on the database. SetState to Ended/Failed
 void FPoseSearchDatabaseAsyncCacheTask::Wait(FCriticalSection& OuterMutex)
@@ -1771,6 +1780,8 @@ FAsyncPoseSearchDatabasesManagement::FAsyncPoseSearchDatabasesManagement()
 	OnObjectModifiedHandle = FCoreUObjectDelegates::OnObjectModified.AddRaw(this, &FAsyncPoseSearchDatabasesManagement::OnObjectModified);
 	OnObjectTransactedHandle = FCoreUObjectDelegates::OnObjectTransacted.AddRaw(this, &FAsyncPoseSearchDatabasesManagement::OnObjectTransacted);
 	OnPackageReloadedHandle = FCoreUObjectDelegates::OnPackageReloaded.AddRaw(this, &FAsyncPoseSearchDatabasesManagement::OnPackageReloaded);
+	OnPreObjectPropertyChangedHandle = FCoreUObjectDelegates::OnPreObjectPropertyChanged.AddRaw(this, &FAsyncPoseSearchDatabasesManagement::OnPreObjectPropertyChanged);
+	OnObjectPropertyChangedHandle = FCoreUObjectDelegates::OnObjectPropertyChanged.AddRaw(this, &FAsyncPoseSearchDatabasesManagement::OnObjectPropertyChanged);
 
 	FCoreDelegates::OnPreExit.AddRaw(this, &FAsyncPoseSearchDatabasesManagement::Shutdown);
 }
@@ -1819,6 +1830,8 @@ void FAsyncPoseSearchDatabasesManagement::SynchronizeDatabases()
 {
 	FScopeLock Lock(&Mutex);
 
+	if (!DatabasesToSynchronize.IsEmpty())
+	{
 	// copying DatabasesToSynchronize because modifying the database will call OnObjectModified that could populate DatabasesToSynchronize again
 	const TDatabasesToSynchronize DatabasesToSynchronizeCopy = DatabasesToSynchronize;
 	DatabasesToSynchronize.Reset();
@@ -1840,9 +1853,28 @@ void FAsyncPoseSearchDatabasesManagement::SynchronizeDatabases()
 		}
 	}
 }
+}
 
 // we're listening to OnObjectModified to cancel any pending Task indexing databases depending from Object to avoid multi threading issues
 void FAsyncPoseSearchDatabasesManagement::OnObjectModified(UObject* Object)
+{
+	PreModified(Object);
+}
+
+void FAsyncPoseSearchDatabasesManagement::ClearPreCancelled()
+{
+	// iterating backwards because of the possible RemoveAtSwap 
+	for (int32 TaskIndex = Tasks.Num() - 1; TaskIndex >= 0; --TaskIndex)
+	{
+		if (Tasks[TaskIndex]->GetState() == FPoseSearchDatabaseAsyncCacheTask::EState::PreCancelled)
+		{
+			UE_LOG(LogPoseSearch, Log, TEXT("%s - %s Removed because it was PreCancelled"), *LexToString(Tasks[TaskIndex]->GetDerivedDataKey()), *Tasks[TaskIndex]->GetDatabase()->GetName());
+			Tasks.RemoveAtSwap(TaskIndex, 1, false);
+		}
+	}
+}
+
+void FAsyncPoseSearchDatabasesManagement::PreModified(UObject* Object)
 {
 	check(IsInGameThread());
 
@@ -1851,26 +1883,30 @@ void FAsyncPoseSearchDatabasesManagement::OnObjectModified(UObject* Object)
 	// iterating backwards because of the possible RemoveAtSwap
 	for (int32 TaskIndex = Tasks.Num() - 1; TaskIndex >= 0; --TaskIndex)
 	{
-		if (Tasks[TaskIndex]->CancelIfDependsOn(Object))
-		{
-			Tasks.RemoveAtSwap(TaskIndex, 1, false);
-		}
+		Tasks[TaskIndex]->PreCancelIfDependsOn(Object);
 	}
 
 	// collecting databases to synchronize prior modifying the Object
 	CollectDatabasesToSynchronize(Object);
 }
 
-void FAsyncPoseSearchDatabasesManagement::OnObjectTransacted(UObject* Object, const FTransactionObjectEvent& TransactionObjectEvent)
+void FAsyncPoseSearchDatabasesManagement::PostModified(UObject* Object)
 {
 	check(IsInGameThread());
 
 	FScopeLock Lock(&Mutex);
 
-	// collecting databases to synchronize on Object transacted, and merging the results with the OnObjectModified collection
+	// collecting databases to synchronize, and merging the results with the PreModified collection
 	CollectDatabasesToSynchronize(Object);
 
 	SynchronizeDatabases();
+
+	ClearPreCancelled();
+}
+
+void FAsyncPoseSearchDatabasesManagement::OnObjectTransacted(UObject* Object, const FTransactionObjectEvent& TransactionObjectEvent)
+{
+	PostModified(Object);
 }
 
 void FAsyncPoseSearchDatabasesManagement::OnPackageReloaded(const EPackageReloadPhase InPackageReloadPhase, FPackageReloadedEvent* InPackageReloadedEvent)
@@ -1896,6 +1932,16 @@ void FAsyncPoseSearchDatabasesManagement::OnPackageReloaded(const EPackageReload
 	}
 }
 
+void FAsyncPoseSearchDatabasesManagement::OnPreObjectPropertyChanged(UObject* InObject, const FEditPropertyChain& InPropertyChain)
+{
+	PreModified(InObject);
+}
+
+void FAsyncPoseSearchDatabasesManagement::OnObjectPropertyChanged(UObject* InObject, FPropertyChangedEvent& InPropertyChangedEvent)
+{
+	PostModified(InObject);
+}
+
 void FAsyncPoseSearchDatabasesManagement::Shutdown()
 {
 	FScopeLock Lock(&Mutex);
@@ -1910,6 +1956,12 @@ void FAsyncPoseSearchDatabasesManagement::Shutdown()
 
 	FCoreUObjectDelegates::OnPackageReloaded.Remove(OnPackageReloadedHandle);
 	OnPackageReloadedHandle.Reset();
+
+	FCoreUObjectDelegates::OnPreObjectPropertyChanged.Remove(OnPreObjectPropertyChangedHandle);
+	OnPreObjectPropertyChangedHandle.Reset();
+
+	FCoreUObjectDelegates::OnObjectPropertyChanged.Remove(OnObjectPropertyChangedHandle);
+	OnObjectPropertyChangedHandle.Reset();
 }
 
 void FAsyncPoseSearchDatabasesManagement::Tick(float DeltaTime)
diff --git a/Engine/Plugins/Experimental/Animation/PoseSearch/Source/Runtime/Public/PoseSearch/PoseSearchDerivedData.h b/Engine/Plugins/Experimental/Animation/PoseSearch/Source/Runtime/Public/PoseSearch/PoseSearchDerivedData.h
--- a/Engine/Plugins/Experimental/Animation/PoseSearch/Source/Runtime/Public/PoseSearch/PoseSearchDerivedData.h
+++ b/Engine/Plugins/Experimental/Animation/PoseSearch/Source/Runtime/Public/PoseSearch/PoseSearchDerivedData.h
@@ -37,9 +37,14 @@ namespace UE::PoseSearch
 		void OnObjectModified(UObject* Object);
 		void OnObjectTransacted(UObject* Object, const FTransactionObjectEvent& TransactionObjectEvent);
 		void OnPackageReloaded(const EPackageReloadPhase InPackageReloadPhase, FPackageReloadedEvent* InPackageReloadedEvent);
+		void OnPreObjectPropertyChanged(UObject* InObject, const FEditPropertyChain& InPropertyChain);
+		void OnObjectPropertyChanged(UObject* InObject, FPropertyChangedEvent& InPropertyChangedEvent);
 
 		void Shutdown();
 		void StartQueuedTasks(int32 MaxActiveTasks);
+		void PreModified(UObject* Object);
+		void PostModified(UObject* Object);
+		void ClearPreCancelled();
 
 		// Begin FTickableGameObject
 		virtual void Tick(float DeltaTime) override;
@@ -70,6 +75,8 @@ namespace UE::PoseSearch
 		FDelegateHandle OnObjectModifiedHandle;
 		FDelegateHandle OnObjectTransactedHandle;
 		FDelegateHandle OnPackageReloadedHandle;
+		FDelegateHandle OnPreObjectPropertyChangedHandle;
+		FDelegateHandle OnObjectPropertyChangedHandle;
 
 		static FCriticalSection Mutex;
 	};
