motion matching - removing unnecessary check

diff --git a/Engine/Plugins/Experimental/Animation/PoseSearch/Source/Runtime/Private/PoseSearchDerivedData.cpp b/Engine/Plugins/Experimental/Animation/PoseSearch/Source/Runtime/Private/PoseSearchDerivedData.cpp
--- a/Engine/Plugins/Experimental/Animation/PoseSearch/Source/Runtime/Private/PoseSearchDerivedData.cpp
+++ b/Engine/Plugins/Experimental/Animation/PoseSearch/Source/Runtime/Private/PoseSearchDerivedData.cpp
@@ -1182,7 +1182,6 @@ void FPoseSearchDatabaseAsyncCacheTask::StartNewRequestIfNeeded(bool bPerformCon
 	using namespace UE::DerivedData;
 
 	check(IsInGameThread());
-	check(GetState() != EState::PreCancelled);
 
 	// making sure there are no active requests
 	// Owner.Cancel must be performed before SearchIndex.Reset() in case any task is flying (launched by Owner.LaunchTask)
@@ -1286,13 +1285,13 @@ void FPoseSearchDatabaseAsyncCacheTask::Update(FCriticalSection& OuterMutex)
 
 	if (GetState() != EState::PreCancelled)
 	{
-	if (bBroadcastOnDerivedDataRebuild)
-	{
-		Database->NotifyDerivedDataRebuild();
-		bBroadcastOnDerivedDataRebuild = false;
+		if (bBroadcastOnDerivedDataRebuild)
+		{
+			Database->NotifyDerivedDataRebuild();
+			bBroadcastOnDerivedDataRebuild = false;
+		}
 	}
 }
-}
 
 // it waits for the task to be done and SetSearchIndex on the database. SetState to Ended/Failed
 void FPoseSearchDatabaseAsyncCacheTask::Wait(FCriticalSection& OuterMutex)
