motion matching - disabling EPoseSearchMode::VPTree until production ready

diff --git a/Engine/Plugins/Experimental/Animation/PoseSearch/Source/Runtime/Public/PoseSearch/PoseSearchDatabase.h b/Engine/Plugins/Experimental/Animation/PoseSearch/Source/Runtime/Public/PoseSearch/PoseSearchDatabase.h
--- a/Engine/Plugins/Experimental/Animation/PoseSearch/Source/Runtime/Public/PoseSearch/PoseSearchDatabase.h
+++ b/Engine/Plugins/Experimental/Animation/PoseSearch/Source/Runtime/Public/PoseSearch/PoseSearchDatabase.h
@@ -33,8 +33,8 @@ enum class EPoseSearchMode : int32
 	// Optimized search mode: the database projects the poses into a PCA space using only the most significant "NumberOfPrincipalComponents" dimensions, and construct a kdtree to facilitate the search.
 	PCAKDTree,
 
-	// Optimized search mode using a vantage point tree
-	VPTree,
+	// Optimized search mode using a vantage point tree (disabled/Hidden until production ready)
+	VPTree UMETA(Hidden),
 
 	Num UMETA(Hidden),
 	Invalid = Num UMETA(Hidden)
