Added tags to pose search database

diff --git a/Engine/Plugins/Experimental/Animation/PoseSearch/Source/Runtime/Public/PoseSearch/PoseSearchDatabase.h b/Engine/Plugins/Experimental/Animation/PoseSearch/Source/Runtime/Public/PoseSearch/PoseSearchDatabase.h
--- a/Engine/Plugins/Experimental/Animation/PoseSearch/Source/Runtime/Public/PoseSearch/PoseSearchDatabase.h
+++ b/Engine/Plugins/Experimental/Animation/PoseSearch/Source/Runtime/Public/PoseSearch/PoseSearchDatabase.h
@@ -280,6 +280,10 @@ public:
 	UPROPERTY(EditAnywhere, Category="Database")
 	TArray<FInstancedStruct> AnimationAssets;
 
+	/** Array of tags that can be used as metadata. */
+	UPROPERTY(EditAnywhere, BlueprintReadOnly, Category = "Database")
+	TArray<FName> Tags;
+
 #if WITH_EDITORONLY_DATA
 	// This optional asset defines a list of databases you want to normalize together. Without it, it would be difficult to compare costs from separately normalized databases containing different types of animation,
 	// like only idles versus only runs animations, given that the range of movement would be dramatically different.
