motion matching - warning for invalid searches

diff --git a/Engine/Plugins/Experimental/Animation/PoseSearch/Source/Runtime/Private/PoseSearchDatabase.cpp b/Engine/Plugins/Experimental/Animation/PoseSearch/Source/Runtime/Private/PoseSearchDatabase.cpp
--- a/Engine/Plugins/Experimental/Animation/PoseSearch/Source/Runtime/Private/PoseSearchDatabase.cpp
+++ b/Engine/Plugins/Experimental/Animation/PoseSearch/Source/Runtime/Private/PoseSearchDatabase.cpp
@@ -947,7 +947,7 @@ static inline void EvaluatePoseKernel(UE::PoseSearch::FSearchResult& Result, con
 	}
 }
 
-FPoseSearchCost UPoseSearchDatabase::SearchContinuingPose(UE::PoseSearch::FSearchContext& SearchContext) const
+UE::PoseSearch::FSearchResult UPoseSearchDatabase::SearchContinuingPose(UE::PoseSearch::FSearchContext& SearchContext) const
 {
 	QUICK_SCOPE_CYCLE_COUNTER(STAT_PoseSearch_ContinuingPose);
 
@@ -955,12 +955,12 @@ FPoseSearchCost UPoseSearchDatabase::SearchContinuingPose(UE::PoseSearch::FSearc
 
 	check(SearchContext.GetCurrentResult().Database.Get() == this);
 
-	FPoseSearchCost ContinuingPoseCost;
+	FSearchResult Result;
 
 #if WITH_EDITOR
 	if (!FAsyncPoseSearchDatabasesManagement::RequestAsyncBuildIndex(this, ERequestAsyncBuildFlag::ContinueRequest))
 	{
-		return ContinuingPoseCost;
+		return Result;
 	}
 #endif // WITH_EDITOR
 
@@ -998,20 +998,24 @@ FPoseSearchCost UPoseSearchDatabase::SearchContinuingPose(UE::PoseSearch::FSearc
 		// is the data padded at 16 bytes (and 16 bytes aligned by construction)?
 		if (NumDimensions % 4 == 0)
 		{
-			ContinuingPoseCost = SearchIndex.CompareAlignedPoses(ContinuingPoseIdx, UpdatedContinuingPoseCostBias, PoseValues, SearchContext.GetOrBuildQuery(Schema).GetValues());
+			Result.PoseCost = SearchIndex.CompareAlignedPoses(ContinuingPoseIdx, UpdatedContinuingPoseCostBias, PoseValues, SearchContext.GetOrBuildQuery(Schema).GetValues());
 		}
 		// data is not 16 bytes padded
 		else
 		{
-			ContinuingPoseCost = SearchIndex.ComparePoses(ContinuingPoseIdx, UpdatedContinuingPoseCostBias, PoseValues, SearchContext.GetOrBuildQuery(Schema).GetValues());
+			Result.PoseCost = SearchIndex.ComparePoses(ContinuingPoseIdx, UpdatedContinuingPoseCostBias, PoseValues, SearchContext.GetOrBuildQuery(Schema).GetValues());
 		}
 
+		Result.AssetTime = SearchContext.GetCurrentResult().AssetTime;
+		Result.PoseIdx = PoseIdx;
+		Result.Database = this;
+
 #if UE_POSE_SEARCH_TRACE_ENABLED
-		SearchContext.Track(this, ContinuingPoseIdx, EPoseCandidateFlags::Valid_ContinuingPose, ContinuingPoseCost);
+		SearchContext.Track(this, ContinuingPoseIdx, EPoseCandidateFlags::Valid_ContinuingPose, Result.PoseCost);
 #endif // UE_POSE_SEARCH_TRACE_ENABLED
 	}
 
-	return ContinuingPoseCost;
+	return Result;
 }
 
 UE::PoseSearch::FSearchResult UPoseSearchDatabase::SearchPCAKDTree(UE::PoseSearch::FSearchContext& SearchContext) const
diff --git a/Engine/Plugins/Experimental/Animation/PoseSearch/Source/Runtime/Private/PoseSearchLibrary.cpp b/Engine/Plugins/Experimental/Animation/PoseSearch/Source/Runtime/Private/PoseSearchLibrary.cpp
--- a/Engine/Plugins/Experimental/Animation/PoseSearch/Source/Runtime/Private/PoseSearchLibrary.cpp
+++ b/Engine/Plugins/Experimental/Animation/PoseSearch/Source/Runtime/Private/PoseSearchLibrary.cpp
@@ -315,7 +315,7 @@ void UPoseSearchLibrary::UpdateMotionMatchingState(
 		FSearchResult SearchResult;
 		if (!SearchContext.IsForceInterrupt() && bCanAdvance)
 		{
-			SearchResult.PoseCost = SearchContext.GetCurrentResult().Database->SearchContinuingPose(SearchContext);
+			SearchResult = SearchContext.GetCurrentResult().Database->SearchContinuingPose(SearchContext);
 			SearchContext.UpdateCurrentBestCost(SearchResult.PoseCost);
 		}
 
@@ -341,6 +341,32 @@ void UPoseSearchLibrary::UpdateMotionMatchingState(
 		}
 #endif // UE_POSE_SEARCH_TRACE_ENABLED
 
+#if !NO_LOGGING
+		if (!SearchResult.IsValid())
+		{
+			TStringBuilder<1024> StringBuilder;
+			StringBuilder << "UPoseSearchLibrary::UpdateMotionMatchingState invalid search result : ForceInterrupt [";
+			StringBuilder << SearchContext.IsForceInterrupt();
+			StringBuilder << "], CanAdvance [";
+			StringBuilder << bCanAdvance;
+			StringBuilder << "], Databases [";
+
+			for (int32 DatabaseIndex = 0; DatabaseIndex < Databases.Num(); ++DatabaseIndex)
+			{
+				StringBuilder << GetNameSafe(Databases[DatabaseIndex]);
+				if (DatabaseIndex != Databases.Num() - 1)
+				{
+					StringBuilder << ", ";
+				}
+			}
+			
+			StringBuilder << "] ";
+
+			FString String = StringBuilder.ToString();
+			UE_LOG(LogPoseSearch, Warning, TEXT("%s"), *String);
+		}
+#endif // !NO_LOGGING
+
 		if (bJumpToPose)
 		{
 			InOutMotionMatchingState.JumpToPose(Context, SearchResult, MaxActiveBlends, BlendTime);
diff --git a/Engine/Plugins/Experimental/Animation/PoseSearch/Source/Runtime/Public/PoseSearch/PoseSearchDatabase.h b/Engine/Plugins/Experimental/Animation/PoseSearch/Source/Runtime/Public/PoseSearch/PoseSearchDatabase.h
--- a/Engine/Plugins/Experimental/Animation/PoseSearch/Source/Runtime/Public/PoseSearch/PoseSearchDatabase.h
+++ b/Engine/Plugins/Experimental/Animation/PoseSearch/Source/Runtime/Public/PoseSearch/PoseSearchDatabase.h
@@ -380,7 +380,7 @@ public:
 	// End UObject
 	
 	UE::PoseSearch::FSearchResult Search(UE::PoseSearch::FSearchContext& SearchContext) const;
-	FPoseSearchCost SearchContinuingPose(UE::PoseSearch::FSearchContext& SearchContext) const;
+	UE::PoseSearch::FSearchResult SearchContinuingPose(UE::PoseSearch::FSearchContext& SearchContext) const;
 
 #if WITH_EDITOR
 	int32 GetNumberOfPrincipalComponents() const;
