motion matching - fix for force interrupting the same requested database

diff --git a/Engine/Plugins/Experimental/Animation/PoseSearch/Source/Runtime/Private/AnimNode_MotionMatching.cpp b/Engine/Plugins/Experimental/Animation/PoseSearch/Source/Runtime/Private/AnimNode_MotionMatching.cpp
--- a/Engine/Plugins/Experimental/Animation/PoseSearch/Source/Runtime/Private/AnimNode_MotionMatching.cpp
+++ b/Engine/Plugins/Experimental/Animation/PoseSearch/Source/Runtime/Private/AnimNode_MotionMatching.cpp
@@ -145,7 +145,7 @@ void FAnimNode_MotionMatching::UpdateAssetPlayer(const FAnimationUpdateContext&
 		const UPoseSearchDatabase* CurrentDatabase = MotionMatchingState.CurrentSearchResult.Database.Get();
 		const UAnimationAsset* CurrentAnimationAsset = AnimPlayers.IsEmpty() ? nullptr : AnimPlayers[0].GetAnimationAsset();
 
-		FString DebugInfo = FString::Printf(TEXT("bForceInterruptNextUpdate(%d)\n"), bForceInterruptNextUpdate);
+		FString DebugInfo = FString::Printf(TEXT("NextUpdateInterruptMode(%s)\n"), *UEnum::GetValueAsString(NextUpdateInterruptMode));
 		DebugInfo += FString::Printf(TEXT("Current Database(%s)\n"), *GetNameSafe(CurrentDatabase));
 		DebugInfo += FString::Printf(TEXT("Current Asset(%s)\n"), *GetNameSafe(CurrentAnimationAsset));
 		if (CVarAnimNodeMotionMatchingDrawInfoVerbose.GetValueOnAnyThread())
@@ -180,7 +180,7 @@ void FAnimNode_MotionMatching::UpdateAssetPlayer(const FAnimationUpdateContext&
 		MotionMatchingState,
 		YawFromAnimationBlendRate,
 		YawFromAnimationTrajectoryBlendTime,
-		bForceInterruptNextUpdate,
+		NextUpdateInterruptMode,
 		bShouldSearch
 		#if ENABLE_ANIM_DEBUG
 		, CVarAnimNodeMotionMatchingDrawQuery.GetValueOnAnyThread()
@@ -213,7 +213,7 @@ void FAnimNode_MotionMatching::UpdateAssetPlayer(const FAnimationUpdateContext&
 	FAnimNode_BlendStack_Standalone::UpdatePlayRate(MotionMatchingState.WantedPlayRate);
 	FAnimNode_BlendStack_Standalone::UpdateAssetPlayer(Context);
 
-	bForceInterruptNextUpdate = false;
+	NextUpdateInterruptMode = EPoseSearchInterruptMode::DoNotInterrupt;
 }
 
 const FAnimNodeFunctionRef& FAnimNode_MotionMatching::GetOnUpdateMotionMatchingStateFunction() const
@@ -221,26 +221,9 @@ const FAnimNodeFunctionRef& FAnimNode_MotionMatching::GetOnUpdateMotionMatchingS
 	return GET_ANIM_NODE_DATA(FAnimNodeFunctionRef, OnMotionMatchingStateUpdated);
 }
 
-void FAnimNode_MotionMatching::SetDatabaseToSearch(UPoseSearchDatabase* InDatabase, bool bForceInterruptIfNew)
+void FAnimNode_MotionMatching::SetDatabaseToSearch(UPoseSearchDatabase* InDatabase, EPoseSearchInterruptMode InterruptMode)
 {
-	if (DatabasesToSearch.Num() == 1 && DatabasesToSearch[0] == InDatabase)
-	{
-		UE_LOG(LogPoseSearch, Verbose, TEXT("FAnimNode_MotionMatching::SetDatabaseToSearch - Database(%s) is already set."), *GetNameSafe(InDatabase));
-	}
-	else
-	{
-		DatabasesToSearch.Reset();
-		bOverrideDatabaseInput = false;
-		if (InDatabase)
-		{
-			DatabasesToSearch.Add(InDatabase);
-			bOverrideDatabaseInput = true;
-		}
-
-		bForceInterruptNextUpdate |= bForceInterruptIfNew;
-
-		UE_LOG(LogPoseSearch, Verbose, TEXT("FAnimNode_MotionMatching::SetDatabaseToSearch - Setting to Database(%s), bForceInterruptIfNew(%d)."), *GetNameSafe(InDatabase), bForceInterruptIfNew);
-	}
+	SetDatabasesToSearch(MakeArrayView(&InDatabase, 1), InterruptMode);
 }
 
 FVector FAnimNode_MotionMatching::GetEstimatedFutureRootMotionVelocity() const
@@ -248,60 +231,27 @@ FVector FAnimNode_MotionMatching::GetEstimatedFutureRootMotionVelocity() const
 	return MotionMatchingState.GetEstimatedFutureRootMotionVelocity();
 }
 
-void FAnimNode_MotionMatching::SetDatabasesToSearch(TConstArrayView<UPoseSearchDatabase*> InDatabases, bool bForceInterruptIfNew)
+void FAnimNode_MotionMatching::SetDatabasesToSearch(TConstArrayView<UPoseSearchDatabase*> InDatabases, EPoseSearchInterruptMode InterruptMode)
 {
-	// Check if InDatabases and DatabasesToSearch are the same.
-	bool bDatabasesAlreadySet = true;
-	if (DatabasesToSearch.Num() != InDatabases.Num())
-	{
-		bDatabasesAlreadySet = false;
-	}
-	else
-	{
-		for (int32 Index = 0; Index < InDatabases.Num(); ++Index)
-		{
-			if (DatabasesToSearch[Index] != InDatabases[Index])
-			{
-				bDatabasesAlreadySet = false;
-				break;
-			}
-		}
-	}
-
-	if (bDatabasesAlreadySet)
-	{
-		UE_LOG(LogPoseSearch, Verbose, TEXT("FAnimNode_MotionMatching::SetDatabasesToSearch - Databases(#%d) already set."), InDatabases.Num());
-	}
-	else
+	DatabasesToSearch.Reset();
+	for (UPoseSearchDatabase* InDatabase : InDatabases)
 	{
-		DatabasesToSearch.Reset();
-		bOverrideDatabaseInput = false;
-		if (!InDatabases.IsEmpty())
-		{
-			DatabasesToSearch.Append(InDatabases);
-			bOverrideDatabaseInput = true;
-		}
-
-		bForceInterruptNextUpdate |= bForceInterruptIfNew;
-
-		UE_LOG(LogPoseSearch, Verbose, TEXT("FAnimNode_MotionMatching::SetDatabaseToSearch - Setting to Databases(#%d), bForceInterruptIfNew(%d)."), InDatabases.Num(), bForceInterruptIfNew);
+		DatabasesToSearch.AddUnique(InDatabase);
 	}
+	NextUpdateInterruptMode = InterruptMode;
+	bOverrideDatabaseInput = true;
 }
 
-void FAnimNode_MotionMatching::ResetDatabasesToSearch(bool bInForceInterrupt)
+void FAnimNode_MotionMatching::ResetDatabasesToSearch(EPoseSearchInterruptMode InterruptMode)
 {
 	DatabasesToSearch.Reset();
 	bOverrideDatabaseInput = false;
-	bForceInterruptNextUpdate = bInForceInterrupt;
-
-	UE_LOG(LogPoseSearch, Verbose, TEXT("FAnimNode_MotionMatching::ResetDatabasesToSearch - Resetting databases, bInForceInterrupt(%d)."), bInForceInterrupt);
+	NextUpdateInterruptMode = InterruptMode;
 }
 
-void FAnimNode_MotionMatching::ForceInterruptNextUpdate()
+void FAnimNode_MotionMatching::SetInterruptMode(EPoseSearchInterruptMode InterruptMode)
 {
-	bForceInterruptNextUpdate = true;
-
-	UE_LOG(LogPoseSearch, Verbose, TEXT("FAnimNode_MotionMatching::ForceInterruptNextUpdate - Forcing interrupt."));
+	NextUpdateInterruptMode = InterruptMode;
 }
 
 // FAnimNode_AssetPlayerBase interface
diff --git a/Engine/Plugins/Experimental/Animation/PoseSearch/Source/Runtime/Private/MotionMatchingAnimNodeLibrary.cpp b/Engine/Plugins/Experimental/Animation/PoseSearch/Source/Runtime/Private/MotionMatchingAnimNodeLibrary.cpp
--- a/Engine/Plugins/Experimental/Animation/PoseSearch/Source/Runtime/Private/MotionMatchingAnimNodeLibrary.cpp
+++ b/Engine/Plugins/Experimental/Animation/PoseSearch/Source/Runtime/Private/MotionMatchingAnimNodeLibrary.cpp
@@ -42,11 +42,11 @@ void UMotionMatchingAnimNodeLibrary::GetMotionMatchingSearchResult(const FMotion
 	}
 }
 
-void UMotionMatchingAnimNodeLibrary::SetDatabaseToSearch(const FMotionMatchingAnimNodeReference& MotionMatchingNode, UPoseSearchDatabase* Database, bool bForceInterruptIfNew)
+void UMotionMatchingAnimNodeLibrary::SetDatabaseToSearch(const FMotionMatchingAnimNodeReference& MotionMatchingNode, UPoseSearchDatabase* Database, EPoseSearchInterruptMode InterruptMode)
 {
 	if (FAnimNode_MotionMatching* MotionMatchingNodePtr = MotionMatchingNode.GetAnimNodePtr<FAnimNode_MotionMatching>())
 	{
-		MotionMatchingNodePtr->SetDatabaseToSearch(Database, bForceInterruptIfNew);
+		MotionMatchingNodePtr->SetDatabaseToSearch(Database, InterruptMode);
 	}
 	else
 	{
@@ -54,11 +54,11 @@ void UMotionMatchingAnimNodeLibrary::SetDatabaseToSearch(const FMotionMatchingAn
 	}
 }
 
-void UMotionMatchingAnimNodeLibrary::SetDatabasesToSearch(const FMotionMatchingAnimNodeReference& MotionMatchingNode, const TArray<UPoseSearchDatabase*>& Databases, bool bForceInterruptIfNew)
+void UMotionMatchingAnimNodeLibrary::SetDatabasesToSearch(const FMotionMatchingAnimNodeReference& MotionMatchingNode, const TArray<UPoseSearchDatabase*>& Databases, EPoseSearchInterruptMode InterruptMode)
 {
 	if (FAnimNode_MotionMatching* MotionMatchingNodePtr = MotionMatchingNode.GetAnimNodePtr<FAnimNode_MotionMatching>())
 	{
-		MotionMatchingNodePtr->SetDatabasesToSearch(Databases, bForceInterruptIfNew);
+		MotionMatchingNodePtr->SetDatabasesToSearch(Databases, InterruptMode);
 	}
 	else
 	{
@@ -66,11 +66,11 @@ void UMotionMatchingAnimNodeLibrary::SetDatabasesToSearch(const FMotionMatchingA
 	}
 }
 
-void UMotionMatchingAnimNodeLibrary::ResetDatabasesToSearch(const FMotionMatchingAnimNodeReference& MotionMatchingNode, bool bForceInterrupt)
+void UMotionMatchingAnimNodeLibrary::ResetDatabasesToSearch(const FMotionMatchingAnimNodeReference& MotionMatchingNode, EPoseSearchInterruptMode InterruptMode)
 {
 	if (FAnimNode_MotionMatching* MotionMatchingNodePtr = MotionMatchingNode.GetAnimNodePtr<FAnimNode_MotionMatching>())
 	{
-		MotionMatchingNodePtr->ResetDatabasesToSearch(bForceInterrupt);
+		MotionMatchingNodePtr->ResetDatabasesToSearch(InterruptMode);
 	}
 	else
 	{
@@ -78,11 +78,11 @@ void UMotionMatchingAnimNodeLibrary::ResetDatabasesToSearch(const FMotionMatchin
 	}
 }
 
-void UMotionMatchingAnimNodeLibrary::ForceInterruptNextUpdate(const FMotionMatchingAnimNodeReference& MotionMatchingNode)
+void UMotionMatchingAnimNodeLibrary::SetInterruptMode(const FMotionMatchingAnimNodeReference& MotionMatchingNode, EPoseSearchInterruptMode InterruptMode)
 {
 	if (FAnimNode_MotionMatching* MotionMatchingNodePtr = MotionMatchingNode.GetAnimNodePtr<FAnimNode_MotionMatching>())
 	{
-		MotionMatchingNodePtr->ForceInterruptNextUpdate();
+		MotionMatchingNodePtr->SetInterruptMode(InterruptMode);
 	}
 	else
 	{
diff --git a/Engine/Plugins/Experimental/Animation/PoseSearch/Source/Runtime/Private/PoseSearchContext.cpp b/Engine/Plugins/Experimental/Animation/PoseSearch/Source/Runtime/Private/PoseSearchContext.cpp
--- a/Engine/Plugins/Experimental/Animation/PoseSearch/Source/Runtime/Private/PoseSearchContext.cpp
+++ b/Engine/Plugins/Experimental/Animation/PoseSearch/Source/Runtime/Private/PoseSearchContext.cpp
@@ -353,7 +353,7 @@ void FDebugDrawParams::DrawFeatureVector(int32 PoseIdx)
 // FSearchContext
 FSearchContext::FSearchContext(const UAnimInstance* InAnimInstance, const IPoseHistory* InHistory, TConstArrayView<const UAnimationAsset*> InAnimationsToConsider,
 		const FPoseSearchQueryTrajectory* InTrajectory, float InDesiredPermutationTimeOffset, const FPoseIndicesHistory* InPoseIndicesHistory,
-		const FSearchResult& InCurrentResult, const FFloatInterval& InPoseJumpThresholdTime, bool bInForceInterrupt)
+		const FSearchResult& InCurrentResult, const FFloatInterval& InPoseJumpThresholdTime)
 : AnimInstance(InAnimInstance)
 , History(InHistory)
 , AnimationsToConsider(InAnimationsToConsider)
@@ -362,7 +362,6 @@ FSearchContext::FSearchContext(const UAnimInstance* InAnimInstance, const IPoseH
 , PoseIndicesHistory(InPoseIndicesHistory)
 , CurrentResult(InCurrentResult)
 , PoseJumpThresholdTime(InPoseJumpThresholdTime)
-, bForceInterrupt(bInForceInterrupt)
 {
 	check(AnimInstance);
 	if (CurrentResult.IsValid())
diff --git a/Engine/Plugins/Experimental/Animation/PoseSearch/Source/Runtime/Private/PoseSearchLibrary.cpp b/Engine/Plugins/Experimental/Animation/PoseSearch/Source/Runtime/Private/PoseSearchLibrary.cpp
--- a/Engine/Plugins/Experimental/Animation/PoseSearch/Source/Runtime/Private/PoseSearchLibrary.cpp
+++ b/Engine/Plugins/Experimental/Animation/PoseSearch/Source/Runtime/Private/PoseSearchLibrary.cpp
@@ -32,6 +32,25 @@ TAutoConsoleVariable<int32> CVarAnimMotionMatchDrawMatchEnable(TEXT("a.MotionMat
 TAutoConsoleVariable<int32> CVarAnimMotionMatchDrawHistoryEnable(TEXT("a.MotionMatch.DrawHistory.Enable"), 0, TEXT("Enable / Disable MotionMatch Draw History"));
 #endif
 
+namespace UE::PoseSearch
+{
+	static bool IsForceInterrupt(EPoseSearchInterruptMode InterruptMode, const UPoseSearchDatabase* CurrentResultDatabase, const TArray<TObjectPtr<const UPoseSearchDatabase>>& Databases)
+	{
+		switch (InterruptMode)
+		{
+		case EPoseSearchInterruptMode::DoNotInterrupt:
+			return false;
+		case EPoseSearchInterruptMode::InterruptOnDatabaseChange:
+			return !Databases.Contains(CurrentResultDatabase);
+		case EPoseSearchInterruptMode::ForceInterrupt:
+			return true;
+		default:
+			unimplemented();
+			return false;
+		}
+	}
+}
+
 //////////////////////////////////////////////////////////////////////////
 // FMotionMatchingState
 
@@ -254,7 +273,7 @@ void UPoseSearchLibrary::UpdateMotionMatchingState(
 	FMotionMatchingState& InOutMotionMatchingState,
 	float YawFromAnimationBlendRate,
 	float YawFromAnimationTrajectoryBlendTime,
-	bool bForceInterrupt,
+	EPoseSearchInterruptMode InterruptMode,
 	bool bShouldSearch,
 	bool bDebugDrawQuery,
 	bool bDebugDrawCurResult,
@@ -301,7 +320,7 @@ void UPoseSearchLibrary::UpdateMotionMatchingState(
 	const UAnimInstance* AnimInstance = Cast<const UAnimInstance>(Context.AnimInstanceProxy->GetAnimInstanceObject());
 	check(AnimInstance);
 	FSearchContext SearchContext(AnimInstance, History, TConstArrayView<const UAnimationAsset*>(), &TrajectoryRootSpace, 0.f,
-		&InOutMotionMatchingState.PoseIndicesHistory, InOutMotionMatchingState.CurrentSearchResult, PoseJumpThresholdTime, bForceInterrupt);
+		&InOutMotionMatchingState.PoseIndicesHistory, InOutMotionMatchingState.CurrentSearchResult, PoseJumpThresholdTime);
 
 	const bool bCanAdvance = InOutMotionMatchingState.CurrentSearchResult.CanAdvance(DeltaTime);
 
@@ -311,11 +330,14 @@ void UPoseSearchLibrary::UpdateMotionMatchingState(
 	{
 		InOutMotionMatchingState.ElapsedPoseSearchTime = 0.f;
 
+		const UPoseSearchDatabase* CurrentResultDatabase = SearchContext.GetCurrentResult().Database.Get();
+		const bool bForceInterrupt = IsForceInterrupt(InterruptMode, CurrentResultDatabase, Databases);
+
 		// Evaluate continuing pose
 		FSearchResult SearchResult;
-		if (!SearchContext.IsForceInterrupt() && bCanAdvance)
+		if (!bForceInterrupt && bCanAdvance)
 		{
-			SearchResult = SearchContext.GetCurrentResult().Database->SearchContinuingPose(SearchContext);
+			SearchResult = CurrentResultDatabase->SearchContinuingPose(SearchContext);
 			SearchContext.UpdateCurrentBestCost(SearchResult.PoseCost);
 		}
 
@@ -346,7 +368,7 @@ void UPoseSearchLibrary::UpdateMotionMatchingState(
 		{
 			TStringBuilder<1024> StringBuilder;
 			StringBuilder << "UPoseSearchLibrary::UpdateMotionMatchingState invalid search result : ForceInterrupt [";
-			StringBuilder << SearchContext.IsForceInterrupt();
+			StringBuilder << bForceInterrupt;
 			StringBuilder << "], CanAdvance [";
 			StringBuilder << bCanAdvance;
 			StringBuilder << "], Databases [";
diff --git a/Engine/Plugins/Experimental/Animation/PoseSearch/Source/Runtime/Public/PoseSearch/AnimNode_MotionMatching.h b/Engine/Plugins/Experimental/Animation/PoseSearch/Source/Runtime/Public/PoseSearch/AnimNode_MotionMatching.h
--- a/Engine/Plugins/Experimental/Animation/PoseSearch/Source/Runtime/Public/PoseSearch/AnimNode_MotionMatching.h
+++ b/Engine/Plugins/Experimental/Animation/PoseSearch/Source/Runtime/Public/PoseSearch/AnimNode_MotionMatching.h
@@ -17,17 +17,17 @@ struct POSESEARCH_API FAnimNode_MotionMatching : public FAnimNode_BlendStack_Sta
 	GENERATED_BODY()
 
 public:
-	// Search InDatabase instead of the Database property on this node. Use bForceInterruptIfNew to ignore the continuing pose if InDatabase is new.
-	void SetDatabaseToSearch(UPoseSearchDatabase* InDatabase, bool bForceInterruptIfNew);
+	// Search InDatabase instead of the Database property on this node. Use InterruptMode to control the continuing pose search
+	void SetDatabaseToSearch(UPoseSearchDatabase* InDatabase, EPoseSearchInterruptMode InterruptMode);
 
-	// Search InDatabases instead of the Database property on the node. Use bForceInterruptIfNew to ignore the continuing pose if InDatabases is new.
-	void SetDatabasesToSearch(TConstArrayView<UPoseSearchDatabase*> InDatabases, bool bForceInterruptIfNew);
+	// Search InDatabases instead of the Database property on the node. Use InterruptMode to control the continuing pose search.
+	void SetDatabasesToSearch(TConstArrayView<UPoseSearchDatabase*> InDatabases, EPoseSearchInterruptMode InterruptMode);
 
 	// Reset the effects of SetDatabaseToSearch/SetDatabasesToSearch and use the Database property on this node.
-	void ResetDatabasesToSearch(bool bInForceInterrupt);
+	void ResetDatabasesToSearch(EPoseSearchInterruptMode InterruptMode);
 
-	// Ignore the continuing pose on the next update and force a search.
-	void ForceInterruptNextUpdate();
+	// Use InterruptMode to control the continuing pose search
+	void SetInterruptMode(EPoseSearchInterruptMode InterruptMode);
 
 	const FMotionMatchingState& GetMotionMatchingState() const { return MotionMatchingState; }
 
@@ -124,8 +124,8 @@ private:
 	UPROPERTY()
 	TArray<TObjectPtr<const UPoseSearchDatabase>> DatabasesToSearch;
 
-	// Ignore the continuing pose on the next update and use the best result from DatabasesToSearch. This is set back to false after each update.
-	bool bForceInterruptNextUpdate = false;
+	// Applied EPoseSearchInterruptMode on the next update that controls the continuing pose search eveluation. This is set back to EPoseSearchInterruptMode::DoNotInterrupt after each update.
+	EPoseSearchInterruptMode NextUpdateInterruptMode = EPoseSearchInterruptMode::DoNotInterrupt;
 
 	// True if the Database property on this node has been overridden by SetDatabaseToSearch/SetDatabasesToSearch.
 	bool bOverrideDatabaseInput = false;
diff --git a/Engine/Plugins/Experimental/Animation/PoseSearch/Source/Runtime/Public/PoseSearch/MotionMatchingAnimNodeLibrary.h b/Engine/Plugins/Experimental/Animation/PoseSearch/Source/Runtime/Public/PoseSearch/MotionMatchingAnimNodeLibrary.h
--- a/Engine/Plugins/Experimental/Animation/PoseSearch/Source/Runtime/Public/PoseSearch/MotionMatchingAnimNodeLibrary.h
+++ b/Engine/Plugins/Experimental/Animation/PoseSearch/Source/Runtime/Public/PoseSearch/MotionMatchingAnimNodeLibrary.h
@@ -45,38 +45,33 @@ public:
 	 * Set the database to search on the motion matching node. This overrides the Database property on the motion matching node.
 	 * @param MotionMatchingNode - The motion matching node to operate on.
 	 * @param Database - The database for the motion matching node to search.
-	 * @param bForceInterruptIfNew - If true, ignore the continuing pose (the current clip that's playing) and force a search of the new database. 
-		If false, the continuing pose will continue to play until a better match is found in the new database. This setting is ignored if the 
-		motion matching node is already searching this database.
+	 * @param InterruptMode - mode to control the continuing pose search (the current animation that's playing)
 	 */
 	UFUNCTION(BlueprintCallable, Category = "Animation|MotionMatching", meta = (BlueprintThreadSafe))
-	static void SetDatabaseToSearch(const FMotionMatchingAnimNodeReference& MotionMatchingNode, UPoseSearchDatabase* Database, bool bForceInterruptIfNew);
+	static void SetDatabaseToSearch(const FMotionMatchingAnimNodeReference& MotionMatchingNode, UPoseSearchDatabase* Database, EPoseSearchInterruptMode InterruptMode);
 
 	/**
 	 * Set the database to search on the motion matching node. This overrides the Database property on the motion matching node.
 	 * @param MotionMatchingNode - The motion matching node to operate on.
 	 * @param Databases - Array of databases for the motion matching node to search.
-	 * @param bForceInterruptIfNew - If true, ignore the continuing pose (the current clip that's playing) and force a search of the new databases. 
-		If false, the continuing pose will continue to play until a better match is found in one of the new databases. This setting is ignored if the 
-		motion matching node is already searching this array of databases (note: a subset of databases or the same databases in a different order is
-		considered to be a new array of databases).
+	 * @param InterruptMode - mode to control the continuing pose search (the current animation that's playing)
 	 */
 	UFUNCTION(BlueprintCallable, Category = "Animation|MotionMatching", meta = (BlueprintThreadSafe))
-	static void SetDatabasesToSearch(const FMotionMatchingAnimNodeReference& MotionMatchingNode, const TArray<UPoseSearchDatabase*>& Databases, bool bForceInterruptIfNew);
+	static void SetDatabasesToSearch(const FMotionMatchingAnimNodeReference& MotionMatchingNode, const TArray<UPoseSearchDatabase*>& Databases, EPoseSearchInterruptMode InterruptMode);
 
 	/**
 	 * Clear the effects of SetDatabaseToSearch/SetDatabasesToSearch and resume searching the Database property on the motion matching node.
 	 * @param MotionMatchingNode - The motion matching node to operate on.
-	 * @param bForceInterrupt - Force a search after the reset. If false, the continuing pose (the current clip that's playing) will continue
-		to play until a better match is found from the Database property.
+	 * @param InterruptMode - mode to control the continuing pose search (the current animation that's playing)
 	 */
 	UFUNCTION(BlueprintCallable, Category = "Animation|MotionMatching", meta = (BlueprintThreadSafe))
-	static void ResetDatabasesToSearch(const FMotionMatchingAnimNodeReference& MotionMatchingNode, bool bForceInterrupt);
+	static void ResetDatabasesToSearch(const FMotionMatchingAnimNodeReference& MotionMatchingNode, EPoseSearchInterruptMode InterruptMode);
 
 	/**
 	 * Ignore the continuing pose (the current clip that's playing) and force a new search.
 	 * @param MotionMatchingNode - The motion matching node to operate on.
+	 * @param InterruptMode - mode to control the continuing pose search (the current animation that's playing)
 	 */
 	UFUNCTION(BlueprintCallable, Category = "Animation|MotionMatching", meta = (BlueprintThreadSafe))
-	static void ForceInterruptNextUpdate(const FMotionMatchingAnimNodeReference& MotionMatchingNode);
+	static void SetInterruptMode(const FMotionMatchingAnimNodeReference& MotionMatchingNode, EPoseSearchInterruptMode InterruptMode);
 };
\ No newline at end of file
diff --git a/Engine/Plugins/Experimental/Animation/PoseSearch/Source/Runtime/Public/PoseSearch/PoseSearchContext.h b/Engine/Plugins/Experimental/Animation/PoseSearch/Source/Runtime/Public/PoseSearch/PoseSearchContext.h
--- a/Engine/Plugins/Experimental/Animation/PoseSearch/Source/Runtime/Public/PoseSearch/PoseSearchContext.h
+++ b/Engine/Plugins/Experimental/Animation/PoseSearch/Source/Runtime/Public/PoseSearch/PoseSearchContext.h
@@ -145,7 +145,7 @@ struct POSESEARCH_API FSearchContext
 {
 	FSearchContext(const UAnimInstance* InAnimInstance, const IPoseHistory* InHistory, TConstArrayView<const UAnimationAsset*> InAnimationsToConsider = TConstArrayView<const UAnimationAsset*>(),
 		const FPoseSearchQueryTrajectory* InTrajectory = nullptr, float InDesiredPermutationTimeOffset = 0.f, const FPoseIndicesHistory* InPoseIndicesHistory = nullptr,
-		const FSearchResult& InCurrentResult = FSearchResult(), const FFloatInterval& InPoseJumpThresholdTime = FFloatInterval(0.f, 0.f), bool bInForceInterrupt = false);
+		const FSearchResult& InCurrentResult = FSearchResult(), const FFloatInterval& InPoseJumpThresholdTime = FFloatInterval(0.f, 0.f));
 
 	// Returns the rotation of the bone Schema.BoneReferences[SchemaSampleBoneIdx] at an offset time of SampleTimeOffset relative to the
 	// transform of the bone Schema.BoneReferences[SchemaOriginBoneIdx] at an offset time of time OriginTimeOffset 
@@ -182,7 +182,6 @@ struct POSESEARCH_API FSearchContext
 	bool IsHistoryValid() const { return History != nullptr; }
 	bool IsTrajectoryValid() const { return Trajectory != nullptr; }
 	float GetDesiredPermutationTimeOffset() const { return DesiredPermutationTimeOffset; }
-	bool IsForceInterrupt() const { return bForceInterrupt; }
 	const UAnimInstance* GetAnimInstance() const { return AnimInstance; }
 
 	void SetAnimationsToConsider(TConstArrayView<const UAnimationAsset*> InAnimationsToConsider) { AnimationsToConsider = InAnimationsToConsider; }
@@ -208,7 +207,6 @@ private:
 	const FPoseIndicesHistory* PoseIndicesHistory = nullptr;
 	const FSearchResult& CurrentResult;
 	const FFloatInterval& PoseJumpThresholdTime;
-	const bool bForceInterrupt = false;
 
 	TConstArrayView<float> CurrentResultPoseVector;
 	TStackAlignedArray<float> CurrentResultPoseVectorData;
diff --git a/Engine/Plugins/Experimental/Animation/PoseSearch/Source/Runtime/Public/PoseSearch/PoseSearchLibrary.h b/Engine/Plugins/Experimental/Animation/PoseSearch/Source/Runtime/Public/PoseSearch/PoseSearchLibrary.h
--- a/Engine/Plugins/Experimental/Animation/PoseSearch/Source/Runtime/Public/PoseSearch/PoseSearchLibrary.h
+++ b/Engine/Plugins/Experimental/Animation/PoseSearch/Source/Runtime/Public/PoseSearch/PoseSearchLibrary.h
@@ -19,6 +19,19 @@ namespace UE::PoseSearch
 	struct FSearchContext;
 } // namespace UE::PoseSearch
 
+UENUM()
+enum class EPoseSearchInterruptMode : uint8
+{
+	// continuing pose search will be performed if valid
+	DoNotInterrupt,
+
+	// continuing pose search will be interrupted if its database is not listed in the searchable databases
+	InterruptOnDatabaseChange,
+
+	// continuing pose search will always be interrupted
+	ForceInterrupt,
+};
+
 struct FAnimationUpdateContext;
 struct FPoseSearchQueryTrajectory;
 
@@ -109,7 +122,7 @@ public:
 	* @param SearchThrottleTime				Input minimum amount of time to wait between searching for a new pose segment. It allows users to define how often the system searches, default for locomotion is searching every update, but you may only want to search once for other situations, like jump.
 	* @param PlayRate						Input effective range of play rate that can be applied to the animations to account for discrepancies in estimated velocity between the movement modeland the animation.
 	* @param InOutMotionMatchingState		Input/Output encapsulated motion matching algorithm and state
-	* @param bForceInterrupt				Input force interrupt request (if true the continuing pose will be invalidated)
+	* @param InterruptMode					Input continuing pose search interrupt mode
 	* @param bShouldSearch					Input if false search will happen only if there's no valid continuing pose
 	* @param bDebugDrawQuery				Input draw the composed query if valid
 	* @param bDebugDrawCurResult			Input draw the current result if valid
@@ -128,7 +141,7 @@ public:
 		FMotionMatchingState& InOutMotionMatchingState,
 		float YawFromAnimationBlendRate,
 		float YawFromAnimationTrajectoryBlendTime,
-		bool bForceInterrupt = false,
+		EPoseSearchInterruptMode InterruptMode = EPoseSearchInterruptMode::DoNotInterrupt,
 		bool bShouldSearch = true,
 		bool bDebugDrawQuery = false,
 		bool bDebugDrawCurResult = false,
