motion matching - bPruneBlockTransitionPosesFromTrees: pruning block transition poses from the KDTree

diff --git a/Engine/Plugins/Experimental/Animation/PoseSearch/Source/Editor/Private/PoseSearchDebuggerDatabaseView.cpp b/Engine/Plugins/Experimental/Animation/PoseSearch/Source/Editor/Private/PoseSearchDebuggerDatabaseView.cpp
--- a/Engine/Plugins/Experimental/Animation/PoseSearch/Source/Editor/Private/PoseSearchDebuggerDatabaseView.cpp
+++ b/Engine/Plugins/Experimental/Animation/PoseSearch/Source/Editor/Private/PoseSearchDebuggerDatabaseView.cpp
@@ -206,10 +206,13 @@ static void AddUnfilteredDatabaseRow(const UPoseSearchDatabase* Database,
 		if (!SharedData->PCAQueryVector.IsEmpty())
 		{
 			const int32 PCAValuesVectorIdx = PoseToPCAValuesVectorIndexes.IsEmpty() ? DbPoseIdx : PoseToPCAValuesVectorIndexes[DbPoseIdx];
-			TConstArrayView<float> PCAPoseValues = SearchIndex.GetPCAPoseValues(PCAValuesVectorIdx);
-			if (SharedData->PCAQueryVector.Num() == PCAPoseValues.Num())
+			if (PCAValuesVectorIdx >= 0)
 			{
-				Row->PosePCACost = CompareFeatureVectors(SharedData->PCAQueryVector, PCAPoseValues);
+				TConstArrayView<float> PCAPoseValues = SearchIndex.GetPCAPoseValues(PCAValuesVectorIdx);
+				if (SharedData->PCAQueryVector.Num() == PCAPoseValues.Num())
+				{
+					Row->PosePCACost = CompareFeatureVectors(SharedData->PCAQueryVector, PCAPoseValues);
+				}
 			}
 		}
 
diff --git a/Engine/Plugins/Experimental/Animation/PoseSearch/Source/Runtime/Private/PoseSearchDatabase.cpp b/Engine/Plugins/Experimental/Animation/PoseSearch/Source/Runtime/Private/PoseSearchDatabase.cpp
--- a/Engine/Plugins/Experimental/Animation/PoseSearch/Source/Runtime/Private/PoseSearchDatabase.cpp
+++ b/Engine/Plugins/Experimental/Animation/PoseSearch/Source/Runtime/Private/PoseSearchDatabase.cpp
@@ -176,9 +176,9 @@ static void PopulateNonSelectableIdx(FNonSelectableIdx& NonSelectableIdx, FSearc
 
 struct FSearchFilters
 {
-	FSearchFilters(const UPoseSearchSchema* Schema, TConstArrayView<int32> NonSelectableIdx, TConstArrayView<int32> SelectableAssetIdx, bool bAnyBlockTransition)
+	FSearchFilters(const UPoseSearchSchema* Schema, TConstArrayView<int32> NonSelectableIdx, TConstArrayView<int32> SelectableAssetIdx, bool bAddBlockTransitionFilter)
 	{
-		if (bAnyBlockTransition)
+		if (bAddBlockTransitionFilter)
 		{
 			Filters.Add(&BlockTransitionFilter);
 		}
@@ -1119,7 +1119,7 @@ UE::PoseSearch::FSearchResult UPoseSearchDatabase::SearchPCAKDTree(UE::PoseSearc
 #endif // WITH_EDITOR && ENABLE_ANIM_DEBUG
 
 		// NonSelectableIdx are already filtered out inside the kdtree search
-		const FSearchFilters SearchFilters(Schema, bRunNonSelectableIdxPostKDTree ? NonSelectableIdx : TConstArrayView<int32>(), SelectableAssetIdx, SearchIndex.bAnyBlockTransition);
+		const FSearchFilters SearchFilters(Schema, bRunNonSelectableIdxPostKDTree ? NonSelectableIdx : TConstArrayView<int32>(), SelectableAssetIdx, SearchIndex.bAnyBlockTransition && !bPCAValuesPruningFromBlockTransitionPoses);
 		
 		// are the PCAValues pruned out of duplicates (multiple poses are associated with the same PCAValuesVectorIdx)
 		if (bArePCAValuesPruned)
diff --git a/Engine/Plugins/Experimental/Animation/PoseSearch/Source/Runtime/Private/PoseSearchDerivedData.cpp b/Engine/Plugins/Experimental/Animation/PoseSearch/Source/Runtime/Private/PoseSearchDerivedData.cpp
--- a/Engine/Plugins/Experimental/Animation/PoseSearch/Source/Runtime/Private/PoseSearchDerivedData.cpp
+++ b/Engine/Plugins/Experimental/Animation/PoseSearch/Source/Runtime/Private/PoseSearchDerivedData.cpp
@@ -1641,6 +1641,11 @@ void FPoseSearchDatabaseAsyncCacheTask::OnGetComplete(UE::DerivedData::FCacheGet
 					return;
 				}
 
+				if (IndexBaseDatabases[0]->bPCAValuesPruningFromBlockTransitionPoses)
+				{
+					SearchIndex.PrunePCAValuesFromBlockTransitionPoses(IndexBaseDatabases[0]->GetNumberOfPrincipalComponents());
+				}
+
 				PreprocessSearchIndexKDTree(SearchIndex, IndexBaseDatabases[0]);
 				if (Owner.IsCanceled())
 				{
diff --git a/Engine/Plugins/Experimental/Animation/PoseSearch/Source/Runtime/Private/PoseSearchIndex.cpp b/Engine/Plugins/Experimental/Animation/PoseSearch/Source/Runtime/Private/PoseSearchIndex.cpp
--- a/Engine/Plugins/Experimental/Animation/PoseSearch/Source/Runtime/Private/PoseSearchIndex.cpp
+++ b/Engine/Plugins/Experimental/Animation/PoseSearch/Source/Runtime/Private/PoseSearchIndex.cpp
@@ -615,6 +615,65 @@ TConstArrayView<float> FSearchIndex::PCAProject(TConstArrayView<float> PoseValue
 	return BufferUsedForProjection;
 }
 
+void FSearchIndex::PrunePCAValuesFromBlockTransitionPoses(int32 NumberOfPrincipalComponents)
+{
+	if (!bAnyBlockTransition)
+	{
+		return;
+	}
+
+	check(PCAValues.Num() % NumberOfPrincipalComponents == 0);
+	const int32 NumPCAValuesVectors = PCAValues.Num() / NumberOfPrincipalComponents;
+
+	TArray<TPair<int32, TArray<int32>>> PrunedPCAValuesVectorToPoseIndexes;
+	TAlignedArray<float> PrunedPCAValues;
+	PrunedPCAValues.Reserve(PCAValues.Num());
+
+	if (PCAValuesVectorToPoseIndexes.Num() > 0)
+	{
+		TArray<int32> PCAValuesVectorIdxPoseIndexes;
+		for (int32 PCAValuesVectorIdx = 0; PCAValuesVectorIdx < NumPCAValuesVectors; ++PCAValuesVectorIdx)
+		{
+			PCAValuesVectorIdxPoseIndexes.Reset();
+			for (int32 PoseIdx : PCAValuesVectorToPoseIndexes[PCAValuesVectorIdx])
+			{
+				if (!PoseMetadata[PoseIdx].IsBlockTransition())
+				{
+					PCAValuesVectorIdxPoseIndexes.Add(PoseIdx);
+				}
+			}
+
+			if (!PCAValuesVectorIdxPoseIndexes.IsEmpty())
+			{
+				PrunedPCAValuesVectorToPoseIndexes.Emplace(PrunedPCAValuesVectorToPoseIndexes.Num(), PCAValuesVectorIdxPoseIndexes);
+				PrunedPCAValues.Append(GetPCAPoseValues(PCAValuesVectorIdx));
+			}
+		}
+	}
+	else
+	{
+		for (int32 PCAValuesVectorIdx = 0; PCAValuesVectorIdx < NumPCAValuesVectors; ++PCAValuesVectorIdx)
+		{
+			// here there's a 1:1 mapping between PCAValuesVectorIdx and PoseIdx
+			
+			const int32 PoseIdx = PCAValuesVectorIdx;
+			if (!PoseMetadata[PoseIdx].IsBlockTransition())
+			{
+				TConstArrayView<int32> PCAValuesVectorIdxPoseIndexes = MakeArrayView(&PoseIdx, 1);
+				PrunedPCAValuesVectorToPoseIndexes.Emplace(PrunedPCAValuesVectorToPoseIndexes.Num(), PCAValuesVectorIdxPoseIndexes);
+				PrunedPCAValues.Append(GetPCAPoseValues(PCAValuesVectorIdx));
+			}
+		}
+	}
+
+	PCAValues = PrunedPCAValues;
+	PCAValuesVectorToPoseIndexes = FSparsePoseMultiMap<int32>(PrunedPCAValuesVectorToPoseIndexes.Num(), GetNumPoses() - 1);
+	for (const TPair<int32, TArray<int32>>& Pair : PrunedPCAValuesVectorToPoseIndexes)
+	{
+		PCAValuesVectorToPoseIndexes.Insert(Pair.Key, Pair.Value);
+	}
+}
+
 void FSearchIndex::PruneDuplicatePCAValues(float SimilarityThreshold, int32 NumberOfPrincipalComponents)
 {
 	PCAValuesVectorToPoseIndexes = FSparsePoseMultiMap<int32>();
@@ -757,7 +816,8 @@ void FSearchIndex::GetPoseToPCAValuesVectorIndexes(TArray<uint32>& PoseToPCAValu
 			}
 		}
 
-		check(!PoseToPCAValuesVectorIndexes.Contains(INDEX_NONE));
+		// if block transition pruning is enabled we could have poses without their PCAValuesVectorIdx counterpart
+		// check(!PoseToPCAValuesVectorIndexes.Contains(INDEX_NONE));
 	}
 	else
 	{
diff --git a/Engine/Plugins/Experimental/Animation/PoseSearch/Source/Runtime/Public/PoseSearch/PoseSearchDatabase.h b/Engine/Plugins/Experimental/Animation/PoseSearch/Source/Runtime/Public/PoseSearch/PoseSearchDatabase.h
--- a/Engine/Plugins/Experimental/Animation/PoseSearch/Source/Runtime/Public/PoseSearch/PoseSearchDatabase.h
+++ b/Engine/Plugins/Experimental/Animation/PoseSearch/Source/Runtime/Public/PoseSearch/PoseSearchDatabase.h
@@ -318,7 +318,7 @@ public:
 	
 	// @todo: rename to KNNQueryNumNeighbors to be usable with the VPTree as well
 	// Out of a kdtree search, results will have only an approximate cost, so the database search will select the best “KDTree Query Num Neighbors” poses to perform the full cost analysis, and be able to elect the best pose.
-	UPROPERTY(EditAnywhere, Category = "Performance", meta = (DisplayName = "KNNQueryNumNeighbors", EditCondition = "PoseSearchMode == EPoseSearchMode::PCAKDTree || PoseSearchMode == VPTree", EditConditionHides, ClampMin = "1", ClampMax = "600", UIMin = "1"))
+	UPROPERTY(EditAnywhere, Category = "Performance", meta = (DisplayName = "KNNQueryNumNeighbors", EditCondition = "PoseSearchMode == EPoseSearchMode::PCAKDTree || PoseSearchMode == EPoseSearchMode::VPTree", EditConditionHides, ClampMin = "1", ClampMax = "600", UIMin = "1"))
 	int32 KDTreeQueryNumNeighbors = 200;
 
 #if WITH_EDITORONLY_DATA
@@ -339,6 +339,11 @@ public:
 	UPROPERTY(EditAnywhere, Category = "Performance", meta = (DisplayName = "KNNQueryNumNeighborsWithDuplicates", EditCondition = "PoseSearchMode == EPoseSearchMode::PCAKDTree && PCAValuesPruningSimilarityThreshold > 0", EditConditionHides, ClampMin = "0", ClampMax = "600", UIMin = "1"))
 	int32 KDTreeQueryNumNeighborsWithDuplicates = 0;
 	
+	// pruning block transition poses from the PCA values (the KDTree will be smaller and faster). This will force pose filtering to be done after tree search,
+	// and in case all the KNNQueryNumNeighbors requested poses get filtered out, the search could end up with no valid results
+	UPROPERTY(EditAnywhere, Category = "Performance", meta = (EditCondition = "PoseSearchMode == EPoseSearchMode::PCAKDTree", EditConditionHides))
+	bool bPCAValuesPruningFromBlockTransitionPoses = false;
+
 private:
 	// Do not use it directly. Use GetSearchIndex / SetSearchIndex interact with it and validate that is ok to do so.
 	UE::PoseSearch::FSearchIndex SearchIndexPrivate;
diff --git a/Engine/Plugins/Experimental/Animation/PoseSearch/Source/Runtime/Public/PoseSearch/PoseSearchIndex.h b/Engine/Plugins/Experimental/Animation/PoseSearch/Source/Runtime/Public/PoseSearch/PoseSearchIndex.h
--- a/Engine/Plugins/Experimental/Animation/PoseSearch/Source/Runtime/Public/PoseSearch/PoseSearchIndex.h
+++ b/Engine/Plugins/Experimental/Animation/PoseSearch/Source/Runtime/Public/PoseSearch/PoseSearchIndex.h
@@ -232,13 +232,12 @@ struct FSearchStats
 template <typename Type = int32>
 struct FSparsePoseMultiMap
 {
-	FSparsePoseMultiMap(Type InMaxKey = Type(0), Type InMaxValue = Type(0))
+	FSparsePoseMultiMap(Type InMaxKey = Type(0), Type InMaxValue = Type(0), TArray<Type>::SizeType InitialAllocationSize = 0)
 	: MaxKey(InMaxKey)
 	, MaxValue(InMaxValue)
 	, DeltaKeyValue(InMaxValue >= InMaxKey ? InMaxValue - InMaxKey + 1 : 0)
 	{
-		// @todo: maybe expose this initial allocation budget
-		DataValues.Reserve(InMaxKey * 2);
+		DataValues.Reserve(InitialAllocationSize > 0 ? InitialAllocationSize : InMaxKey * 2);
 		for (Type Index = 0; Index < InMaxKey; ++Index)
 		{
 			DataValues.Add(Type(INDEX_NONE));
@@ -458,6 +457,7 @@ struct FSearchIndex : public FSearchIndexBase
 	POSESEARCH_API FPoseSearchCost CompareAlignedPoses(int32 PoseIdx, float ContinuingPoseCostBias, TConstArrayView<float> PoseValues, TConstArrayView<float> QueryValues) const;
 
 	void PruneDuplicatePCAValues(float SimilarityThreshold, int32 NumberOfPrincipalComponents);
+	void PrunePCAValuesFromBlockTransitionPoses(int32 NumberOfPrincipalComponents);
 
 	// returns the inverse mapping of PCAValuesVectorToPoseIndexes
 	POSESEARCH_API void GetPoseToPCAValuesVectorIndexes(TArray<uint32>& PoseToPCAValuesVectorIndexes) const;
