motion matching - fix for Incremental Compile Monolithics

diff --git a/Engine/Plugins/Experimental/Animation/PoseSearch/Source/Runtime/Public/PoseSearch/PoseSearchIndex.h b/Engine/Plugins/Experimental/Animation/PoseSearch/Source/Runtime/Public/PoseSearch/PoseSearchIndex.h
--- a/Engine/Plugins/Experimental/Animation/PoseSearch/Source/Runtime/Public/PoseSearch/PoseSearchIndex.h
+++ b/Engine/Plugins/Experimental/Animation/PoseSearch/Source/Runtime/Public/PoseSearch/PoseSearchIndex.h
@@ -232,12 +232,13 @@ struct FSearchStats
 template <typename Type = int32>
 struct FSparsePoseMultiMap
 {
-	FSparsePoseMultiMap(Type InMaxKey = Type(0), Type InMaxValue = Type(0), TArray<Type>::SizeType InitialAllocationSize = 0)
+	FSparsePoseMultiMap(Type InMaxKey = Type(0), Type InMaxValue = Type(0))
 	: MaxKey(InMaxKey)
 	, MaxValue(InMaxValue)
 	, DeltaKeyValue(InMaxValue >= InMaxKey ? InMaxValue - InMaxKey + 1 : 0)
 	{
-		DataValues.Reserve(InitialAllocationSize > 0 ? InitialAllocationSize : InMaxKey * 2);
+		// @todo: maybe expose this initial allocation budget
+		DataValues.Reserve(InMaxKey * 2);
 		for (Type Index = 0; Index < InMaxKey; ++Index)
 		{
 			DataValues.Add(Type(INDEX_NONE));
