motion matching
- added MM database AdditionalExtrapolationTime to clamp trajectory / root motion extrapolation
- ExcludeFromDatabaseParameters as FFloatInterval, made consistent with AdditionalExtrapolationTime (flipped ExcludeFromDatabaseParameters.Max sign)
- more logging in FPoseSearchDatabaseAsyncCacheTask::Cancel
- fixed AsyncPoseSearchDatabasesManagement::Tasks rare corruption when calling FAsyncPoseSearchDatabasesManagement::ClearPreCancelled

diff --git a/Engine/Plugins/Experimental/Animation/PoseSearch/Source/Runtime/Private/PoseSearchAssetIndexer.cpp b/Engine/Plugins/Experimental/Animation/PoseSearch/Source/Runtime/Private/PoseSearchAssetIndexer.cpp
--- a/Engine/Plugins/Experimental/Animation/PoseSearch/Source/Runtime/Private/PoseSearchAssetIndexer.cpp
+++ b/Engine/Plugins/Experimental/Animation/PoseSearch/Source/Runtime/Private/PoseSearchAssetIndexer.cpp
@@ -91,14 +91,15 @@ FAssetSamplingContext::FAssetSamplingContext(const UPoseSearchDatabase& Database
 
 //////////////////////////////////////////////////////////////////////////
 // FAssetIndexer
-FAssetIndexer::FAssetIndexer(const FBoneContainer& InBoneContainer, const FSearchIndexAsset& InSearchIndexAsset, 
-	const FAssetSamplingContext& InSamplingContext, const UPoseSearchSchema& InSchema, const FAnimationAssetSampler& InAssetSampler)
+FAssetIndexer::FAssetIndexer(const FBoneContainer& InBoneContainer, const FSearchIndexAsset& InSearchIndexAsset, const FAssetSamplingContext& InSamplingContext,
+	const UPoseSearchSchema& InSchema, const FAnimationAssetSampler& InAssetSampler, const FFloatInterval& InExtrapolationTimeInterval)
 : BoneContainer(InBoneContainer)
 , CachedEntries()
 , SearchIndexAsset(InSearchIndexAsset)
 , SamplingContext(InSamplingContext)
 , Schema(InSchema)
 , AssetSampler(InAssetSampler)
+, ExtrapolationTimeInterval(InExtrapolationTimeInterval)
 {
 }
 
@@ -287,6 +288,8 @@ FAssetIndexer::CachedEntry& FAssetIndexer::GetEntry(float SampleTime)
 	bDisableCaching = CVarMotionMatchTestDisableIndexerCaching.GetValueOnAnyThread();
 #endif // ENABLE_ANIM_DEBUG
 
+	SampleTime = FMath::Clamp(SampleTime, ExtrapolationTimeInterval.Min, ExtrapolationTimeInterval.Max);
+
 	CachedEntry* Entry = bDisableCaching ? nullptr : CachedEntries.Find(SampleTime);
 	if (!Entry)
 	{
diff --git a/Engine/Plugins/Experimental/Animation/PoseSearch/Source/Runtime/Private/PoseSearchDerivedData.cpp b/Engine/Plugins/Experimental/Animation/PoseSearch/Source/Runtime/Private/PoseSearchDerivedData.cpp
--- a/Engine/Plugins/Experimental/Animation/PoseSearch/Source/Runtime/Private/PoseSearchDerivedData.cpp
+++ b/Engine/Plugins/Experimental/Animation/PoseSearch/Source/Runtime/Private/PoseSearchDerivedData.cpp
@@ -277,7 +277,7 @@ static inline FFloatInterval GetEffectiveSamplingRange(const UAnimSequenceBase*
 }
 
 static void FindValidSequenceIntervals(const UAnimSequenceBase* SequenceBase, FFloatInterval SamplingRange, bool bIsLooping,
-	const FPoseSearchExcludeFromDatabaseParameters& ExcludeFromDatabaseParameters, TArray<FFloatRange>& ValidRanges)
+	const FFloatInterval& ExcludeFromDatabaseParameters, TArray<FFloatRange>& ValidRanges)
 {
 	check(SequenceBase);
 
@@ -287,7 +287,7 @@ static void FindValidSequenceIntervals(const UAnimSequenceBase* SequenceBase, FF
 	FFloatRange EffectiveSamplingRange = FFloatRange::Inclusive(EffectiveSamplingInterval.Min, EffectiveSamplingInterval.Max);
 	if (!bIsLooping)
 	{
-		const FFloatRange ExcludeFromDatabaseRange(ExcludeFromDatabaseParameters.SequenceStartInterval, SequenceLength - ExcludeFromDatabaseParameters.SequenceEndInterval);
+		const FFloatRange ExcludeFromDatabaseRange(ExcludeFromDatabaseParameters.Min, SequenceLength + ExcludeFromDatabaseParameters.Max);
 		EffectiveSamplingRange = FFloatRange::Intersection(EffectiveSamplingRange, ExcludeFromDatabaseRange);
 	}
 
@@ -861,7 +861,8 @@ static bool IndexDatabase(FSearchIndexBase& SearchIndexBase, const UPoseSearchDa
 		check(DatabaseAnimationAssetBase && DatabaseAnimationAssetBase->GetAnimationAsset());
 		const FAnimationAssetSampler& AssetSampler = Samplers[SamplerMap[{ DatabaseAnimationAssetBase->GetAnimationAsset(), SearchIndexAsset.GetBlendParameters() }]];
 
-		Indexers.Emplace(BoneContainer, SearchIndexAsset, SamplingContext, *Schema, AssetSampler);
+		const FFloatInterval ExtrapolationTimeInterval = SearchIndexAsset.GetExtrapolationTimeInterval(Schema->SampleRate, Database.AdditionalExtrapolationTime);
+		Indexers.Emplace(BoneContainer, SearchIndexAsset, SamplingContext, *Schema, AssetSampler, ExtrapolationTimeInterval);
 		TotalPoses += SearchIndexAsset.GetNumPoses();
 	}
 
@@ -1254,6 +1255,9 @@ void FPoseSearchDatabaseAsyncCacheTask::Cancel()
 {
 	check(IsInGameThread());
 
+	FString DatabaseName = IsValid() ? *Database->GetName() : TEXT("Garbage Collected Database");
+	UE_LOG(LogPoseSearch, Log, TEXT("%s - %s Cancelled"), *LexToString(DerivedDataKey), *DatabaseName);
+
 	// Owner.Cancel must be performed before SearchIndex.Reset() in case any task is flying (launched by Owner.LaunchTask)
 	Owner.Cancel();
 
@@ -1872,8 +1876,7 @@ void FAsyncPoseSearchDatabasesManagement::ClearPreCancelled()
 	{
 		if (Tasks[TaskIndex]->GetState() == FPoseSearchDatabaseAsyncCacheTask::EState::PreCancelled)
 		{
-			UE_LOG(LogPoseSearch, Log, TEXT("%s - %s Removed because it was PreCancelled"), *LexToString(Tasks[TaskIndex]->GetDerivedDataKey()), *Tasks[TaskIndex]->GetDatabase()->GetName());
-			Tasks.RemoveAtSwap(TaskIndex, 1, false);
+			Tasks[TaskIndex]->Cancel();
 		}
 	}
 }
@@ -2004,7 +2007,7 @@ void FAsyncPoseSearchDatabasesManagement::Tick(float DeltaTime)
 	// iterating backwards because of the possible RemoveAtSwap 
 	for (int32 TaskIndex = Tasks.Num() - 1; TaskIndex >= 0; --TaskIndex)
 	{
-		if (!Tasks[TaskIndex]->IsValid())
+		if (!Tasks[TaskIndex]->IsValid() || Tasks[TaskIndex]->GetState() == FPoseSearchDatabaseAsyncCacheTask::EState::Cancelled)
 		{
 			Tasks.RemoveAtSwap(TaskIndex, 1, false);
 		}
diff --git a/Engine/Plugins/Experimental/Animation/PoseSearch/Source/Runtime/Private/PoseSearchHistory.cpp b/Engine/Plugins/Experimental/Animation/PoseSearch/Source/Runtime/Private/PoseSearchHistory.cpp
--- a/Engine/Plugins/Experimental/Animation/PoseSearch/Source/Runtime/Private/PoseSearchHistory.cpp
+++ b/Engine/Plugins/Experimental/Animation/PoseSearch/Source/Runtime/Private/PoseSearchHistory.cpp
@@ -348,6 +348,9 @@ void FPoseHistory::DebugDraw(FAnimInstanceProxy& AnimInstanceProxy, FColor Color
 			for (int32 i = 0; i < Entry.Num(); ++i)
 			{
 				const FTransform RootTransform = bValidTrajectory ? Trajectory->GetSampleAtTime(-Entry.Time).GetTransform() : AnimInstanceProxy.GetComponentTransform();
+
+				AnimInstanceProxy.AnimDrawDebugPoint(RootTransform.GetTranslation(), 6.f, FColor::Blue, false, 0.f, ESceneDepthPriorityGroup::SDPG_Foreground);
+
 				PrevGlobalTransforms[i] = Entry.GetComponentSpaceTransform(i) * RootTransform;
 			}
 		}
@@ -356,6 +359,9 @@ void FPoseHistory::DebugDraw(FAnimInstanceProxy& AnimInstanceProxy, FColor Color
 			for (int32 i = 0; i < Entry.Num(); ++i)
 			{
 				const FTransform RootTransform = bValidTrajectory ? Trajectory->GetSampleAtTime(-Entry.Time).GetTransform() : AnimInstanceProxy.GetComponentTransform();
+
+				AnimInstanceProxy.AnimDrawDebugPoint(RootTransform.GetTranslation(), 6.f, FColor::Blue, false, 0.f, ESceneDepthPriorityGroup::SDPG_Foreground);
+
 				const FTransform GlobalTransforms = Entry.GetComponentSpaceTransform(i) * RootTransform;
 
 				AnimInstanceProxy.AnimDrawDebugLine(PrevGlobalTransforms[i].GetTranslation(), GlobalTransforms.GetTranslation(), Color, false, 0.f, ESceneDepthPriorityGroup::SDPG_Foreground);
diff --git a/Engine/Plugins/Experimental/Animation/PoseSearch/Source/Runtime/Private/PoseSearchIndex.cpp b/Engine/Plugins/Experimental/Animation/PoseSearch/Source/Runtime/Private/PoseSearchIndex.cpp
--- a/Engine/Plugins/Experimental/Animation/PoseSearch/Source/Runtime/Private/PoseSearchIndex.cpp
+++ b/Engine/Plugins/Experimental/Animation/PoseSearch/Source/Runtime/Private/PoseSearchIndex.cpp
@@ -314,6 +314,12 @@ FArchive& operator<<(FArchive& Ar, FPoseMetadata& Metadata)
 
 //////////////////////////////////////////////////////////////////////////
 // FSearchIndexAsset
+FFloatInterval FSearchIndexAsset::GetExtrapolationTimeInterval(int32 SchemaSampleRate, const FFloatInterval& AdditionalExtrapolationTime) const
+{
+	return FFloatInterval(FirstSampleIdx / float(SchemaSampleRate) + AdditionalExtrapolationTime.Min,
+		LastSampleIdx / float(SchemaSampleRate) + AdditionalExtrapolationTime.Max);
+}
+
 bool FSearchIndexAsset::operator==(const FSearchIndexAsset& Other) const
 {
 	return
diff --git a/Engine/Plugins/Experimental/Animation/PoseSearch/Source/Runtime/Public/PoseSearch/PoseSearchAssetIndexer.h b/Engine/Plugins/Experimental/Animation/PoseSearch/Source/Runtime/Public/PoseSearch/PoseSearchAssetIndexer.h
--- a/Engine/Plugins/Experimental/Animation/PoseSearch/Source/Runtime/Public/PoseSearch/PoseSearchAssetIndexer.h
+++ b/Engine/Plugins/Experimental/Animation/PoseSearch/Source/Runtime/Public/PoseSearch/PoseSearchAssetIndexer.h
@@ -39,8 +39,8 @@ public:
 		float MaxAcceleration = 0.f;
 	};
 
-	FAssetIndexer(const FBoneContainer& InBoneContainer, const FSearchIndexAsset& InSearchIndexAsset, 
-		const FAssetSamplingContext& InSamplingContext, const UPoseSearchSchema& InSchema, const FAnimationAssetSampler& InAssetSampler);
+	FAssetIndexer(const FBoneContainer& InBoneContainer, const FSearchIndexAsset& InSearchIndexAsset, const FAssetSamplingContext& InSamplingContext,
+		const UPoseSearchSchema& InSchema, const FAnimationAssetSampler& InAssetSampler, const FFloatInterval& InExtrapolationTimeInterval);
 	void AssignWorkingData(int32 InStartPoseIdx, TArrayView<float> InOutFeatureVectorTable, TArrayView<FPoseMetadata> InOutPoseMetadata);
 	void Process(int32 AssetIdx);
 	const FStats& GetStats() const { return Stats; }
@@ -125,6 +125,8 @@ private:
 	const FAssetSamplingContext& SamplingContext;
 	const UPoseSearchSchema& Schema;
 	const FAnimationAssetSampler& AssetSampler;
+	FFloatInterval ExtrapolationTimeInterval = FFloatInterval(-UE_BIG_NUMBER, UE_BIG_NUMBER);
+	
 	int32 StartPoseIdx = 0;
 	
 	TArrayView<float> FeatureVectorTable;
diff --git a/Engine/Plugins/Experimental/Animation/PoseSearch/Source/Runtime/Public/PoseSearch/PoseSearchDatabase.h b/Engine/Plugins/Experimental/Animation/PoseSearch/Source/Runtime/Public/PoseSearch/PoseSearchDatabase.h
--- a/Engine/Plugins/Experimental/Animation/PoseSearch/Source/Runtime/Public/PoseSearch/PoseSearchDatabase.h
+++ b/Engine/Plugins/Experimental/Animation/PoseSearch/Source/Runtime/Public/PoseSearch/PoseSearchDatabase.h
@@ -51,25 +51,6 @@ enum class EPoseSearchMirrorOption : int32
 	Invalid = Num UMETA(Hidden)
 };
 
-USTRUCT()
-struct FPoseSearchExcludeFromDatabaseParameters
-{
-	GENERATED_BODY()
-
-#if WITH_EDITORONLY_DATA
-	// Determines how much of the start of an animation segment is preserved for blending in seconds.
-	// Excluding the beginning of animation segments can help ensure an exact past trajectory is used when building the channels.
-	UPROPERTY(EditAnywhere, Category = "Settings", meta = (DisplayName = "Anim Start Interval"))
-	float SequenceStartInterval = 0.0f;
-
-	// Determines how much of the end of an animation segment is preserved for blending in seconds.
-	// Excluding the end of animation segments helps ensure an exact future trajectory,
-	// and also prevents the selection of a anim segment which will end too soon to be worth selecting.
-	UPROPERTY(EditAnywhere, Category = "Settings", meta = (DisplayName = "Anim End Interval"))
-	float SequenceEndInterval = 0.3f;
-#endif // WITH_EDITORONLY_DATA
-};
-
 USTRUCT()
 struct POSESEARCH_API FPoseSearchDatabaseAnimationAssetBase
 {
@@ -273,11 +254,16 @@ public:
 
 #if WITH_EDITORONLY_DATA
 	// These settings allow users to trim the start and end of animations in the database to preserve start/end frames for blending, and prevent the system from selecting the very last frames before it blends out.
+	// valid animation frames will be AnimationAssetTimeStart + ExcludeFromDatabaseParameters.Min, AnimationAssetTimeEnd + ExcludeFromDatabaseParameters.Max
 	UPROPERTY(EditAnywhere, Category = "Database")
-	FPoseSearchExcludeFromDatabaseParameters ExcludeFromDatabaseParameters;
+	FFloatInterval ExcludeFromDatabaseParameters = FFloatInterval(0.f, -0.3f);
+
+	// extrapolation of animation assets will be clamped by AnimationAssetTimeStart + AdditionalExtrapolationTime.Min, AnimationAssetTimeEnd + AdditionalExtrapolationTime.Max
+	UPROPERTY(EditAnywhere, Category = "Database")
+	FFloatInterval AdditionalExtrapolationTime = FFloatInterval(-100.f, 100.f);
 #endif // WITH_EDITORONLY_DATA
 
-	UPROPERTY(EditAnywhere, Category="Database")
+	UPROPERTY(EditAnywhere, Category = "Database")
 	TArray<FInstancedStruct> AnimationAssets;
 
 	/** Array of tags that can be used as metadata. */
diff --git a/Engine/Plugins/Experimental/Animation/PoseSearch/Source/Runtime/Public/PoseSearch/PoseSearchIndex.h b/Engine/Plugins/Experimental/Animation/PoseSearch/Source/Runtime/Public/PoseSearch/PoseSearchIndex.h
--- a/Engine/Plugins/Experimental/Animation/PoseSearch/Source/Runtime/Public/PoseSearch/PoseSearchIndex.h
+++ b/Engine/Plugins/Experimental/Animation/PoseSearch/Source/Runtime/Public/PoseSearch/PoseSearchIndex.h
@@ -133,6 +133,7 @@ struct FSearchIndexAsset
 		check(SchemaSampleRate > 0);
 	}
 	
+	FFloatInterval GetExtrapolationTimeInterval(int32 SchemaSampleRate, const FFloatInterval& AdditionalExtrapolationTime) const;
 	int32 GetSourceAssetIdx() const { return SourceAssetIdx; }
 	bool IsMirrored() const { return bMirrored; }
 	bool IsLooping() const { return bLooping; }
