motion matching
- SDatabaseDataDetails stats support
- fix for FDatabasePreviewActor::PoseIdxForTimeOffset

diff --git a/Engine/Plugins/Experimental/Animation/PoseSearch/Source/Editor/Private/PoseSearchDatabaseDataDetails.cpp b/Engine/Plugins/Experimental/Animation/PoseSearch/Source/Editor/Private/PoseSearchDatabaseDataDetails.cpp
--- a/Engine/Plugins/Experimental/Animation/PoseSearch/Source/Editor/Private/PoseSearchDatabaseDataDetails.cpp
+++ b/Engine/Plugins/Experimental/Animation/PoseSearch/Source/Editor/Private/PoseSearchDatabaseDataDetails.cpp
@@ -24,6 +24,14 @@ public:
 	{
 	}
 
+	FChannelItem(const FString& InFullLabel, const FString& InCompactLabel)
+	: FullLabel(InFullLabel)
+	, CompactLabel(InCompactLabel)
+	, DataOffset(INDEX_NONE)
+	, Cardinality(INDEX_NONE)
+	{
+	}
+
 	TArray<FChannelItemPtr>& GetChannelItems()
 	{
 		return ChannelItems;
@@ -149,7 +157,7 @@ public:
 				.VAlign(VAlign_Fill)
 				[
 					SNew(SExpanderArrow, SharedThis(this))
-					.StyleSet(ExpanderStyleSet)
+						.StyleSet(ExpanderStyleSet)
 				]
 				+ SHorizontalBox::Slot()
 				.FillWidth(1)
@@ -157,122 +165,167 @@ public:
 				.VAlign(VAlign_Center)
 				[
 					SNew(STextBlock)
-					.Text(FText::FromString(ChannelItem->GetCompactLabel()))
+						.Text(FText::FromString(ChannelItem->GetCompactLabel()))
 				];
 		}
-		
+
 		if (ColumnName == FName("DataOffset"))
 		{
+			if (ChannelItem->GetDataOffset() == INDEX_NONE)
+			{
+				return SNew(STextBlock);
+			}
 			return SNew(STextBlock)
-			.Text(FText::FromString(FString::Printf(TEXT("%d"), ChannelItem->GetDataOffset())));
+				.Text(FText::FromString(FString::Printf(TEXT("%d"), ChannelItem->GetDataOffset())));
 		}
 
 		if (ColumnName == FName("Query"))
 		{
-			return SNew(STextBlock)
-			.Text_Lambda([this, ColumnName]() -> FText
+			if (ChannelItem->GetDataOffset() == INDEX_NONE)
 			{
-				TStringBuilder<256> StringBuilder;
-				if (!ChannelItem->IsExpanded())
-				{
-					if (const TSharedPtr<FDatabaseViewModel> ViewModel = EditorViewModel.Pin())
-					{
-						const int32 DataOffset = ChannelItem->GetDataOffset();
-						const int32 Cardinality = ChannelItem->GetCardinality();
+				return SNew(STextBlock);
+			}
 
-						const TConstArrayView<float> QueryValues = ViewModel->GetQueryVector();
-						if (DataOffset + Cardinality <= QueryValues.Num())
+			return SNew(STextBlock)
+				.Text_Lambda([this, ColumnName]() -> FText
+					{
+						TStringBuilder<256> StringBuilder;
+						if (!ChannelItem->IsExpanded())
 						{
-							if (Cardinality > 1)
+							if (const TSharedPtr<FDatabaseViewModel> ViewModel = EditorViewModel.Pin())
 							{
-								// using only one decimal to keep the string compact
-								for (int32 i = 0; i < Cardinality; ++i)
+								const int32 DataOffset = ChannelItem->GetDataOffset();
+								const int32 Cardinality = ChannelItem->GetCardinality();
+
+								const TConstArrayView<float> QueryValues = ViewModel->GetQueryVector();
+								if (DataOffset + Cardinality <= QueryValues.Num())
 								{
-									if (i != 0)
+									if (Cardinality > 1)
 									{
-										StringBuilder.Append(TEXT(", "));
+										// using only one decimal to keep the string compact
+										for (int32 i = 0; i < Cardinality; ++i)
+										{
+											if (i != 0)
+											{
+												StringBuilder.Append(TEXT(", "));
+											}
+											const float Value = QueryValues[i + DataOffset];
+											StringBuilder.Appendf(TEXT("%.1f"), Value);
+										}
+									}
+									else
+									{
+										// using all the float digits 
+										const float Value = QueryValues[DataOffset];
+										StringBuilder.Appendf(TEXT("%f"), Value);
 									}
-									const float Value = QueryValues[i + DataOffset];
-									StringBuilder.Appendf(TEXT("%.1f"), Value);
 								}
 							}
-							else
-							{
-								// using all the float digits 
-								const float Value = QueryValues[DataOffset];
-								StringBuilder.Appendf(TEXT("%f"), Value);
-							}
 						}
-					}
-				}
-				return FText::FromString(StringBuilder.ToString());
-			});
+						return FText::FromString(StringBuilder.ToString());
+					});
 		}
-		
-		return SNew(STextBlock)
-		.Margin(FMargin(1.f, 1.f, 1.f, 1.f))
-		.Text_Lambda([this, ColumnName]() -> FText
+
+		if (ChannelItem->GetDataOffset() == INDEX_NONE)
 		{
-			TStringBuilder<256> StringBuilder;
-			
-			if (!ChannelItem->IsExpanded())
+			if (ChannelItem->GetCompactLabel() == "[Stats]")
 			{
-				if (const TSharedPtr<FDatabaseViewModel> ViewModel = EditorViewModel.Pin())
-				{
-					const UPoseSearchDatabase* PoseSearchDatabase = ViewModel->GetPoseSearchDatabase();
-					if (FAsyncPoseSearchDatabasesManagement::RequestAsyncBuildIndex(PoseSearchDatabase, ERequestAsyncBuildFlag::ContinueRequest))
-					{
-						if (const FDatabasePreviewActor* FoundPreviewActor = ViewModel->GetPreviewActors().FindByPredicate(
-							[ColumnName](const FDatabasePreviewActor& PreviewActor)
-							{ return *GetNameSafe(PreviewActor.GetActor()) == ColumnName; }))
-						{
-							const int32 PoseIdx = FoundPreviewActor->GetCurrentPoseIndex();
-							const TArray<float> PoseValues = PoseSearchDatabase->GetSearchIndex().GetPoseValuesSafe(PoseIdx);
-
-							const int32 DataOffset = ChannelItem->GetDataOffset();
-							const int32 Cardinality = ChannelItem->GetCardinality();
+				return SNew(STextBlock);
+			}
 
-							if (Cardinality > 1)
+			if (ChannelItem->GetCompactLabel() == "PoseIndex")
+			{
+				return SNew(STextBlock)
+					.Margin(FMargin(1.f, 1.f, 1.f, 1.f))
+					.Text_Lambda([this, ColumnName]() -> FText
+						{
+							FString PoseIdxString;
+							if (!ChannelItem->IsExpanded())
 							{
-								// using only one decimal to keep the string compact
-								for (int32 i = 0; i < Cardinality; ++i)
+								if (const TSharedPtr<FDatabaseViewModel> ViewModel = EditorViewModel.Pin())
 								{
-									if (i != 0)
+									const UPoseSearchDatabase* PoseSearchDatabase = ViewModel->GetPoseSearchDatabase();
+									if (FAsyncPoseSearchDatabasesManagement::RequestAsyncBuildIndex(PoseSearchDatabase, ERequestAsyncBuildFlag::ContinueRequest))
 									{
-										StringBuilder.Append(TEXT(", "));
+										if (const FDatabasePreviewActor* FoundPreviewActor = ViewModel->GetPreviewActors().FindByPredicate(
+											[ColumnName](const FDatabasePreviewActor& PreviewActor)
+											{ return *GetNameSafe(PreviewActor.GetActor()) == ColumnName; }))
+										{
+											const int32 PoseIdx = FoundPreviewActor->GetCurrentPoseIndex();
+											PoseIdxString.AppendInt(PoseIdx);
+										}
 									}
+								}
+							}
+							return FText::FromString(PoseIdxString);
+						});
+			}
+		}
 
-									const int32 PoseValueIndex = i + DataOffset;
-									if (PoseValueIndex < PoseValues.Num())
+		return SNew(STextBlock)
+			.Margin(FMargin(1.f, 1.f, 1.f, 1.f))
+			.Text_Lambda([this, ColumnName]() -> FText
+				{
+					TStringBuilder<256> StringBuilder;
+
+					const int32 DataOffset = ChannelItem->GetDataOffset();
+					const int32 Cardinality = ChannelItem->GetCardinality();
+
+					if (!ChannelItem->IsExpanded() && DataOffset != INDEX_NONE && Cardinality != INDEX_NONE)
+					{
+						if (const TSharedPtr<FDatabaseViewModel> ViewModel = EditorViewModel.Pin())
+						{
+							const UPoseSearchDatabase* PoseSearchDatabase = ViewModel->GetPoseSearchDatabase();
+							if (FAsyncPoseSearchDatabasesManagement::RequestAsyncBuildIndex(PoseSearchDatabase, ERequestAsyncBuildFlag::ContinueRequest))
+							{
+								if (const FDatabasePreviewActor* FoundPreviewActor = ViewModel->GetPreviewActors().FindByPredicate(
+									[ColumnName](const FDatabasePreviewActor& PreviewActor)
+									{ return *GetNameSafe(PreviewActor.GetActor()) == ColumnName; }))
+								{
+									const int32 PoseIdx = FoundPreviewActor->GetCurrentPoseIndex();
+									const TArray<float> PoseValues = PoseSearchDatabase->GetSearchIndex().GetPoseValuesSafe(PoseIdx);
+
+									if (Cardinality > 1)
 									{
-										const float Value = PoseValues[PoseValueIndex];
-										StringBuilder.Appendf(TEXT("%.1f"), Value);
+										// using only one decimal to keep the string compact
+										for (int32 i = 0; i < Cardinality; ++i)
+										{
+											if (i != 0)
+											{
+												StringBuilder.Append(TEXT(", "));
+											}
+
+											const int32 PoseValueIndex = i + DataOffset;
+											if (PoseValueIndex < PoseValues.Num())
+											{
+												const float Value = PoseValues[PoseValueIndex];
+												StringBuilder.Appendf(TEXT("%.1f"), Value);
+											}
+											else
+											{
+												StringBuilder.Append(TEXT("---"));
+											}
+										}
 									}
 									else
 									{
-										StringBuilder.Append(TEXT("---"));
+										// using all the float digits 
+										if (DataOffset < PoseValues.Num())
+										{
+											const float Value = PoseValues[DataOffset];
+											StringBuilder.Appendf(TEXT("%f"), Value);
+										}
+										else
+										{
+											StringBuilder.Append(TEXT("---"));
+										}
 									}
 								}
 							}
-							else
-							{
-								// using all the float digits 
-								if (DataOffset < PoseValues.Num())
-								{
-									const float Value = PoseValues[DataOffset];
-									StringBuilder.Appendf(TEXT("%f"), Value);
-								}
-								else
-								{
-									StringBuilder.Append(TEXT("---"));
-								}
-							}
 						}
 					}
-				}
-			}
-			return FText::FromString(StringBuilder.ToString());
-		});
+					return FText::FromString(StringBuilder.ToString());
+				});
 	}
 };
 
@@ -329,6 +382,7 @@ void SDatabaseDataDetails::Reconstruct(int32 MaxPreviewActors)
 	ChannelItems.Reset();
 
 	RebuildChannelItemsTreeRecursively(ChannelItems, PoseSearchDatabase->Schema->GetChannels());
+	RebuildChannelItemsStats(ChannelItems);
 
 	// generating the header
 	TSharedPtr<SHeaderRow> HeaderRow = SNew(SHeaderRow);
@@ -416,6 +470,18 @@ void SDatabaseDataDetails::RebuildChannelItemsTreeRecursively(TArray<FChannelIte
 	}
 }
 
+void SDatabaseDataDetails::RebuildChannelItemsStats(TArray<FChannelItemPtr>& ChannelItems)
+{
+	TSharedRef<FChannelItem> ChannelItem = MakeShareable(new FChannelItem("SDatabaseDataDetailsStats", "[Stats]"));
+	ChannelItems.Add(ChannelItem);
+
+	ChannelItem->GetChannelItems().Add(MakeShareable(new FChannelItem("SDatabaseDataDetailsPoseIndex", "PoseIndex")));
+	// @todo: add support for additional stats
+	//ChannelItem->GetChannelItems().Add(MakeShareable(new FChannelItem("SDatabaseDataDetailsValueVectorIndex", "ValueVectorIndex")));
+	//ChannelItem->GetChannelItems().Add(MakeShareable(new FChannelItem("SDatabaseDataDetailsPCAValueVectorIndex", "PCAValueVectorIndex")));
+	//ChannelItem->GetChannelItems().Add(MakeShareable(new FChannelItem("SDatabaseDataDetailsSpeed", "Speed")));
+}
+
 } // namespace UE::PoseSearch
 
 #undef LOCTEXT_NAMESPACE
diff --git a/Engine/Plugins/Experimental/Animation/PoseSearch/Source/Editor/Private/PoseSearchDatabaseViewModel.cpp b/Engine/Plugins/Experimental/Animation/PoseSearch/Source/Editor/Private/PoseSearchDatabaseViewModel.cpp
--- a/Engine/Plugins/Experimental/Animation/PoseSearch/Source/Editor/Private/PoseSearchDatabaseViewModel.cpp
+++ b/Engine/Plugins/Experimental/Animation/PoseSearch/Source/Editor/Private/PoseSearchDatabaseViewModel.cpp
@@ -56,7 +56,15 @@ bool FDatabasePreviewActor::SpawnPreviewActor(UWorld* World, const UPoseSearchDa
 
 	IndexAssetIndex = IndexAssetIdx;
 	CurrentPoseIndex = INDEX_NONE;
-	PlayTimeOffset = PoseIdxForTimeOffset < 0 ? 0.f : PoseSearchDatabase->GetRealAssetTime(PoseIdxForTimeOffset);
+
+	if (PoseIdxForTimeOffset < 0)
+	{
+		PlayTimeOffset = 0.f;
+	}
+	else
+	{
+		PlayTimeOffset = PoseSearchDatabase->GetRealAssetTime(PoseIdxForTimeOffset) - IndexAsset.GetFirstSampleTime(PoseSearchDatabase->Schema->SampleRate);
+	}
 
 	// @todo: should we always use the PlayTimeOffset to extract the root transform?
 	if (PlayTimeOffset != 0.f)
diff --git a/Engine/Plugins/Experimental/Animation/PoseSearch/Source/Editor/Public/PoseSearchDatabaseDataDetails.h b/Engine/Plugins/Experimental/Animation/PoseSearch/Source/Editor/Public/PoseSearchDatabaseDataDetails.h
--- a/Engine/Plugins/Experimental/Animation/PoseSearch/Source/Editor/Public/PoseSearchDatabaseDataDetails.h
+++ b/Engine/Plugins/Experimental/Animation/PoseSearch/Source/Editor/Public/PoseSearchDatabaseDataDetails.h
@@ -24,6 +24,7 @@ namespace UE::PoseSearch
 
 	private:
 		static void RebuildChannelItemsTreeRecursively(TArray<FChannelItemPtr>& ChannelItems, TConstArrayView<TObjectPtr<UPoseSearchFeatureChannel>> Channels);
+		static void RebuildChannelItemsStats(TArray<FChannelItemPtr>& ChannelItems);
 		static void TrackExpandedItems(const TArray<FChannelItemPtr>& ChannelItems, TMap<const FString, bool>& ExpandedItems);
 		static void SetExpandedItems(TArray<FChannelItemPtr>& ChannelItems, const TMap<const FString, bool>& ExpandedItems, SChannelItemsTreeView* ChannelItemsTreeView);
 
