motion matching - fix for referencing released memory

diff --git a/Engine/Plugins/Experimental/Animation/PoseSearch/Source/Editor/Private/PoseSearchDebuggerDatabaseView.cpp b/Engine/Plugins/Experimental/Animation/PoseSearch/Source/Editor/Private/PoseSearchDebuggerDatabaseView.cpp
--- a/Engine/Plugins/Experimental/Animation/PoseSearch/Source/Editor/Private/PoseSearchDebuggerDatabaseView.cpp
+++ b/Engine/Plugins/Experimental/Animation/PoseSearch/Source/Editor/Private/PoseSearchDebuggerDatabaseView.cpp
@@ -543,8 +543,7 @@ void SDebuggerDatabaseView::OnShowAllPosesCheckboxChanged(ECheckBoxState State)
 	{
 		if (TSharedPtr<FDebuggerViewModel> ViewModel = DebuggerView->GetViewModel())
 		{
-			const FTraceMotionMatchingStateMessage* MotionMatchingState = ViewModel.Get()->GetMotionMatchingState();
-			if (MotionMatchingState)
+			if (const FTraceMotionMatchingStateMessage* MotionMatchingState = ViewModel.Get()->GetMotionMatchingState())
 			{
 				Update(*MotionMatchingState);
 			}
diff --git a/Engine/Plugins/Experimental/Animation/PoseSearch/Source/Editor/Private/PoseSearchDebuggerView.cpp b/Engine/Plugins/Experimental/Animation/PoseSearch/Source/Editor/Private/PoseSearchDebuggerView.cpp
--- a/Engine/Plugins/Experimental/Animation/PoseSearch/Source/Editor/Private/PoseSearchDebuggerView.cpp
+++ b/Engine/Plugins/Experimental/Animation/PoseSearch/Source/Editor/Private/PoseSearchDebuggerView.cpp
@@ -212,14 +212,14 @@ void SDebuggerView::Tick(const FGeometry& AllottedGeometry, const double InCurre
 		// Haven't updated since passing through frame gate, update once
 		else if (!bUpdated)
 		{
-		Model->OnUpdate();
-		if (UpdateNodeSelection())
-		{
-			Model->OnUpdateNodeSelection(SelectedNodeId);
-			UpdateViews();
+			Model->OnUpdate();
+			if (UpdateNodeSelection())
+			{
+				Model->OnUpdateNodeSelection(SelectedNodeId);
+				UpdateViews();
+			}
+			bUpdated = true;
 		}
-		bUpdated = true;
-	}
 	}
 
 	// Draw visualization every tick
@@ -290,8 +290,7 @@ bool SDebuggerView::UpdateNodeSelection()
 
 void SDebuggerView::UpdateViews() const
 {
-	const FTraceMotionMatchingStateMessage* State = ViewModel.Get()->GetMotionMatchingState();
-	if (State)
+	if (const FTraceMotionMatchingStateMessage* State = ViewModel.Get()->GetMotionMatchingState())
 	{
 		DatabaseView->Update(*State);
 		DetailsView->Update(*State);
@@ -438,9 +437,7 @@ int32 SDebuggerView::SelectView() const
 void SDebuggerView::OnPoseSelectionChanged(const UPoseSearchDatabase* Database, int32 DbPoseIdx, float Time)
 {
 	const TSharedPtr<FDebuggerViewModel> Model = ViewModel.Get();
-	const FTraceMotionMatchingStateMessage* State = Model->GetMotionMatchingState();
-
-	if (State)
+	if (const FTraceMotionMatchingStateMessage* State = Model->GetMotionMatchingState())
 	{
 		DetailsView->Update(*State);
 	}
diff --git a/Engine/Plugins/Experimental/Animation/PoseSearch/Source/Editor/Private/PoseSearchDebuggerViewModel.cpp b/Engine/Plugins/Experimental/Animation/PoseSearch/Source/Editor/Private/PoseSearchDebuggerViewModel.cpp
--- a/Engine/Plugins/Experimental/Animation/PoseSearch/Source/Editor/Private/PoseSearchDebuggerViewModel.cpp
+++ b/Engine/Plugins/Experimental/Animation/PoseSearch/Source/Editor/Private/PoseSearchDebuggerViewModel.cpp
@@ -38,12 +38,20 @@ FDebuggerViewModel::~FDebuggerViewModel()
 
 const FTraceMotionMatchingStateMessage* FDebuggerViewModel::GetMotionMatchingState() const
 {
-	return ActiveMotionMatchingState;
+	if (MotionMatchingStates.IsValidIndex(ActiveMotionMatchingStateIdx))
+	{
+		return &MotionMatchingStates[ActiveMotionMatchingStateIdx];
+	}
+	return nullptr;
 }
 
 const UPoseSearchDatabase* FDebuggerViewModel::GetCurrentDatabase() const
 {
-	return ActiveMotionMatchingState ? ActiveMotionMatchingState->GetCurrentDatabase() : nullptr;
+	if (const FTraceMotionMatchingStateMessage* State = GetMotionMatchingState())
+	{
+		return State->GetCurrentDatabase();
+	}
+	return nullptr;
 }
 
 void FDebuggerViewModel::ShowSelectedSkeleton(const UPoseSearchDatabase* Database, int32 DbPoseIdx, float Time)
@@ -123,32 +131,33 @@ void FDebuggerViewModel::OnUpdateNodeSelection(int32 InNodeId)
 		return;
 	}
 
-	ActiveMotionMatchingState = nullptr;
-
 	// Find node in all motion matching states this frame
+	ActiveMotionMatchingStateIdx = INDEX_NONE;
 	const int32 NodesNum = NodeIds.Num();
 	for (int32 i = 0; i < NodesNum; ++i)
 	{
 		if (NodeIds[i] == InNodeId)
 		{
-			ActiveMotionMatchingState = MotionMatchingStates[i];
+			ActiveMotionMatchingStateIdx = i;
 			break;
 		}
 	}
 
-	if (ActiveMotionMatchingState)
+	if (const FTraceMotionMatchingStateMessage* State = GetMotionMatchingState())
 	{
-		const UPoseSearchDatabase* CurrentDatabase = ActiveMotionMatchingState->GetCurrentDatabase();
-		if (FAsyncPoseSearchDatabasesManagement::RequestAsyncBuildIndex(CurrentDatabase, ERequestAsyncBuildFlag::ContinueRequest))
+		if (const UPoseSearchDatabase* CurrentDatabase = State->GetCurrentDatabase())
 		{
-			const FSearchIndex& CurrentSearchIndex = CurrentDatabase->GetSearchIndex();
-			int32 CurrentPoseIdx = ActiveMotionMatchingState->GetCurrentDatabasePoseIndex();
-			if (const FSearchIndexAsset* IndexAsset = CurrentSearchIndex.GetAssetForPoseSafe(CurrentPoseIdx))
+			if (FAsyncPoseSearchDatabasesManagement::RequestAsyncBuildIndex(CurrentDatabase, ERequestAsyncBuildFlag::ContinueRequest))
 			{
-				Skeletons[Asset].bMirrored = IndexAsset->IsMirrored();
-				Skeletons[Asset].SourceDatabase = CurrentDatabase;
-				Skeletons[Asset].AssetIdx = IndexAsset->GetSourceAssetIdx();
-				Skeletons[Asset].BlendParameters = IndexAsset->GetBlendParameters();
+				const FSearchIndex& CurrentSearchIndex = CurrentDatabase->GetSearchIndex();
+				int32 CurrentPoseIdx = State->GetCurrentDatabasePoseIndex();
+				if (const FSearchIndexAsset* IndexAsset = CurrentSearchIndex.GetAssetForPoseSafe(CurrentPoseIdx))
+				{
+					Skeletons[Asset].bMirrored = IndexAsset->IsMirrored();
+					Skeletons[Asset].SourceDatabase = CurrentDatabase;
+					Skeletons[Asset].AssetIdx = IndexAsset->GetSourceAssetIdx();
+					Skeletons[Asset].BlendParameters = IndexAsset->GetBlendParameters();
+				}
 			}
 		}
 	}
@@ -220,7 +229,9 @@ void FDebuggerViewModel::UpdateFromTimeline()
 		if (Message)
 		{
 			NodeIds.Add(Message->NodeId);
-			MotionMatchingStates.Add(Message);
+
+			// @todo: figure out if we can avoid this copy
+			MotionMatchingStates.Add(*Message);
 			SkeletalMeshComponentId = Message->SkeletalMeshComponentId;
 		}
 	});
diff --git a/Engine/Plugins/Experimental/Animation/PoseSearch/Source/Editor/Private/PoseSearchDebuggerViewModel.h b/Engine/Plugins/Experimental/Animation/PoseSearch/Source/Editor/Private/PoseSearchDebuggerViewModel.h
--- a/Engine/Plugins/Experimental/Animation/PoseSearch/Source/Editor/Private/PoseSearchDebuggerViewModel.h
+++ b/Engine/Plugins/Experimental/Animation/PoseSearch/Source/Editor/Private/PoseSearchDebuggerViewModel.h
@@ -98,10 +98,10 @@ private:
 	TArray<int32> NodeIds;
 	
 	/** List of all updated motion matching states per node */
-	TArray<const FTraceMotionMatchingStateMessage*> MotionMatchingStates;
+	TArray<FTraceMotionMatchingStateMessage> MotionMatchingStates;
 	
-	/** Currently active motion matching state based on node selection in the view */
-	const FTraceMotionMatchingStateMessage* ActiveMotionMatchingState = nullptr;
+	/** Currently active motion matching state index based on node selection in the view */
+	int32 ActiveMotionMatchingStateIdx = INDEX_NONE;
 
 	/** Current Skeletal Mesh Component Id for the AnimInstance */
 	uint64 SkeletalMeshComponentId = 0;
diff --git a/Engine/Plugins/Experimental/Animation/PoseSearch/Source/Runtime/Private/Trace/PoseSearchTraceLogger.cpp b/Engine/Plugins/Experimental/Animation/PoseSearch/Source/Runtime/Private/Trace/PoseSearchTraceLogger.cpp
--- a/Engine/Plugins/Experimental/Animation/PoseSearch/Source/Runtime/Private/Trace/PoseSearchTraceLogger.cpp
+++ b/Engine/Plugins/Experimental/Animation/PoseSearch/Source/Runtime/Private/Trace/PoseSearchTraceLogger.cpp
@@ -113,31 +113,34 @@ void FTraceMotionMatchingState::Output(const UObject* AnimInstance, int32 NodeId
 
 const UPoseSearchDatabase* FTraceMotionMatchingState::GetCurrentDatabase() const
 {
-	if (CurrentDbEntryIdx == INDEX_NONE)
+	const UPoseSearchDatabase* Database = nullptr;
+	if (DatabaseEntries.IsValidIndex(CurrentDbEntryIdx))
 	{
-		return nullptr;
+		Database = GetObjectFromId<UPoseSearchDatabase>(DatabaseEntries[CurrentDbEntryIdx].DatabaseId);
 	}
-
-	const FTraceMotionMatchingStateDatabaseEntry& DbEntry = DatabaseEntries[CurrentDbEntryIdx];
-	const UPoseSearchDatabase* Database = GetObjectFromId<UPoseSearchDatabase>(DbEntry.DatabaseId);
 	return Database;
 }
 
 int32 FTraceMotionMatchingState::GetCurrentDatabasePoseIndex() const
 {
-	const FTraceMotionMatchingStatePoseEntry* PoseEntry = GetCurrentPoseEntry();
-	return PoseEntry ? PoseEntry->DbPoseIdx : INDEX_NONE;
+	if (const FTraceMotionMatchingStatePoseEntry* PoseEntry = GetCurrentPoseEntry())
+	{
+		return PoseEntry->DbPoseIdx;
+	}
+	return INDEX_NONE;
 }
 
 const FTraceMotionMatchingStatePoseEntry* FTraceMotionMatchingState::GetCurrentPoseEntry() const
 {
-	if ((CurrentDbEntryIdx == INDEX_NONE) || (CurrentPoseEntryIdx == INDEX_NONE))
+	if (DatabaseEntries.IsValidIndex(CurrentDbEntryIdx))
 	{
-		return nullptr;
+		const FTraceMotionMatchingStateDatabaseEntry& DbEntry = DatabaseEntries[CurrentDbEntryIdx];
+		if (DbEntry.PoseEntries.IsValidIndex(CurrentPoseEntryIdx))
+		{
+			return &DbEntry.PoseEntries[CurrentPoseEntryIdx];
+		}
 	}
-
-	const FTraceMotionMatchingStateDatabaseEntry& DbEntry = DatabaseEntries[CurrentDbEntryIdx];
-	return &DbEntry.PoseEntries[CurrentPoseEntryIdx];
+	return nullptr;
 }
 
 } // namespace UE::PoseSearch
\ No newline at end of file
