motion matching - excluding UAnimationModifier from pose search database DDC key

diff --git a/Engine/Plugins/Experimental/Animation/PoseSearch/Source/Runtime/PoseSearch.Build.cs b/Engine/Plugins/Experimental/Animation/PoseSearch/Source/Runtime/PoseSearch.Build.cs
--- a/Engine/Plugins/Experimental/Animation/PoseSearch/Source/Runtime/PoseSearch.Build.cs
+++ b/Engine/Plugins/Experimental/Animation/PoseSearch/Source/Runtime/PoseSearch.Build.cs
@@ -31,6 +31,8 @@ public class PoseSearch : ModuleRules
 		{
 			PrivateDependencyModuleNames.AddRange(
 				new string[] {
+					"AnimationModifiers",
+					"AnimationBlueprintLibrary",
 					"DerivedDataCache",
 					"UnrealEd"
 				}
diff --git a/Engine/Plugins/Experimental/Animation/PoseSearch/Source/Runtime/Private/PoseSearchDerivedDataKey.cpp b/Engine/Plugins/Experimental/Animation/PoseSearch/Source/Runtime/Private/PoseSearchDerivedDataKey.cpp
--- a/Engine/Plugins/Experimental/Animation/PoseSearch/Source/Runtime/Private/PoseSearchDerivedDataKey.cpp
+++ b/Engine/Plugins/Experimental/Animation/PoseSearch/Source/Runtime/Private/PoseSearchDerivedDataKey.cpp
@@ -6,6 +6,7 @@
 #include "Animation/AnimBoneCompressionSettings.h"
 #include "Animation/AnimCurveCompressionSettings.h"
 #include "Animation/AnimData/IAnimationDataModel.h"
+#include "AnimationModifier.h"
 #include "Animation/MirrorDataTable.h"
 #include "Animation/Skeleton.h"
 #include "Engine/SkeletalMesh.h"
@@ -40,7 +41,7 @@ FKeyBuilder::FKeyBuilder(const UObject* Object, bool bUseDataVer, bool bPerformC
 	if (bUseDataVer)
 	{
 		// used to invalidate the key without having to change POSESEARCHDB_DERIVEDDATA_VER all the times
-		int32 POSESEARCHDB_DERIVEDDATA_VER_SMALL = 231;
+		int32 POSESEARCHDB_DERIVEDDATA_VER_SMALL = 237;
 		FGuid VersionGuid = FDevSystemGuids::GetSystemGuid(FDevSystemGuids::Get().POSESEARCHDB_DERIVEDDATA_VER);
 
 		*this << VersionGuid;
@@ -151,10 +152,16 @@ FArchive& FKeyBuilder::operator<<(class UObject*& Object)
 				UE_LOG(LogPoseSearch, Log, TEXT("%sTransient '%s' (%s)"), *GetIndentation(), *Object->GetName(), *Object->GetClass()->GetName());
 				#endif
 			}
+			else if (IsExcludedType(Object))
+			{
+				#if UE_POSE_SEARCH_DERIVED_DATA_LOGGING
+				UE_LOG(LogPoseSearch, Log, TEXT("%sExcluded '%s' (%s)"), *GetIndentation(), *Object->GetName(), *Object->GetClass()->GetName());
+				#endif
+			}
 			else
 			{
 				bool bAlreadyProcessed = false;
-				ObjectsAlreadySerialized.Add(Object, &bAlreadyProcessed);
+				Dependencies.Add(Object, &bAlreadyProcessed);
 
 				// If we haven't already serialized this object
 				if (bAlreadyProcessed)
@@ -164,7 +171,7 @@ FArchive& FKeyBuilder::operator<<(class UObject*& Object)
 					#endif
 				}
 				// for specific types we only add their names to the hash
-				else if (AddNameOnly(Object))
+				else if (IsAddNameOnlyType(Object))
 				{
 					#if UE_POSE_SEARCH_DERIVED_DATA_LOGGING
 					UE_LOG(LogPoseSearch, Log, TEXT("%sAddingNameOnly '%s' (%s)"), *GetIndentation(), *Object->GetName(), *Object->GetClass()->GetName());
@@ -220,23 +227,31 @@ FIoHash FKeyBuilder::Finalize() const
 
 const TSet<const UObject*>& FKeyBuilder::GetDependencies() const
 {
-	return ObjectsAlreadySerialized;
+	return Dependencies;
+}
+
+// to keep the key generation lightweight, we don't hash these types
+bool FKeyBuilder::IsExcludedType(class UObject* Object)
+{
+	return
+		nullptr != Cast<UAnimationModifier>(Object);
 }
 
-bool FKeyBuilder::AddNameOnly(class UObject* Object) const
+// to keep the key generation lightweight, we hash only the full names for these types. Object(s) will be added to Dependencies
+bool FKeyBuilder::IsAddNameOnlyType(class UObject* Object)
 {
 	return
-		Cast<IAnimationDataModel>(Object) ||
-		Cast<UActorComponent>(Object) ||
-		Cast<UAnimBoneCompressionSettings>(Object) ||
-		Cast<UAnimCurveCompressionSettings>(Object) ||
-		Cast<UAssetImportData>(Object) ||
-		Cast<UFunction>(Object) ||
-		Cast<UMirrorDataTable>(Object) ||
-		Cast<USkeletalMesh>(Object) ||
-		Cast<USkeletalMeshSocket>(Object) ||
-		Cast<USkeleton>(Object) ||
-		Cast<UStreamableRenderAsset>(Object);
+		nullptr != Cast<IAnimationDataModel>(Object) ||
+		nullptr != Cast<UActorComponent>(Object) ||
+		nullptr != Cast<UAnimBoneCompressionSettings>(Object) ||
+		nullptr != Cast<UAnimCurveCompressionSettings>(Object) ||
+		nullptr != Cast<UAssetImportData>(Object) ||
+		nullptr != Cast<UFunction>(Object) ||
+		nullptr != Cast<UMirrorDataTable>(Object) ||
+		nullptr != Cast<USkeletalMesh>(Object) ||
+		nullptr != Cast<USkeletalMeshSocket>(Object) ||
+		nullptr != Cast<USkeleton>(Object) ||
+		nullptr != Cast<UStreamableRenderAsset>(Object);
 }
 
 #if UE_POSE_SEARCH_DERIVED_DATA_LOGGING
diff --git a/Engine/Plugins/Experimental/Animation/PoseSearch/Source/Runtime/Public/PoseSearch/PoseSearchDerivedDataKey.h b/Engine/Plugins/Experimental/Animation/PoseSearch/Source/Runtime/Public/PoseSearch/PoseSearchDerivedDataKey.h
--- a/Engine/Plugins/Experimental/Animation/PoseSearch/Source/Runtime/Public/PoseSearch/PoseSearchDerivedDataKey.h
+++ b/Engine/Plugins/Experimental/Animation/PoseSearch/Source/Runtime/Public/PoseSearch/PoseSearchDerivedDataKey.h
@@ -53,8 +53,11 @@ public:
 	const TSet<const UObject*>& GetDependencies() const;
 
 protected:
-	// to keep the key generation lightweight, we hash only the full names for these types
-	bool AddNameOnly(class UObject* Object) const;
+	// to keep the key generation lightweight, we don't hash these types
+	static bool IsExcludedType(class UObject* Object);
+
+	// to keep the key generation lightweight, we hash only the full names for these types. Object(s) will be added to Dependencies
+	static bool IsAddNameOnlyType(class UObject* Object);
 
 #if UE_POSE_SEARCH_DERIVED_DATA_LOGGING
 	FString GetIndentation() const;
@@ -64,7 +67,7 @@ protected:
 	HashBuilderType Hasher;
 
 	// Set of objects that have already been serialized
-	TSet<const UObject*> ObjectsAlreadySerialized;
+	TSet<const UObject*> Dependencies;
 
 	// Object currently being serialized
 	const UObject* ObjectBeingSerialized = nullptr;
