motion matching - fix for weird animation behaviours in PIE when editing pose search databases dependant assets (resetting CurrentSearchResult if any DDC indexing on the requested databases is still in progress)

diff --git a/Engine/Plugins/Experimental/Animation/PoseSearch/Source/Runtime/Private/PoseSearchDatabase.cpp b/Engine/Plugins/Experimental/Animation/PoseSearch/Source/Runtime/Private/PoseSearchDatabase.cpp
--- a/Engine/Plugins/Experimental/Animation/PoseSearch/Source/Runtime/Private/PoseSearchDatabase.cpp
+++ b/Engine/Plugins/Experimental/Animation/PoseSearch/Source/Runtime/Private/PoseSearchDatabase.cpp
@@ -852,6 +852,7 @@ UE::PoseSearch::FSearchResult UPoseSearchDatabase::Search(UE::PoseSearch::FSearc
 #if WITH_EDITOR
 	if (!FAsyncPoseSearchDatabasesManagement::RequestAsyncBuildIndex(this, ERequestAsyncBuildFlag::ContinueRequest))
 	{
+		SearchContext.SetAsyncBuildIndexInProgress();
 		return Result;
 	}
 #endif // WITH_EDITOR
@@ -960,6 +961,7 @@ UE::PoseSearch::FSearchResult UPoseSearchDatabase::SearchContinuingPose(UE::Pose
 #if WITH_EDITOR
 	if (!FAsyncPoseSearchDatabasesManagement::RequestAsyncBuildIndex(this, ERequestAsyncBuildFlag::ContinueRequest))
 	{
+		SearchContext.SetAsyncBuildIndexInProgress();
 		return Result;
 	}
 #endif // WITH_EDITOR
diff --git a/Engine/Plugins/Experimental/Animation/PoseSearch/Source/Runtime/Private/PoseSearchLibrary.cpp b/Engine/Plugins/Experimental/Animation/PoseSearch/Source/Runtime/Private/PoseSearchLibrary.cpp
--- a/Engine/Plugins/Experimental/Animation/PoseSearch/Source/Runtime/Private/PoseSearchLibrary.cpp
+++ b/Engine/Plugins/Experimental/Animation/PoseSearch/Source/Runtime/Private/PoseSearchLibrary.cpp
@@ -363,6 +363,15 @@ void UPoseSearchLibrary::UpdateMotionMatchingState(
 		}
 #endif // UE_POSE_SEARCH_TRACE_ENABLED
 
+		
+#if WITH_EDITOR
+		// resetting CurrentSearchResult if any DDC indexing on the requested databases is still in progress
+		if (SearchContext.IsAsyncBuildIndexInProgress())
+		{
+			InOutMotionMatchingState.CurrentSearchResult.Reset();
+		}
+#endif // WITH_EDITOR
+
 #if !NO_LOGGING
 		if (!SearchResult.IsValid())
 		{
@@ -371,6 +380,14 @@ void UPoseSearchLibrary::UpdateMotionMatchingState(
 			StringBuilder << bForceInterrupt;
 			StringBuilder << "], CanAdvance [";
 			StringBuilder << bCanAdvance;
+			StringBuilder << "], Indexing [";
+
+#if WITH_EDITOR
+			StringBuilder << SearchContext.IsAsyncBuildIndexInProgress();
+//#else // WITH_EDITOR
+			StringBuilder << false;
+#endif // WITH_EDITOR
+
 			StringBuilder << "], Databases [";
 
 			for (int32 DatabaseIndex = 0; DatabaseIndex < Databases.Num(); ++DatabaseIndex)
@@ -428,11 +445,7 @@ void UPoseSearchLibrary::UpdateMotionMatchingState(
 		const UPoseSearchDatabase* CurResultDatabase = CurResult.Database.Get();
 
 #if WITH_EDITOR
-		// in case we're still indexing MotionMatchingState.CurrentSearchResult.Database we Reset the MotionMatchingState
-		if (!FAsyncPoseSearchDatabasesManagement::RequestAsyncBuildIndex(CurResultDatabase, ERequestAsyncBuildFlag::ContinueRequest))
-		{
-		}
-		else
+		if (FAsyncPoseSearchDatabasesManagement::RequestAsyncBuildIndex(CurResultDatabase, ERequestAsyncBuildFlag::ContinueRequest))
 #endif // WITH_EDITOR
 		{
 			if (bDebugDrawCurResult)
diff --git a/Engine/Plugins/Experimental/Animation/PoseSearch/Source/Runtime/Public/PoseSearch/PoseSearchContext.h b/Engine/Plugins/Experimental/Animation/PoseSearch/Source/Runtime/Public/PoseSearch/PoseSearchContext.h
--- a/Engine/Plugins/Experimental/Animation/PoseSearch/Source/Runtime/Public/PoseSearch/PoseSearchContext.h
+++ b/Engine/Plugins/Experimental/Animation/PoseSearch/Source/Runtime/Public/PoseSearch/PoseSearchContext.h
@@ -190,6 +190,12 @@ struct POSESEARCH_API FSearchContext
 	// returns the world space transform of the bone SchemaBoneIdx at time SampleTime
 	FTransform GetWorldBoneTransformAtTime(float SampleTime, const UPoseSearchSchema* Schema = nullptr, int8 SchemaBoneIdx = RootSchemaBoneIdx);
 	
+#if WITH_EDITOR
+	void SetAsyncBuildIndexInProgress() { bAsyncBuildIndexInProgress = true; }
+	void ResetAsyncBuildIndexInProgress() { bAsyncBuildIndexInProgress = false; }
+	bool IsAsyncBuildIndexInProgress() const { return bAsyncBuildIndexInProgress; }
+#endif // WITH_EDITOR
+
 private:
 	FVector GetSamplePositionInternal(float SampleTime, float OriginTime, const UPoseSearchSchema* Schema, int8 SchemaSampleBoneIdx = RootSchemaBoneIdx, int8 SchemaOriginBoneIdx = RootSchemaBoneIdx, const FVector* SampleBonePositionWorldOverride = nullptr);
 	FQuat GetSampleRotationInternal(float SampleTime, float OriginTime, const UPoseSearchSchema* Schema, int8 SchemaSampleBoneIdx = RootSchemaBoneIdx, int8 SchemaOriginBoneIdx = RootSchemaBoneIdx, const FQuat* SampleBoneRotationWorldOverride = nullptr);
@@ -217,6 +223,10 @@ private:
 
 	float CurrentBestTotalCost = MAX_flt;
 	
+#if WITH_EDITOR
+	bool bAsyncBuildIndexInProgress = false;
+#endif // WITH_EDITOR
+
 #if UE_POSE_SEARCH_TRACE_ENABLED
 
 	struct FPoseCandidateIdCost
