loco - adding InterruptOnDatabaseChangeAndInvalidateContinuingPose

diff --git a/Engine/Plugins/Experimental/Animation/PoseSearch/Source/Runtime/Private/PoseSearchLibrary.cpp b/Engine/Plugins/Experimental/Animation/PoseSearch/Source/Runtime/Private/PoseSearchLibrary.cpp
--- a/Engine/Plugins/Experimental/Animation/PoseSearch/Source/Runtime/Private/PoseSearchLibrary.cpp
+++ b/Engine/Plugins/Experimental/Animation/PoseSearch/Source/Runtime/Private/PoseSearchLibrary.cpp
@@ -40,10 +40,36 @@ namespace UE::PoseSearch
 		{
 		case EPoseSearchInterruptMode::DoNotInterrupt:
 			return false;
-		case EPoseSearchInterruptMode::InterruptOnDatabaseChange:
+
+		case EPoseSearchInterruptMode::InterruptOnDatabaseChange:	// Fall through
+		case EPoseSearchInterruptMode::InterruptOnDatabaseChangeAndInvalidateContinuingPose:
 			return !Databases.Contains(CurrentResultDatabase);
-		case EPoseSearchInterruptMode::ForceInterrupt:
+
+		case EPoseSearchInterruptMode::ForceInterrupt:				// Fall through
+		case EPoseSearchInterruptMode::ForceInterruptAndInvalidateContinuingPose:
 			return true;
+
+		default:
+			unimplemented();
+			return false;
+		}
+	}
+
+	static bool IsInvalidatingContinuingPose(EPoseSearchInterruptMode InterruptMode, const UPoseSearchDatabase* CurrentResultDatabase, const TArray<TObjectPtr<const UPoseSearchDatabase>>& Databases)
+	{
+		switch (InterruptMode)
+		{
+		case EPoseSearchInterruptMode::DoNotInterrupt:				// Fall through
+		case EPoseSearchInterruptMode::InterruptOnDatabaseChange:	// Fall through
+		case EPoseSearchInterruptMode::ForceInterrupt:	
+			return false;
+
+		case EPoseSearchInterruptMode::InterruptOnDatabaseChangeAndInvalidateContinuingPose:
+			return !Databases.Contains(CurrentResultDatabase);
+
+		case EPoseSearchInterruptMode::ForceInterruptAndInvalidateContinuingPose:
+			return true;
+
 		default:
 			unimplemented();
 			return false;
@@ -319,6 +345,13 @@ void UPoseSearchLibrary::UpdateMotionMatchingState(
 	FMemMark Mark(FMemStack::Get());
 	const UAnimInstance* AnimInstance = Cast<const UAnimInstance>(Context.AnimInstanceProxy->GetAnimInstanceObject());
 	check(AnimInstance);
+
+	const UPoseSearchDatabase* CurrentResultDatabase = InOutMotionMatchingState.CurrentSearchResult.Database.Get();
+	if (IsInvalidatingContinuingPose(InterruptMode, CurrentResultDatabase, Databases))
+	{
+		InOutMotionMatchingState.CurrentSearchResult.Reset();
+	}
+
 	FSearchContext SearchContext(AnimInstance, History, TConstArrayView<const UAnimationAsset*>(), &TrajectoryRootSpace, 0.f,
 		&InOutMotionMatchingState.PoseIndicesHistory, InOutMotionMatchingState.CurrentSearchResult, PoseJumpThresholdTime);
 
@@ -329,8 +362,6 @@ void UPoseSearchLibrary::UpdateMotionMatchingState(
 	if (bSearch)
 	{
 		InOutMotionMatchingState.ElapsedPoseSearchTime = 0.f;
-
-		const UPoseSearchDatabase* CurrentResultDatabase = SearchContext.GetCurrentResult().Database.Get();
 		const bool bForceInterrupt = IsForceInterrupt(InterruptMode, CurrentResultDatabase, Databases);
 
 		// Evaluate continuing pose
diff --git a/Engine/Plugins/Experimental/Animation/PoseSearch/Source/Runtime/Public/PoseSearch/PoseSearchLibrary.h b/Engine/Plugins/Experimental/Animation/PoseSearch/Source/Runtime/Public/PoseSearch/PoseSearchLibrary.h
--- a/Engine/Plugins/Experimental/Animation/PoseSearch/Source/Runtime/Public/PoseSearch/PoseSearchLibrary.h
+++ b/Engine/Plugins/Experimental/Animation/PoseSearch/Source/Runtime/Public/PoseSearch/PoseSearchLibrary.h
@@ -28,8 +28,16 @@ enum class EPoseSearchInterruptMode : uint8
 	// continuing pose search will be interrupted if its database is not listed in the searchable databases
 	InterruptOnDatabaseChange,
 
+	// continuing pose search will be interrupted if its database is not listed in the searchable databases, 
+	// and continuing pose will be invalidated (forcing the schema to use pose history to build the query)
+	InterruptOnDatabaseChangeAndInvalidateContinuingPose,
+
 	// continuing pose search will always be interrupted
 	ForceInterrupt,
+
+	/// continuing pose search will always be interrupted
+	// and continuing pose will be invalidated (forcing the schema to use pose history to build the query)
+	ForceInterruptAndInvalidateContinuingPose,
 };
 
 struct FAnimationUpdateContext;
