motion matching - renaming FKDTree::Impl to FKDTree::KDTreeImplementation to avoid clashing with check macro struct Impl

diff --git a/Engine/Plugins/Experimental/Animation/PoseSearch/Source/Runtime/Private/KDTree.cpp b/Engine/Plugins/Experimental/Animation/PoseSearch/Source/Runtime/Private/KDTree.cpp
--- a/Engine/Plugins/Experimental/Animation/PoseSearch/Source/Runtime/Private/KDTree.cpp
+++ b/Engine/Plugins/Experimental/Animation/PoseSearch/Source/Runtime/Private/KDTree.cpp
@@ -147,19 +147,19 @@ private:
 
 FKDTree::FKDTree(AccessorType Count, AccessorType Dim, const float* Data, AccessorType MaxLeafSize)
 : DataSource(Count, Dim, Data)
-, Impl(nullptr)
+, KDTreeImplementation(nullptr)
 {
 #if UE_POSE_SEARCH_USE_NANOFLANN
 	if (Count > 0 && Dim > 0 && Data)
 	{
-		Impl = new FKDTreeImplementation(Dim, DataSource, nanoflann::KDTreeSingleIndexAdaptorParams(MaxLeafSize));
+		KDTreeImplementation = new FKDTreeImplementation(Dim, DataSource, nanoflann::KDTreeSingleIndexAdaptorParams(MaxLeafSize));
 	}
 #endif // UE_POSE_SEARCH_USE_NANOFLANN
 }
 
 FKDTree::FKDTree()
 : DataSource(0, 0, nullptr)
-, Impl(nullptr)
+, KDTreeImplementation(nullptr)
 {
 }
 
@@ -171,9 +171,9 @@ FKDTree::~FKDTree()
 #if UE_POSE_SEARCH_USE_NANOFLANN
 void CopySubTree(FKDTree& KDTree, FKDTreeImplementation::NodePtr& ThisNode, const FKDTreeImplementation::NodePtr& OtherNode)
 {
-	check(KDTree.Impl);
+	check(KDTree.KDTreeImplementation);
 
-	ThisNode = KDTree.Impl->pool.template allocate<FKDTreeImplementation::Node>();
+	ThisNode = KDTree.KDTreeImplementation->pool.template allocate<FKDTreeImplementation::Node>();
 	
 	ThisNode->node_type = OtherNode->node_type;
 
@@ -201,42 +201,42 @@ void CopySubTree(FKDTree& KDTree, FKDTreeImplementation::NodePtr& ThisNode, cons
 FKDTree::FKDTree(const FKDTree& Other)
 {
 #if UE_POSE_SEARCH_USE_NANOFLANN
-	if (this != &Other && Other.Impl)
+	if (this != &Other && Other.KDTreeImplementation)
 	{
-		Impl = new FKDTreeImplementation(0, DataSource, nanoflann::KDTreeSingleIndexAdaptorParams(0));
+		KDTreeImplementation = new FKDTreeImplementation(0, DataSource, nanoflann::KDTreeSingleIndexAdaptorParams(0));
 
 		DataSource = Other.DataSource;
 
-		check(Other.Impl->m_size <= AccessorTypeMax);
-		Impl->m_size = Other.Impl->m_size;
+		check(Other.KDTreeImplementation->m_size <= AccessorTypeMax);
+		KDTreeImplementation->m_size = Other.KDTreeImplementation->m_size;
 
-		if (Impl->m_size > 0)
+		if (KDTreeImplementation->m_size > 0)
 		{
-			Impl->dim = Other.Impl->dim;
+			KDTreeImplementation->dim = Other.KDTreeImplementation->dim;
 
-			check(Other.Impl->root_bbox.size() <= AccessorTypeMax);
-			const AccessorType root_bbox_size = Other.Impl->root_bbox.size();
-			Impl->root_bbox.resize(root_bbox_size);
+			check(Other.KDTreeImplementation->root_bbox.size() <= AccessorTypeMax);
+			const AccessorType root_bbox_size = Other.KDTreeImplementation->root_bbox.size();
+			KDTreeImplementation->root_bbox.resize(root_bbox_size);
 
 			for (AccessorType i = 0; i < root_bbox_size; ++i)
 			{
-				Impl->root_bbox[i] = Other.Impl->root_bbox[i];
+				KDTreeImplementation->root_bbox[i] = Other.KDTreeImplementation->root_bbox[i];
 			}
 
-			check(Other.Impl->m_leaf_max_size <= AccessorTypeMax);
-			const AccessorType KDTreeLeafMaxSize = Other.Impl->m_leaf_max_size;
-			Impl->m_leaf_max_size = KDTreeLeafMaxSize;
+			check(Other.KDTreeImplementation->m_leaf_max_size <= AccessorTypeMax);
+			const AccessorType KDTreeLeafMaxSize = Other.KDTreeImplementation->m_leaf_max_size;
+			KDTreeImplementation->m_leaf_max_size = KDTreeLeafMaxSize;
 
-			check(Other.Impl->vAcc.size() <= AccessorTypeMax);
-			const AccessorType VAccSize = Other.Impl->vAcc.size();
-			Impl->vAcc.resize(VAccSize);
+			check(Other.KDTreeImplementation->vAcc.size() <= AccessorTypeMax);
+			const AccessorType VAccSize = Other.KDTreeImplementation->vAcc.size();
+			KDTreeImplementation->vAcc.resize(VAccSize);
 			
 			for (AccessorType i = 0; i < VAccSize; ++i)
 			{
-				Impl->vAcc[i] = Other.Impl->vAcc[i];
+				KDTreeImplementation->vAcc[i] = Other.KDTreeImplementation->vAcc[i];
 			}
 			
-			CopySubTree(*this, Impl->root_node, Other.Impl->root_node);
+			CopySubTree(*this, KDTreeImplementation->root_node, Other.KDTreeImplementation->root_node);
 		}
 	}
 #endif // UE_POSE_SEARCH_USE_NANOFLANN
@@ -254,8 +254,8 @@ FKDTree& FKDTree::operator=(const FKDTree& Other)
 
 bool FKDTree::operator==(const FKDTree& Other) const
 {
-	const bool bAnyImpl = Impl != nullptr;
-	const bool bAnyOtherImpl = Other.Impl != nullptr;
+	const bool bAnyImpl = KDTreeImplementation != nullptr;
+	const bool bAnyOtherImpl = Other.KDTreeImplementation != nullptr;
 	if (bAnyImpl != bAnyOtherImpl)
 	{
 		return false;
@@ -264,7 +264,7 @@ bool FKDTree::operator==(const FKDTree& Other) const
 #if UE_POSE_SEARCH_USE_NANOFLANN
 	if (bAnyImpl && bAnyOtherImpl)
 	{
-		if (*Impl != *Other.Impl)
+		if (*KDTreeImplementation != *Other.KDTreeImplementation)
 		{
 			return false;
 		}
@@ -282,8 +282,8 @@ bool FKDTree::operator==(const FKDTree& Other) const
 void FKDTree::Reset()
 {
 #if UE_POSE_SEARCH_USE_NANOFLANN
-	delete Impl;
-	Impl = nullptr;
+	delete KDTreeImplementation;
+	KDTreeImplementation = nullptr;
 #endif // UE_POSE_SEARCH_USE_NANOFLANN
 	DataSource = FDataSource();
 }
@@ -295,19 +295,19 @@ void FKDTree::Construct(AccessorType Count, AccessorType Dim, const float* Data,
 }
 
 template <typename RESULTSET>
-inline int32 FindNeighborsInternal(FKDTreeImplementation* Impl, RESULTSET& Result, TConstArrayView<float> Query)
+inline int32 FindNeighborsInternal(FKDTreeImplementation* KDTreeImplementation, RESULTSET& Result, TConstArrayView<float> Query)
 {
 #if UE_POSE_SEARCH_USE_NANOFLANN
 
 	QUICK_SCOPE_CYCLE_COUNTER(STAT_FKDTree_FindNeighbors);
 
-	check(Query.GetData() && Query.Num() == Impl->dim && Impl->root_node);
+	check(Query.GetData() && Query.Num() == KDTreeImplementation->dim && KDTreeImplementation->root_node);
 
 	const nanoflann::SearchParams SearchParams(
 		32,			// Ignored parameter (Kept for compatibility with the FLANN interface).
 		0.f,		// search for eps-approximate neighbours (default: 0)
 		false);		// only for radius search, require neighbours sorted by
-	Impl->findNeighbors(Result, Query.GetData(), SearchParams);
+	KDTreeImplementation->findNeighbors(Result, Query.GetData(), SearchParams);
 	return Result.Num();
 
 #else // UE_POSE_SEARCH_USE_NANOFLANN
@@ -320,17 +320,17 @@ inline int32 FindNeighborsInternal(FKDTreeImplementation* Impl, RESULTSET& Resul
 
 int32 FKDTree::FindNeighbors(FKNNResultSet& Result, TConstArrayView<float> Query) const
 {
-	return FindNeighborsInternal(Impl, Result, Query);
+	return FindNeighborsInternal(KDTreeImplementation, Result, Query);
 }
 
 int32 FKDTree::FindNeighbors(FFilteredKNNResultSet& Result, TConstArrayView<float> Query) const
 {
-	return FindNeighborsInternal(Impl, Result, Query);
+	return FindNeighborsInternal(KDTreeImplementation, Result, Query);
 }
 
 int32 FKDTree::FindNeighbors(FRadiusResultSet& Result, TConstArrayView<float> Query) const
 {
-	return FindNeighborsInternal(Impl, Result, Query);
+	return FindNeighborsInternal(KDTreeImplementation, Result, Query);
 }
 
 SIZE_T FKDTree::GetAllocatedSize() const
@@ -338,10 +338,10 @@ SIZE_T FKDTree::GetAllocatedSize() const
 	SIZE_T AllocatedSize = sizeof(FKDTree);
 
 #if UE_POSE_SEARCH_USE_NANOFLANN
-	if (Impl)
+	if (KDTreeImplementation)
 	{
 		AllocatedSize += sizeof(FKDTreeImplementation);
-		AllocatedSize += Impl->usedMemory(*Impl);
+		AllocatedSize += KDTreeImplementation->usedMemory(*KDTreeImplementation);
 	}
 #endif // UE_POSE_SEARCH_USE_NANOFLANN
 
@@ -351,11 +351,11 @@ SIZE_T FKDTree::GetAllocatedSize() const
 #if UE_POSE_SEARCH_USE_NANOFLANN
 FArchive& SerializeSubTree(FArchive& Ar, FKDTree& KDTree, FKDTreeImplementation::NodePtr& KDTreeNode)
 {
-	check(KDTree.Impl);
+	check(KDTree.KDTreeImplementation);
 
 	if (Ar.IsLoading())
 	{
-		KDTreeNode = KDTree.Impl->pool.template allocate<FKDTreeImplementation::Node>();
+		KDTreeNode = KDTree.KDTreeImplementation->pool.template allocate<FKDTreeImplementation::Node>();
 		// zeroing FKDTreeImplementation::Node memory since it contains a union and doesn't have a constructor 
 		FMemory::Memzero(KDTreeNode, sizeof(FKDTreeImplementation::Node));
 	}
@@ -412,58 +412,58 @@ FArchive& SerializeSubTree(FArchive& Ar, FKDTree& KDTree, FKDTreeImplementation:
 FArchive& Serialize(FArchive& Ar, FKDTree& KDTree, const float* KDTreeData)
 {
 #if UE_POSE_SEARCH_USE_NANOFLANN
-	check(!KDTree.Impl || KDTree.Impl->m_size <= AccessorTypeMax);
+	check(!KDTree.KDTreeImplementation || KDTree.KDTreeImplementation->m_size <= AccessorTypeMax);
 
-	AccessorType KDTreeSize = KDTree.Impl ? KDTree.Impl->m_size : 0;
+	AccessorType KDTreeSize = KDTree.KDTreeImplementation ? KDTree.KDTreeImplementation->m_size : 0;
 
 	Ar << KDTreeSize;
 
 	if (KDTreeSize > 0)
 	{
-		if (Ar.IsLoading() && !KDTree.Impl)
+		if (Ar.IsLoading() && !KDTree.KDTreeImplementation)
 		{
-			KDTree.Impl = new FKDTreeImplementation(0, KDTree.DataSource, nanoflann::KDTreeSingleIndexAdaptorParams(0));
+			KDTree.KDTreeImplementation = new FKDTreeImplementation(0, KDTree.DataSource, nanoflann::KDTreeSingleIndexAdaptorParams(0));
 		}
 
-		KDTree.Impl->m_size = KDTreeSize;
+		KDTree.KDTreeImplementation->m_size = KDTreeSize;
 
-		Ar << KDTree.Impl->dim;
+		Ar << KDTree.KDTreeImplementation->dim;
 
-		AccessorType root_bbox_size = KDTree.Impl->root_bbox.size();
-		check(KDTree.Impl->root_bbox.size() <= AccessorTypeMax);
+		AccessorType root_bbox_size = KDTree.KDTreeImplementation->root_bbox.size();
+		check(KDTree.KDTreeImplementation->root_bbox.size() <= AccessorTypeMax);
 		Ar << root_bbox_size;
 
 		if (Ar.IsLoading())
 		{
 			KDTree.DataSource.Data = KDTreeData;
-			KDTree.DataSource.PointDim = KDTree.Impl->dim;
-			KDTree.DataSource.PointCount = KDTree.Impl->m_size;
+			KDTree.DataSource.PointDim = KDTree.KDTreeImplementation->dim;
+			KDTree.DataSource.PointCount = KDTree.KDTreeImplementation->m_size;
 
-			KDTree.Impl->root_bbox.resize(root_bbox_size);
+			KDTree.KDTreeImplementation->root_bbox.resize(root_bbox_size);
 		}
 
-		for (FKDTreeImplementation::Interval& el : KDTree.Impl->root_bbox)
+		for (FKDTreeImplementation::Interval& el : KDTree.KDTreeImplementation->root_bbox)
 		{
 			Ar.Serialize(&el, sizeof(FKDTreeImplementation::Interval));
 		}
 
-		check(KDTree.Impl->m_leaf_max_size <= AccessorTypeMax);
-		AccessorType KDTreeLeafMaxSize = KDTree.Impl->m_leaf_max_size;
+		check(KDTree.KDTreeImplementation->m_leaf_max_size <= AccessorTypeMax);
+		AccessorType KDTreeLeafMaxSize = KDTree.KDTreeImplementation->m_leaf_max_size;
 		Ar << KDTreeLeafMaxSize;
-		KDTree.Impl->m_leaf_max_size = KDTreeLeafMaxSize;
+		KDTree.KDTreeImplementation->m_leaf_max_size = KDTreeLeafMaxSize;
 
-		check(KDTree.Impl->vAcc.size() <= AccessorTypeMax);
-		AccessorType VAccSize = KDTree.Impl->vAcc.size();
+		check(KDTree.KDTreeImplementation->vAcc.size() <= AccessorTypeMax);
+		AccessorType VAccSize = KDTree.KDTreeImplementation->vAcc.size();
 		Ar << VAccSize;
 		if (Ar.IsLoading())
 		{
-			KDTree.Impl->vAcc.resize(VAccSize);
+			KDTree.KDTreeImplementation->vAcc.resize(VAccSize);
 		}
-		for (AccessorType& el : KDTree.Impl->vAcc)
+		for (AccessorType& el : KDTree.KDTreeImplementation->vAcc)
 		{
 			Ar << el;
 		}
-		SerializeSubTree(Ar, KDTree, KDTree.Impl->root_node);
+		SerializeSubTree(Ar, KDTree, KDTree.KDTreeImplementation->root_node);
 	}
 	else if (Ar.IsLoading())
 	{
diff --git a/Engine/Plugins/Experimental/Animation/PoseSearch/Source/Runtime/Public/PoseSearch/KDTree.h b/Engine/Plugins/Experimental/Animation/PoseSearch/Source/Runtime/Public/PoseSearch/KDTree.h
--- a/Engine/Plugins/Experimental/Animation/PoseSearch/Source/Runtime/Public/PoseSearch/KDTree.h
+++ b/Engine/Plugins/Experimental/Animation/PoseSearch/Source/Runtime/Public/PoseSearch/KDTree.h
@@ -201,7 +201,7 @@ struct FKDTree
 	POSESEARCH_API SIZE_T GetAllocatedSize() const;
 
 	FDataSource DataSource;
-	FKDTreeImplementation* Impl = nullptr;
+	FKDTreeImplementation* KDTreeImplementation = nullptr;
 };
 
 FArchive& Serialize(FArchive& Ar, FKDTree& KDTree, const float* data);
