motion matching - schema channel caching

diff --git a/Engine/Plugins/Experimental/Animation/PoseSearch/Source/Runtime/Private/AnimNode_MotionMatching.cpp b/Engine/Plugins/Experimental/Animation/PoseSearch/Source/Runtime/Private/AnimNode_MotionMatching.cpp
--- a/Engine/Plugins/Experimental/Animation/PoseSearch/Source/Runtime/Private/AnimNode_MotionMatching.cpp
+++ b/Engine/Plugins/Experimental/Animation/PoseSearch/Source/Runtime/Private/AnimNode_MotionMatching.cpp
@@ -181,7 +181,8 @@ void FAnimNode_MotionMatching::UpdateAssetPlayer(const FAnimationUpdateContext&
 		YawFromAnimationBlendRate,
 		YawFromAnimationTrajectoryBlendTime,
 		NextUpdateInterruptMode,
-		bShouldSearch
+		bShouldSearch,
+		bShouldUseCachedChannelData
 		#if ENABLE_ANIM_DEBUG
 		, CVarAnimNodeMotionMatchingDrawQuery.GetValueOnAnyThread()
 		, CVarAnimNodeMotionMatchingDrawCurResult.GetValueOnAnyThread()
diff --git a/Engine/Plugins/Experimental/Animation/PoseSearch/Source/Runtime/Private/PoseSearchContext.cpp b/Engine/Plugins/Experimental/Animation/PoseSearch/Source/Runtime/Private/PoseSearchContext.cpp
--- a/Engine/Plugins/Experimental/Animation/PoseSearch/Source/Runtime/Private/PoseSearchContext.cpp
+++ b/Engine/Plugins/Experimental/Animation/PoseSearch/Source/Runtime/Private/PoseSearchContext.cpp
@@ -349,11 +349,20 @@ void FDebugDrawParams::DrawFeatureVector(int32 PoseIdx)
 }
 #endif // ENABLE_DRAW_DEBUG
 
+//////////////////////////////////////////////////////////////////////////
+// FCachedQuery
+FCachedQuery::FCachedQuery(const UPoseSearchSchema* InSchema)
+{
+	check(InSchema && InSchema->IsValid());
+	Schema = InSchema;
+	Values.SetNumZeroed(Schema->SchemaCardinality);
+}
+
 //////////////////////////////////////////////////////////////////////////
 // FSearchContext
 FSearchContext::FSearchContext(const UAnimInstance* InAnimInstance, const IPoseHistory* InHistory, TConstArrayView<const UAnimationAsset*> InAnimationsToConsider,
 		const FPoseSearchQueryTrajectory* InTrajectory, float InDesiredPermutationTimeOffset, const FPoseIndicesHistory* InPoseIndicesHistory,
-		const FSearchResult& InCurrentResult, const FFloatInterval& InPoseJumpThresholdTime)
+		const FSearchResult& InCurrentResult, const FFloatInterval& InPoseJumpThresholdTime, bool bInUseCachedChannelData)
 : AnimInstance(InAnimInstance)
 , History(InHistory)
 , AnimationsToConsider(InAnimationsToConsider)
@@ -362,6 +371,7 @@ FSearchContext::FSearchContext(const UAnimInstance* InAnimInstance, const IPoseH
 , PoseIndicesHistory(InPoseIndicesHistory)
 , CurrentResult(InCurrentResult)
 , PoseJumpThresholdTime(InPoseJumpThresholdTime)
+, bUseCachedChannelData(bInUseCachedChannelData)
 {
 	check(AnimInstance);
 	if (CurrentResult.IsValid())
@@ -380,7 +390,7 @@ FSearchContext::FSearchContext(const UAnimInstance* InAnimInstance, const IPoseH
 	}
 }
 
-FQuat FSearchContext::GetSampleRotation(float SampleTimeOffset, float OriginTimeOffset, const UPoseSearchSchema* Schema, int8 SchemaSampleBoneIdx, int8 SchemaOriginBoneIdx, EPermutationTimeType PermutationTimeType, const FQuat* SampleBoneRotationWorldOverride)
+FQuat FSearchContext::GetSampleRotation(float SampleTimeOffset, float OriginTimeOffset, int8 SchemaSampleBoneIdx, int8 SchemaOriginBoneIdx, EPermutationTimeType PermutationTimeType, const FQuat* SampleBoneRotationWorldOverride)
 {
 	float PermutationSampleTimeOffset = 0.f;
 	float PermutationOriginTimeOffset = 0.f;
@@ -389,10 +399,10 @@ FQuat FSearchContext::GetSampleRotation(float SampleTimeOffset, float OriginTime
 	const float SampleTime = SampleTimeOffset + PermutationSampleTimeOffset;
 	const float OriginTime = OriginTimeOffset + PermutationOriginTimeOffset;
 
-	return GetSampleRotationInternal(SampleTime, OriginTime, Schema, SchemaSampleBoneIdx, SchemaOriginBoneIdx, SampleBoneRotationWorldOverride);
+	return GetSampleRotationInternal(SampleTime, OriginTime, SchemaSampleBoneIdx, SchemaOriginBoneIdx, SampleBoneRotationWorldOverride);
 }
 
-FVector FSearchContext::GetSamplePosition(float SampleTimeOffset, float OriginTimeOffset, const UPoseSearchSchema* Schema, int8 SchemaSampleBoneIdx, int8 SchemaOriginBoneIdx, EPermutationTimeType PermutationTimeType, const FVector* SampleBonePositionWorldOverride)
+FVector FSearchContext::GetSamplePosition(float SampleTimeOffset, float OriginTimeOffset, int8 SchemaSampleBoneIdx, int8 SchemaOriginBoneIdx, EPermutationTimeType PermutationTimeType, const FVector* SampleBonePositionWorldOverride)
 {
 	float PermutationSampleTimeOffset = 0.f;
 	float PermutationOriginTimeOffset = 0.f;
@@ -400,10 +410,10 @@ FVector FSearchContext::GetSamplePosition(float SampleTimeOffset, float OriginTi
 
 	const float SampleTime = SampleTimeOffset + PermutationSampleTimeOffset;
 	const float OriginTime = OriginTimeOffset + PermutationOriginTimeOffset;
-	return GetSamplePositionInternal(SampleTime, OriginTime, Schema, SchemaSampleBoneIdx, SchemaOriginBoneIdx, SampleBonePositionWorldOverride);
+	return GetSamplePositionInternal(SampleTime, OriginTime, SchemaSampleBoneIdx, SchemaOriginBoneIdx, SampleBonePositionWorldOverride);
 }
 
-FVector FSearchContext::GetSampleVelocity(float SampleTimeOffset, float OriginTimeOffset, const UPoseSearchSchema* Schema, int8 SchemaSampleBoneIdx, int8 SchemaOriginBoneIdx, bool bUseCharacterSpaceVelocities, EPermutationTimeType PermutationTimeType, const FVector* SampleBoneVelocityWorldOverride)
+FVector FSearchContext::GetSampleVelocity(float SampleTimeOffset, float OriginTimeOffset, int8 SchemaSampleBoneIdx, int8 SchemaOriginBoneIdx, bool bUseCharacterSpaceVelocities, EPermutationTimeType PermutationTimeType, const FVector* SampleBoneVelocityWorldOverride)
 {
 	using namespace UE::PoseSearch;
 
@@ -416,13 +426,13 @@ FVector FSearchContext::GetSampleVelocity(float SampleTimeOffset, float OriginTi
 
 	if (SampleBoneVelocityWorldOverride)
 	{
-		const FTransform RootBoneTransform = GetWorldBoneTransformAtTime(OriginTime, Schema, RootSchemaBoneIdx);
+		const FTransform RootBoneTransform = GetWorldBoneTransformAtTime(OriginTime, RootSchemaBoneIdx);
 		return RootBoneTransform.InverseTransformVector(*SampleBoneVelocityWorldOverride);
 	}
 
 	// calculating the local Position for the bone indexed by SchemaSampleBoneIdx
-	const FVector PreviousTranslation = GetSamplePositionInternal(SampleTime - FiniteDelta, bUseCharacterSpaceVelocities ? OriginTime - FiniteDelta : OriginTime, Schema, SchemaSampleBoneIdx, SchemaOriginBoneIdx);
-	const FVector CurrentTranslation = GetSamplePositionInternal(SampleTime, OriginTime, Schema, SchemaSampleBoneIdx, SchemaOriginBoneIdx);
+	const FVector PreviousTranslation = GetSamplePositionInternal(SampleTime - FiniteDelta, bUseCharacterSpaceVelocities ? OriginTime - FiniteDelta : OriginTime, SchemaSampleBoneIdx, SchemaOriginBoneIdx);
+	const FVector CurrentTranslation = GetSamplePositionInternal(SampleTime, OriginTime, SchemaSampleBoneIdx, SchemaOriginBoneIdx);
 
 	const FVector LinearVelocity = (CurrentTranslation - PreviousTranslation) / FiniteDelta;
 	return LinearVelocity;
@@ -448,8 +458,10 @@ FTransform FSearchContext::GetWorldRootBoneTransformAtTime(float SampleTime) con
 	return FTransform::Identity;
 }
 
-FTransform FSearchContext::GetWorldBoneTransformAtTime(float SampleTime, const UPoseSearchSchema* Schema, int8 SchemaBoneIdx)
+FTransform FSearchContext::GetWorldBoneTransformAtTime(float SampleTime, int8 SchemaBoneIdx)
 {
+	// CachedQueries.Last is the query we're building 
+	const UPoseSearchSchema* Schema = !CachedQueries.IsEmpty() ? CachedQueries.Last().GetSchema() : nullptr;
 	const FBoneIndexType BoneIndexType = !Schema || SchemaBoneIdx == RootSchemaBoneIdx ? RootBoneIndexType : Schema->GetBoneIndexType(SchemaBoneIdx);
 	if (const FCachedTransform<FTransform>* CachedTransform = CachedTransforms.Find(SampleTime, BoneIndexType))
 	{
@@ -511,57 +523,80 @@ FTransform FSearchContext::GetWorldBoneTransformAtTime(float SampleTime, const U
 	return WorldBoneTransform;
 }
 
-FVector FSearchContext::GetSamplePositionInternal(float SampleTime, float OriginTime, const UPoseSearchSchema* Schema, int8 SchemaSampleBoneIdx, int8 SchemaOriginBoneIdx, const FVector* SampleBonePositionWorldOverride)
+FVector FSearchContext::GetSamplePositionInternal(float SampleTime, float OriginTime, int8 SchemaSampleBoneIdx, int8 SchemaOriginBoneIdx, const FVector* SampleBonePositionWorldOverride)
 {
 	if (SampleBonePositionWorldOverride)
 	{
-		const FTransform RootBoneTransform = GetWorldBoneTransformAtTime(OriginTime, Schema, RootSchemaBoneIdx);
+		const FTransform RootBoneTransform = GetWorldBoneTransformAtTime(OriginTime, RootSchemaBoneIdx);
 		if (SchemaOriginBoneIdx == RootSchemaBoneIdx)
 		{
 			return RootBoneTransform.InverseTransformPosition(*SampleBonePositionWorldOverride);
 		}
 
 		// @todo: validate this still works for when root bone is not Identity
-		const FTransform OriginBoneTransform = GetWorldBoneTransformAtTime(OriginTime, Schema, SchemaOriginBoneIdx);
+		const FTransform OriginBoneTransform = GetWorldBoneTransformAtTime(OriginTime, SchemaOriginBoneIdx);
 		const FVector DeltaBoneTranslation = *SampleBonePositionWorldOverride - OriginBoneTransform.GetTranslation();
 		return RootBoneTransform.InverseTransformVector(DeltaBoneTranslation);
 	}
 
-	const FTransform RootBoneTransform = GetWorldBoneTransformAtTime(OriginTime, Schema, RootSchemaBoneIdx);
-	const FTransform SampleBoneTransform = GetWorldBoneTransformAtTime(SampleTime, Schema, SchemaSampleBoneIdx);
+	const FTransform RootBoneTransform = GetWorldBoneTransformAtTime(OriginTime, RootSchemaBoneIdx);
+	const FTransform SampleBoneTransform = GetWorldBoneTransformAtTime(SampleTime, SchemaSampleBoneIdx);
 	if (SchemaOriginBoneIdx == RootSchemaBoneIdx)
 	{
 		return RootBoneTransform.InverseTransformPosition(SampleBoneTransform.GetTranslation());
 	}
 
-	const FTransform OriginBoneTransform = GetWorldBoneTransformAtTime(OriginTime, Schema, SchemaOriginBoneIdx);
+	const FTransform OriginBoneTransform = GetWorldBoneTransformAtTime(OriginTime, SchemaOriginBoneIdx);
 	const FVector DeltaBoneTranslation = SampleBoneTransform.GetTranslation() - OriginBoneTransform.GetTranslation();
 	return RootBoneTransform.InverseTransformVector(DeltaBoneTranslation);
 }
 
-FQuat FSearchContext::GetSampleRotationInternal(float SampleTime, float OriginTime, const UPoseSearchSchema* Schema, int8 SchemaSampleBoneIdx, int8 SchemaOriginBoneIdx, const FQuat* SampleBoneRotationWorldOverride)
+FQuat FSearchContext::GetSampleRotationInternal(float SampleTime, float OriginTime, int8 SchemaSampleBoneIdx, int8 SchemaOriginBoneIdx, const FQuat* SampleBoneRotationWorldOverride)
 {
 	if (SampleBoneRotationWorldOverride)
 	{
-		const FTransform RootBoneTransform = GetWorldBoneTransformAtTime(OriginTime, Schema, RootSchemaBoneIdx);
+		const FTransform RootBoneTransform = GetWorldBoneTransformAtTime(OriginTime, RootSchemaBoneIdx);
 		if (SchemaOriginBoneIdx == RootSchemaBoneIdx)
 		{
 			return RootBoneTransform.InverseTransformRotation(*SampleBoneRotationWorldOverride);
 		}
 
-		const FTransform OriginBoneTransform = GetWorldBoneTransformAtTime(OriginTime, Schema, SchemaOriginBoneIdx);
+		const FTransform OriginBoneTransform = GetWorldBoneTransformAtTime(OriginTime, SchemaOriginBoneIdx);
 		const FQuat DeltaBoneRotation = OriginBoneTransform.InverseTransformRotation(*SampleBoneRotationWorldOverride);
 		return RootBoneTransform.InverseTransformRotation(DeltaBoneRotation);
 	}
 
-	const FTransform RootBoneTransform = GetWorldBoneTransformAtTime(OriginTime, Schema, RootSchemaBoneIdx);
-	const FTransform SampleBoneTransform = GetWorldBoneTransformAtTime(SampleTime, Schema, SchemaSampleBoneIdx);
+	const FTransform RootBoneTransform = GetWorldBoneTransformAtTime(OriginTime, RootSchemaBoneIdx);
+	const FTransform SampleBoneTransform = GetWorldBoneTransformAtTime(SampleTime, SchemaSampleBoneIdx);
 	return RootBoneTransform.InverseTransformRotation(SampleBoneTransform.GetRotation());
 }
 
-void FSearchContext::ClearCachedEntries()
+TArrayView<float> FSearchContext::EditFeatureVector()
 {
-	CachedTransforms.Reset();
+	// CachedQueries.Last is the query we're building 
+	check(!CachedQueries.IsEmpty());
+	return CachedQueries.Last().EditValues();
+}
+
+const UPoseSearchFeatureChannel* FSearchContext::GetCachedChannelData(uint32 ChannelUniqueIdentifier, const UPoseSearchFeatureChannel* Channel, TConstArrayView<float>& CachedChannelData)
+{
+	// searching CachedChannels for the ChannelUniqueIdentifier as representation of Channel
+	FCachedChannel& CachedChannel = CachedChannels.FindOrAdd(ChannelUniqueIdentifier);
+	if (CachedChannel.Channel)
+	{
+		// we found CachedChannel.Channel, a channel from a different schema (CachedQueries[CachedChannel.CachedQueryIndex].GetSchema()) compatible with Channel.
+		// let's collect the associated data to CachedChannel.Channel 
+		CachedChannelData = CachedQueries[CachedChannel.CachedQueryIndex].GetValues().Slice(CachedChannel.Channel->GetChannelDataOffset(), CachedChannel.Channel->GetChannelCardinality());
+		return CachedChannel.Channel;
+	}
+	
+	// we couldn't find the cached channel, so let's add the pair ChannelUniqueIdentifier / Channel to CachedChannels.
+	// the associated CachedQueries[CachedQueries.Num() - 1].GetValues() data will be filled up by the end of Channel BuildQuery
+	CachedChannel.CachedQueryIndex = CachedQueries.Num() - 1;
+	CachedChannel.Channel = Channel;
+	
+	CachedChannelData = TConstArrayView<float>();
+	return nullptr;
 }
 
 void FSearchContext::ResetCurrentBestCost()
@@ -577,27 +612,34 @@ void FSearchContext::UpdateCurrentBestCost(const FPoseSearchCost& PoseSearchCost
 	};
 }
 
-const FFeatureVectorBuilder* FSearchContext::GetCachedQuery(const UPoseSearchSchema* Schema) const
+TConstArrayView<float> FSearchContext::GetCachedQuery(const UPoseSearchSchema* Schema) const
 {
-	return CachedQueries.FindByPredicate([Schema](const FFeatureVectorBuilder& CachedQuery)
+	if (const FCachedQuery* FoundCachedQuery = CachedQueries.FindByPredicate(
+		[Schema](const FCachedQuery& CachedQuery)
+		{
+			return CachedQuery.GetSchema() == Schema;
+		}))
 	{
-		return CachedQuery.GetSchema() == Schema;
-	});
+		return FoundCachedQuery->GetValues();
+	}
+	return TConstArrayView<float>();
 }
 
-const FFeatureVectorBuilder& FSearchContext::GetOrBuildQuery(const UPoseSearchSchema* Schema)
+TConstArrayView<float> FSearchContext::GetOrBuildQuery(const UPoseSearchSchema* Schema)
 {
 	QUICK_SCOPE_CYCLE_COUNTER(STAT_PoseSearch_GetOrBuildQuery);
 
 	check(Schema && Schema->IsValid());
-	if (const FFeatureVectorBuilder* CachedFeatureVectorBuilder = GetCachedQuery(Schema))
+	if (const FCachedQuery* FoundCachedQuery = CachedQueries.FindByPredicate(
+		[Schema](const FCachedQuery& CachedQuery)
+		{
+			return CachedQuery.GetSchema() == Schema;
+		}))
 	{
-		return *CachedFeatureVectorBuilder;
+		return FoundCachedQuery->GetValues();
 	}
 	
-	FFeatureVectorBuilder& NewCachedQuery = CachedQueries[CachedQueries.Emplace(Schema)];
-	Schema->BuildQuery(*this, NewCachedQuery);
-	return NewCachedQuery;
+	return Schema->BuildQuery(*this);
 }
 
 bool FSearchContext::IsCurrentResultFromDatabase(const UPoseSearchDatabase* Database) const
@@ -605,4 +647,12 @@ bool FSearchContext::IsCurrentResultFromDatabase(const UPoseSearchDatabase* Data
 	return CurrentResult.IsValid() && CurrentResult.Database == Database;
 }
 
+bool FSearchContext::CanUseCurrentResult() const
+{
+	// CachedQueries.Last is the query we're building 
+	check(!CachedQueries.IsEmpty());
+	return CurrentResult.IsValid() && CurrentResult.Database->Schema == CachedQueries.Last().GetSchema();
+}
+
+
 } // namespace UE::PoseSearch
\ No newline at end of file
diff --git a/Engine/Plugins/Experimental/Animation/PoseSearch/Source/Runtime/Private/PoseSearchDatabase.cpp b/Engine/Plugins/Experimental/Animation/PoseSearch/Source/Runtime/Private/PoseSearchDatabase.cpp
--- a/Engine/Plugins/Experimental/Animation/PoseSearch/Source/Runtime/Private/PoseSearchDatabase.cpp
+++ b/Engine/Plugins/Experimental/Animation/PoseSearch/Source/Runtime/Private/PoseSearchDatabase.cpp
@@ -1000,12 +1000,12 @@ UE::PoseSearch::FSearchResult UPoseSearchDatabase::SearchContinuingPose(UE::Pose
 		// is the data padded at 16 bytes (and 16 bytes aligned by construction)?
 		if (NumDimensions % 4 == 0)
 		{
-			Result.PoseCost = SearchIndex.CompareAlignedPoses(ContinuingPoseIdx, UpdatedContinuingPoseCostBias, PoseValues, SearchContext.GetOrBuildQuery(Schema).GetValues());
+			Result.PoseCost = SearchIndex.CompareAlignedPoses(ContinuingPoseIdx, UpdatedContinuingPoseCostBias, PoseValues, SearchContext.GetOrBuildQuery(Schema));
 		}
 		// data is not 16 bytes padded
 		else
 		{
-			Result.PoseCost = SearchIndex.ComparePoses(ContinuingPoseIdx, UpdatedContinuingPoseCostBias, PoseValues, SearchContext.GetOrBuildQuery(Schema).GetValues());
+			Result.PoseCost = SearchIndex.ComparePoses(ContinuingPoseIdx, UpdatedContinuingPoseCostBias, PoseValues, SearchContext.GetOrBuildQuery(Schema));
 		}
 
 		Result.AssetTime = SearchContext.GetCurrentResult().AssetTime;
@@ -1044,7 +1044,7 @@ UE::PoseSearch::FSearchResult UPoseSearchDatabase::SearchPCAKDTree(UE::PoseSearc
 		TArrayView<float> ResultDistanceSqr((float*)FMemory_Alloca((ClampedKDTreeQueryNumNeighbors + 1) * sizeof(float)), ClampedKDTreeQueryNumNeighbors + 1);
 		TArrayView<float> ProjectedQueryValues((float*)FMemory_Alloca(ClampedNumberOfPrincipalComponents * sizeof(float)), ClampedNumberOfPrincipalComponents);
 	
-		TConstArrayView<float> QueryValues = SearchContext.GetOrBuildQuery(Schema).GetValues();
+		TConstArrayView<float> QueryValues = SearchContext.GetOrBuildQuery(Schema);
 
 		FSelectableAssetIdx SelectableAssetIdx;
 		PopulateSelectableAssetIdx(SelectableAssetIdx, SearchContext.GetAnimationsToConsider(), this);
@@ -1188,7 +1188,7 @@ UE::PoseSearch::FSearchResult UPoseSearchDatabase::SearchPCAKDTree(UE::PoseSearc
 	{
 #if UE_POSE_SEARCH_TRACE_ENABLED
 		// calling just for reporting non selectable poses
-		TConstArrayView<float> QueryValues = SearchContext.GetOrBuildQuery(Schema).GetValues();
+		TConstArrayView<float> QueryValues = SearchContext.GetOrBuildQuery(Schema);
 		FNonSelectableIdx NonSelectableIdx;
 		PopulateNonSelectableIdx(NonSelectableIdx, SearchContext, this, QueryValues);
 #endif // UE_POSE_SEARCH_TRACE_ENABLED
@@ -1218,7 +1218,7 @@ UE::PoseSearch::FSearchResult UPoseSearchDatabase::SearchVPTree(UE::PoseSearch::
 	// there's no point in performing the search if CurrentBestTotalCost is already better than that
 	if (!GetSkipSearchIfPossible() || SearchContext.GetCurrentBestTotalCost() > SearchIndex.MinCostAddend)
 	{
-		TConstArrayView<float> QueryValues = SearchContext.GetOrBuildQuery(Schema).GetValues();
+		TConstArrayView<float> QueryValues = SearchContext.GetOrBuildQuery(Schema);
 
 		FSelectableAssetIdx SelectableAssetIdx;
 		PopulateSelectableAssetIdx(SelectableAssetIdx, SearchContext.GetAnimationsToConsider(), this);
@@ -1275,7 +1275,7 @@ UE::PoseSearch::FSearchResult UPoseSearchDatabase::SearchVPTree(UE::PoseSearch::
 	{
 #if UE_POSE_SEARCH_TRACE_ENABLED
 		// calling just for reporting non selectable poses
-		TConstArrayView<float> QueryValues = SearchContext.GetOrBuildQuery(Schema).GetValues();
+		TConstArrayView<float> QueryValues = SearchContext.GetOrBuildQuery(Schema);
 		FNonSelectableIdx NonSelectableIdx;
 		PopulateNonSelectableIdx(NonSelectableIdx, SearchContext, this, QueryValues);
 #endif // UE_POSE_SEARCH_TRACE_ENABLED
@@ -1305,7 +1305,7 @@ UE::PoseSearch::FSearchResult UPoseSearchDatabase::SearchBruteForce(UE::PoseSear
 	// there's no point in performing the search if CurrentBestTotalCost is already better than that
 	if (!GetSkipSearchIfPossible() || SearchContext.GetCurrentBestTotalCost() > SearchIndex.MinCostAddend)
 	{
-		TConstArrayView<float> QueryValues = SearchContext.GetOrBuildQuery(Schema).GetValues();
+		TConstArrayView<float> QueryValues = SearchContext.GetOrBuildQuery(Schema);
 
 		FSelectableAssetIdx SelectableAssetIdx;
 		PopulateSelectableAssetIdx(SelectableAssetIdx, SearchContext.GetAnimationsToConsider(), this);
@@ -1353,7 +1353,7 @@ UE::PoseSearch::FSearchResult UPoseSearchDatabase::SearchBruteForce(UE::PoseSear
 	{
 #if UE_POSE_SEARCH_TRACE_ENABLED
 		// calling just for reporting non selectable poses
-		TConstArrayView<float> QueryValues = SearchContext.GetOrBuildQuery(Schema).GetValues();
+		TConstArrayView<float> QueryValues = SearchContext.GetOrBuildQuery(Schema);
 		FNonSelectableIdx NonSelectableIdx;
 		PopulateNonSelectableIdx(NonSelectableIdx, SearchContext, this, QueryValues);
 #endif // UE_POSE_SEARCH_TRACE_ENABLED
diff --git a/Engine/Plugins/Experimental/Animation/PoseSearch/Source/Runtime/Private/PoseSearchFeatureChannel_FilterCrashingLegs.cpp b/Engine/Plugins/Experimental/Animation/PoseSearch/Source/Runtime/Private/PoseSearchFeatureChannel_FilterCrashingLegs.cpp
--- a/Engine/Plugins/Experimental/Animation/PoseSearch/Source/Runtime/Private/PoseSearchFeatureChannel_FilterCrashingLegs.cpp
+++ b/Engine/Plugins/Experimental/Animation/PoseSearch/Source/Runtime/Private/PoseSearchFeatureChannel_FilterCrashingLegs.cpp
@@ -49,36 +49,70 @@ void UPoseSearchFeatureChannel_FilterCrashingLegs::AddDependentChannels(UPoseSea
 		UPoseSearchFeatureChannel_Position::FindOrAddToSchema(Schema, 0.f, RightFoot.BoneName);
 	}
 }
-void UPoseSearchFeatureChannel_FilterCrashingLegs::BuildQuery(UE::PoseSearch::FSearchContext& SearchContext, UE::PoseSearch::FFeatureVectorBuilder& InOutQuery) const
+
+void UPoseSearchFeatureChannel_FilterCrashingLegs::BuildQuery(UE::PoseSearch::FSearchContext& SearchContext) const
 {
 	using namespace UE::PoseSearch;
 
-	const bool bIsCurrentResultValid = SearchContext.GetCurrentResult().IsValid() && SearchContext.GetCurrentResult().Database->Schema == InOutQuery.GetSchema();
-	const bool bSkip = InputQueryPose != EInputQueryPose::UseCharacterPose && bIsCurrentResultValid;
-	
-	if (bSkip || !SearchContext.IsHistoryValid())
+	// trying to get the BuildQuery data from another schema UPoseSearchFeatureChannel_FilterCrashingLegs already cached in the SearchContext
+	if (SearchContext.IsUseCachedChannelData())
 	{
-		if (bIsCurrentResultValid)
+		// composing a unique identifier to specify this channel with all the required properties to be able to share the query data with other channels of the same type
+		uint32 UniqueIdentifier = GetClass()->GetUniqueID();
+		UniqueIdentifier = HashCombineFast(UniqueIdentifier, ::GetTypeHash(LeftThighIdx));
+		UniqueIdentifier = HashCombineFast(UniqueIdentifier, ::GetTypeHash(RightThighIdx));
+		UniqueIdentifier = HashCombineFast(UniqueIdentifier, ::GetTypeHash(LeftFootIdx));
+		UniqueIdentifier = HashCombineFast(UniqueIdentifier, ::GetTypeHash(RightFootIdx));
+		UniqueIdentifier = HashCombineFast(UniqueIdentifier, ::GetTypeHash(InputQueryPose));
+
+		TConstArrayView<float> CachedChannelData;
+		if (const UPoseSearchFeatureChannel* CachedChannel = SearchContext.GetCachedChannelData(UniqueIdentifier, this, CachedChannelData))
 		{
-			FFeatureVectorHelper::Copy(InOutQuery.EditValues(), ChannelDataOffset, ChannelCardinality, SearchContext.GetCurrentResultPoseVector());
+#if DO_CHECK
+			const UPoseSearchFeatureChannel_FilterCrashingLegs* CachedFilterCrashingLegs = Cast<UPoseSearchFeatureChannel_FilterCrashingLegs>(CachedChannel);
+			check(CachedFilterCrashingLegs);
+			check(CachedFilterCrashingLegs->GetChannelCardinality() == ChannelCardinality);
+			check(CachedChannelData.Num() == ChannelCardinality);
+
+			// making sure there were no hash collisions
+			check(CachedFilterCrashingLegs->LeftThighIdx == LeftThighIdx);
+			check(CachedFilterCrashingLegs->RightThighIdx == RightThighIdx);
+			check(CachedFilterCrashingLegs->LeftFootIdx == LeftFootIdx);
+			check(CachedFilterCrashingLegs->RightFootIdx == RightFootIdx);
+			check(CachedFilterCrashingLegs->InputQueryPose == InputQueryPose);
+#endif //DO_CHECK
+
+			// copying the CachedChannelData into this channel portion of the FeatureVectorBuilder
+			FFeatureVectorHelper::Copy(SearchContext.EditFeatureVector().Slice(ChannelDataOffset, ChannelCardinality), 0, ChannelCardinality, CachedChannelData);
+			return;
 		}
-		else
+	}
+
+	// trying to get the BuildQuery data from the continuing pose
+	const bool bCanUseCurrentResult = SearchContext.CanUseCurrentResult();
+	const bool bSkip = InputQueryPose != EInputQueryPose::UseCharacterPose && bCanUseCurrentResult;
+	if (bSkip || !SearchContext.IsHistoryValid())
+	{
+		if (bCanUseCurrentResult)
 		{
-			// we leave the InOutQuery set to zero since the SearchContext.History is invalid and it'll fail if we continue
-			UE_LOG(LogPoseSearch, Error, TEXT("UPoseSearchFeatureChannel_FilterCrashingLegs::BuildQuery - Failed because Pose History Node is missing."));
+			FFeatureVectorHelper::Copy(SearchContext.EditFeatureVector(), ChannelDataOffset, ChannelCardinality, SearchContext.GetCurrentResultPoseVector());
+			return;
 		}
+
+		// we leave the SearchContext.EditFeatureVector() set to zero since the SearchContext.History is invalid and it'll fail if we continue
+		UE_LOG(LogPoseSearch, Error, TEXT("UPoseSearchFeatureChannel_FilterCrashingLegs::BuildQuery - Failed because Pose History Node is missing."));
+		return;
 	}
-	else
-	{
-		const FVector RightThighPosition = SearchContext.GetSamplePosition(0.f, 0.f, InOutQuery.GetSchema(), RightThighIdx);
-		const FVector LeftThighPosition = SearchContext.GetSamplePosition(0.f, 0.f, InOutQuery.GetSchema(), LeftThighIdx);
-		const FVector RightFootPosition = SearchContext.GetSamplePosition(0.f, 0.f, InOutQuery.GetSchema(), RightFootIdx);
-		const FVector LeftFootPosition = SearchContext.GetSamplePosition(0.f, 0.f, InOutQuery.GetSchema(), LeftFootIdx);
 
-		const float CrashingLegsValue = ComputeCrashingLegsValue(RightThighPosition, LeftThighPosition, RightFootPosition, LeftFootPosition);
+	// composing the BuildQuery data from the bones
+	const FVector RightThighPosition = SearchContext.GetSamplePosition(0.f, 0.f, RightThighIdx);
+	const FVector LeftThighPosition = SearchContext.GetSamplePosition(0.f, 0.f, LeftThighIdx);
+	const FVector RightFootPosition = SearchContext.GetSamplePosition(0.f, 0.f, RightFootIdx);
+	const FVector LeftFootPosition = SearchContext.GetSamplePosition(0.f, 0.f, LeftFootIdx);
 
-		FFeatureVectorHelper::EncodeFloat(InOutQuery.EditValues(), ChannelDataOffset, CrashingLegsValue);
-	}
+	const float CrashingLegsValue = ComputeCrashingLegsValue(RightThighPosition, LeftThighPosition, RightFootPosition, LeftFootPosition);
+
+	FFeatureVectorHelper::EncodeFloat(SearchContext.EditFeatureVector(), ChannelDataOffset, CrashingLegsValue);
 }
 
 #if ENABLE_DRAW_DEBUG
diff --git a/Engine/Plugins/Experimental/Animation/PoseSearch/Source/Runtime/Private/PoseSearchFeatureChannel_FilterCrashingLegs.h b/Engine/Plugins/Experimental/Animation/PoseSearch/Source/Runtime/Private/PoseSearchFeatureChannel_FilterCrashingLegs.h
--- a/Engine/Plugins/Experimental/Animation/PoseSearch/Source/Runtime/Private/PoseSearchFeatureChannel_FilterCrashingLegs.h
+++ b/Engine/Plugins/Experimental/Animation/PoseSearch/Source/Runtime/Private/PoseSearchFeatureChannel_FilterCrashingLegs.h
@@ -55,7 +55,7 @@ public:
 
 	// UPoseSearchFeatureChannel interface
 	virtual void Finalize(UPoseSearchSchema* Schema) override;
-	virtual void BuildQuery(UE::PoseSearch::FSearchContext& SearchContext, UE::PoseSearch::FFeatureVectorBuilder& InOutQuery) const override;
+	virtual void BuildQuery(UE::PoseSearch::FSearchContext& SearchContext) const override;
 
 	virtual void AddDependentChannels(UPoseSearchSchema* Schema) const override;
 
diff --git a/Engine/Plugins/Experimental/Animation/PoseSearch/Source/Runtime/Private/PoseSearchFeatureChannel_Group.cpp b/Engine/Plugins/Experimental/Animation/PoseSearch/Source/Runtime/Private/PoseSearchFeatureChannel_Group.cpp
--- a/Engine/Plugins/Experimental/Animation/PoseSearch/Source/Runtime/Private/PoseSearchFeatureChannel_Group.cpp
+++ b/Engine/Plugins/Experimental/Animation/PoseSearch/Source/Runtime/Private/PoseSearchFeatureChannel_Group.cpp
@@ -55,13 +55,13 @@ bool UPoseSearchFeatureChannel_GroupBase::IndexAsset(UE::PoseSearch::FAssetIndex
 }
 #endif // WITH_EDITOR
 
-void UPoseSearchFeatureChannel_GroupBase::BuildQuery(UE::PoseSearch::FSearchContext& SearchContext, UE::PoseSearch::FFeatureVectorBuilder& InOutQuery) const
+void UPoseSearchFeatureChannel_GroupBase::BuildQuery(UE::PoseSearch::FSearchContext& SearchContext) const
 {
 	for (const TObjectPtr<UPoseSearchFeatureChannel>& SubChannelPtr : GetSubChannels())
 	{
 		if (const UPoseSearchFeatureChannel* SubChannel = SubChannelPtr.Get())
 		{
-			SubChannel->BuildQuery(SearchContext, InOutQuery);
+			SubChannel->BuildQuery(SearchContext);
 		}
 	}
 }
diff --git a/Engine/Plugins/Experimental/Animation/PoseSearch/Source/Runtime/Private/PoseSearchFeatureChannel_Group.h b/Engine/Plugins/Experimental/Animation/PoseSearch/Source/Runtime/Private/PoseSearchFeatureChannel_Group.h
--- a/Engine/Plugins/Experimental/Animation/PoseSearch/Source/Runtime/Private/PoseSearchFeatureChannel_Group.h
+++ b/Engine/Plugins/Experimental/Animation/PoseSearch/Source/Runtime/Private/PoseSearchFeatureChannel_Group.h
@@ -15,7 +15,7 @@ class POSESEARCH_API UPoseSearchFeatureChannel_GroupBase : public UPoseSearchFea
 public:
 	// UPoseSearchFeatureChannel interface
 	virtual void Finalize(UPoseSearchSchema* Schema) override;
-	virtual void BuildQuery(UE::PoseSearch::FSearchContext& SearchContext, UE::PoseSearch::FFeatureVectorBuilder& InOutQuery) const override;
+	virtual void BuildQuery(UE::PoseSearch::FSearchContext& SearchContext) const override;
 	virtual void AddDependentChannels(UPoseSearchSchema* Schema) const override; 
 
 #if ENABLE_DRAW_DEBUG
diff --git a/Engine/Plugins/Experimental/Animation/PoseSearch/Source/Runtime/Private/PoseSearchFeatureChannel_Heading.cpp b/Engine/Plugins/Experimental/Animation/PoseSearch/Source/Runtime/Private/PoseSearchFeatureChannel_Heading.cpp
--- a/Engine/Plugins/Experimental/Animation/PoseSearch/Source/Runtime/Private/PoseSearchFeatureChannel_Heading.cpp
+++ b/Engine/Plugins/Experimental/Animation/PoseSearch/Source/Runtime/Private/PoseSearchFeatureChannel_Heading.cpp
@@ -85,41 +85,79 @@ FVector UPoseSearchFeatureChannel_Heading::GetAxis(const FQuat& Rotation) const
 	return FVector::XAxisVector;
 }
 
-void UPoseSearchFeatureChannel_Heading::BuildQuery(UE::PoseSearch::FSearchContext& SearchContext, UE::PoseSearch::FFeatureVectorBuilder& InOutQuery) const
+void UPoseSearchFeatureChannel_Heading::BuildQuery(UE::PoseSearch::FSearchContext& SearchContext) const
 {
 	using namespace UE::PoseSearch;
 
-	check(InOutQuery.GetSchema());
 	const bool bIsRootBone = SchemaBoneIdx == RootSchemaBoneIdx;
 	if (bUseBlueprintQueryOverride)
 	{
 		const FQuat BoneRotationWorld = BP_GetWorldRotation(SearchContext.GetAnimInstance());
-		const FQuat BoneRotation = SearchContext.GetSampleRotation(SampleTimeOffset, OriginTimeOffset, InOutQuery.GetSchema(), SchemaBoneIdx, SchemaOriginBoneIdx, EPermutationTimeType::UseSampleTime, &BoneRotationWorld);
-		FFeatureVectorHelper::EncodeVector(InOutQuery.EditValues(), ChannelDataOffset, GetAxis(BoneRotation), ComponentStripping);
+		const FQuat BoneRotation = SearchContext.GetSampleRotation(SampleTimeOffset, OriginTimeOffset, SchemaBoneIdx, SchemaOriginBoneIdx, EPermutationTimeType::UseSampleTime, &BoneRotationWorld);
+		FFeatureVectorHelper::EncodeVector(SearchContext.EditFeatureVector(), ChannelDataOffset, GetAxis(BoneRotation), ComponentStripping);
+		return;
 	}
-	else
+	
+	// trying to get the BuildQuery data from another schema UPoseSearchFeatureChannel_Heading already cached in the SearchContext
+	if (SearchContext.IsUseCachedChannelData())
 	{
-		const bool bIsCurrentResultValid = SearchContext.GetCurrentResult().IsValid() && SearchContext.GetCurrentResult().Database->Schema == InOutQuery.GetSchema();
-		const bool bSkip = InputQueryPose != EInputQueryPose::UseCharacterPose && bIsCurrentResultValid;
-		if (bSkip || (!SearchContext.IsHistoryValid() && !bIsRootBone))
+		// composing a unique identifier to specify this channel with all the required properties to be able to share the query data with other channels of the same type
+		uint32 UniqueIdentifier = GetClass()->GetUniqueID();
+		UniqueIdentifier = HashCombineFast(UniqueIdentifier, ::GetTypeHash(SamplingAttributeId));
+		UniqueIdentifier = HashCombineFast(UniqueIdentifier, ::GetTypeHash(SampleTimeOffset));
+		UniqueIdentifier = HashCombineFast(UniqueIdentifier, ::GetTypeHash(OriginTimeOffset));
+		UniqueIdentifier = HashCombineFast(UniqueIdentifier, ::GetTypeHash(HeadingAxis));
+		UniqueIdentifier = HashCombineFast(UniqueIdentifier, ::GetTypeHash(SchemaBoneIdx));
+		UniqueIdentifier = HashCombineFast(UniqueIdentifier, ::GetTypeHash(SchemaOriginBoneIdx));
+		UniqueIdentifier = HashCombineFast(UniqueIdentifier, ::GetTypeHash(InputQueryPose));
+		UniqueIdentifier = HashCombineFast(UniqueIdentifier, ::GetTypeHash(ComponentStripping));
+		UniqueIdentifier = HashCombineFast(UniqueIdentifier, ::GetTypeHash(PermutationTimeType));
+
+		TConstArrayView<float> CachedChannelData;
+		if (const UPoseSearchFeatureChannel* CachedChannel = SearchContext.GetCachedChannelData(UniqueIdentifier, this, CachedChannelData))
 		{
-			if (bIsCurrentResultValid)
-			{
-				FFeatureVectorHelper::Copy(InOutQuery.EditValues(), ChannelDataOffset, ChannelCardinality, SearchContext.GetCurrentResultPoseVector());
-			}
-			else
-			{
-				// we leave the InOutQuery set to zero since the SearchContext.History is invalid and it'll fail if we continue
-				UE_LOG(LogPoseSearch, Error, TEXT("UPoseSearchFeatureChannel_Heading::BuildQuery - Failed because Pose History Node is missing."));
-			}
+#if DO_CHECK
+			const UPoseSearchFeatureChannel_Heading* CachedHeadingChannel = Cast<UPoseSearchFeatureChannel_Heading>(CachedChannel);
+			check(CachedHeadingChannel);
+			check(CachedHeadingChannel->GetChannelCardinality() == ChannelCardinality);
+			check(CachedChannelData.Num() == ChannelCardinality);
+
+			// making sure there were no hash collisions
+			check(CachedHeadingChannel->SamplingAttributeId == SamplingAttributeId);
+			check(CachedHeadingChannel->SampleTimeOffset == SampleTimeOffset);
+			check(CachedHeadingChannel->OriginTimeOffset == OriginTimeOffset);
+			check(CachedHeadingChannel->HeadingAxis == HeadingAxis);
+			check(CachedHeadingChannel->SchemaBoneIdx == SchemaBoneIdx);
+			check(CachedHeadingChannel->SchemaOriginBoneIdx == SchemaOriginBoneIdx);
+			check(CachedHeadingChannel->InputQueryPose == InputQueryPose);
+			check(CachedHeadingChannel->ComponentStripping == ComponentStripping);
+			check(CachedHeadingChannel->PermutationTimeType == PermutationTimeType);
+#endif //DO_CHECK
+
+			// copying the CachedChannelData into this channel portion of the FeatureVectorBuilder
+			FFeatureVectorHelper::Copy(SearchContext.EditFeatureVector().Slice(ChannelDataOffset, ChannelCardinality), 0, ChannelCardinality, CachedChannelData);
+			return;
 		}
-		else
+	}
+
+	const bool bCanUseCurrentResult = SearchContext.CanUseCurrentResult();
+	const bool bSkip = InputQueryPose != EInputQueryPose::UseCharacterPose && bCanUseCurrentResult;
+	if (bSkip || (!SearchContext.IsHistoryValid() && !bIsRootBone))
+	{
+		if (bCanUseCurrentResult)
 		{
-			// calculating the BoneRotation in component space for the bone indexed by SchemaBoneIdx
-			const FQuat BoneRotation = SearchContext.GetSampleRotation(SampleTimeOffset, OriginTimeOffset, InOutQuery.GetSchema(), SchemaBoneIdx, SchemaOriginBoneIdx, PermutationTimeType);
-			FFeatureVectorHelper::EncodeVector(InOutQuery.EditValues(), ChannelDataOffset, GetAxis(BoneRotation), ComponentStripping);
+			FFeatureVectorHelper::Copy(SearchContext.EditFeatureVector(), ChannelDataOffset, ChannelCardinality, SearchContext.GetCurrentResultPoseVector());
+			return;
 		}
+		
+		// we leave the SearchContext.EditFeatureVector() set to zero since the SearchContext.History is invalid and it'll fail if we continue
+		UE_LOG(LogPoseSearch, Error, TEXT("UPoseSearchFeatureChannel_Heading::BuildQuery - Failed because Pose History Node is missing."));
+		return;
 	}
+	
+	// calculating the BoneRotation in component space for the bone indexed by SchemaBoneIdx
+	const FQuat BoneRotation = SearchContext.GetSampleRotation(SampleTimeOffset, OriginTimeOffset, SchemaBoneIdx, SchemaOriginBoneIdx, PermutationTimeType);
+	FFeatureVectorHelper::EncodeVector(SearchContext.EditFeatureVector(), ChannelDataOffset, GetAxis(BoneRotation), ComponentStripping);
 }
 
 #if ENABLE_DRAW_DEBUG
diff --git a/Engine/Plugins/Experimental/Animation/PoseSearch/Source/Runtime/Private/PoseSearchFeatureChannel_Heading.h b/Engine/Plugins/Experimental/Animation/PoseSearch/Source/Runtime/Private/PoseSearchFeatureChannel_Heading.h
--- a/Engine/Plugins/Experimental/Animation/PoseSearch/Source/Runtime/Private/PoseSearchFeatureChannel_Heading.h
+++ b/Engine/Plugins/Experimental/Animation/PoseSearch/Source/Runtime/Private/PoseSearchFeatureChannel_Heading.h
@@ -84,7 +84,7 @@ public:
 
 	// UPoseSearchFeatureChannel interface
 	virtual void Finalize(UPoseSearchSchema* Schema) override;
-	virtual void BuildQuery(UE::PoseSearch::FSearchContext& SearchContext, UE::PoseSearch::FFeatureVectorBuilder& InOutQuery) const override;
+	virtual void BuildQuery(UE::PoseSearch::FSearchContext& SearchContext) const override;
 
 	virtual EPermutationTimeType GetPermutationTimeType() const override { return PermutationTimeType; }
 	virtual void AddDependentChannels(UPoseSearchSchema* Schema) const override;
diff --git a/Engine/Plugins/Experimental/Animation/PoseSearch/Source/Runtime/Private/PoseSearchFeatureChannel_Padding.cpp b/Engine/Plugins/Experimental/Animation/PoseSearch/Source/Runtime/Private/PoseSearchFeatureChannel_Padding.cpp
--- a/Engine/Plugins/Experimental/Animation/PoseSearch/Source/Runtime/Private/PoseSearchFeatureChannel_Padding.cpp
+++ b/Engine/Plugins/Experimental/Animation/PoseSearch/Source/Runtime/Private/PoseSearchFeatureChannel_Padding.cpp
@@ -19,7 +19,7 @@ void UPoseSearchFeatureChannel_Padding::Finalize(UPoseSearchSchema* Schema)
 	Schema->SchemaCardinality += ChannelCardinality;
 }
 
-void UPoseSearchFeatureChannel_Padding::BuildQuery(UE::PoseSearch::FSearchContext& SearchContext, UE::PoseSearch::FFeatureVectorBuilder& InOutQuery) const
+void UPoseSearchFeatureChannel_Padding::BuildQuery(UE::PoseSearch::FSearchContext& SearchContext) const
 {
 	// nothing to do here
 }
diff --git a/Engine/Plugins/Experimental/Animation/PoseSearch/Source/Runtime/Private/PoseSearchFeatureChannel_Padding.h b/Engine/Plugins/Experimental/Animation/PoseSearch/Source/Runtime/Private/PoseSearchFeatureChannel_Padding.h
--- a/Engine/Plugins/Experimental/Animation/PoseSearch/Source/Runtime/Private/PoseSearchFeatureChannel_Padding.h
+++ b/Engine/Plugins/Experimental/Animation/PoseSearch/Source/Runtime/Private/PoseSearchFeatureChannel_Padding.h
@@ -16,7 +16,7 @@ public:
 
 	// UPoseSearchFeatureChannel interface
 	virtual void Finalize(UPoseSearchSchema* Schema) override;
-	virtual void BuildQuery(UE::PoseSearch::FSearchContext& SearchContext, UE::PoseSearch::FFeatureVectorBuilder& InOutQuery) const override;
+	virtual void BuildQuery(UE::PoseSearch::FSearchContext& SearchContext) const override;
 
 #if WITH_EDITOR
 	virtual void FillWeights(TArrayView<float> Weights) const override;
diff --git a/Engine/Plugins/Experimental/Animation/PoseSearch/Source/Runtime/Private/PoseSearchFeatureChannel_PermutationTime.cpp b/Engine/Plugins/Experimental/Animation/PoseSearch/Source/Runtime/Private/PoseSearchFeatureChannel_PermutationTime.cpp
--- a/Engine/Plugins/Experimental/Animation/PoseSearch/Source/Runtime/Private/PoseSearchFeatureChannel_PermutationTime.cpp
+++ b/Engine/Plugins/Experimental/Animation/PoseSearch/Source/Runtime/Private/PoseSearchFeatureChannel_PermutationTime.cpp
@@ -24,10 +24,10 @@ void UPoseSearchFeatureChannel_PermutationTime::Finalize(UPoseSearchSchema* Sche
 	Schema->SchemaCardinality += ChannelCardinality;
 }
 
-void UPoseSearchFeatureChannel_PermutationTime::BuildQuery(UE::PoseSearch::FSearchContext& SearchContext, UE::PoseSearch::FFeatureVectorBuilder& InOutQuery) const
+void UPoseSearchFeatureChannel_PermutationTime::BuildQuery(UE::PoseSearch::FSearchContext& SearchContext) const
 {
 	using namespace UE::PoseSearch;
-	FFeatureVectorHelper::EncodeFloat(InOutQuery.EditValues(), ChannelDataOffset, SearchContext.GetDesiredPermutationTimeOffset());
+	FFeatureVectorHelper::EncodeFloat(SearchContext.EditFeatureVector(), ChannelDataOffset, SearchContext.GetDesiredPermutationTimeOffset());
 }
 
 #if WITH_EDITOR
diff --git a/Engine/Plugins/Experimental/Animation/PoseSearch/Source/Runtime/Private/PoseSearchFeatureChannel_PermutationTime.h b/Engine/Plugins/Experimental/Animation/PoseSearch/Source/Runtime/Private/PoseSearchFeatureChannel_PermutationTime.h
--- a/Engine/Plugins/Experimental/Animation/PoseSearch/Source/Runtime/Private/PoseSearchFeatureChannel_PermutationTime.h
+++ b/Engine/Plugins/Experimental/Animation/PoseSearch/Source/Runtime/Private/PoseSearchFeatureChannel_PermutationTime.h
@@ -18,7 +18,7 @@ public:
 
 	// UPoseSearchFeatureChannel interface
 	virtual void Finalize(UPoseSearchSchema* Schema) override;
-	virtual void BuildQuery(UE::PoseSearch::FSearchContext& SearchContext, UE::PoseSearch::FFeatureVectorBuilder& InOutQuery) const override;
+	virtual void BuildQuery(UE::PoseSearch::FSearchContext& SearchContext) const override;
 
 #if WITH_EDITOR
 	virtual void FillWeights(TArrayView<float> Weights) const override;
diff --git a/Engine/Plugins/Experimental/Animation/PoseSearch/Source/Runtime/Private/PoseSearchFeatureChannel_Phase.cpp b/Engine/Plugins/Experimental/Animation/PoseSearch/Source/Runtime/Private/PoseSearchFeatureChannel_Phase.cpp
--- a/Engine/Plugins/Experimental/Animation/PoseSearch/Source/Runtime/Private/PoseSearchFeatureChannel_Phase.cpp
+++ b/Engine/Plugins/Experimental/Animation/PoseSearch/Source/Runtime/Private/PoseSearchFeatureChannel_Phase.cpp
@@ -300,29 +300,55 @@ void UPoseSearchFeatureChannel_Phase::AddDependentChannels(UPoseSearchSchema* Sc
 	}
 }
 
-void UPoseSearchFeatureChannel_Phase::BuildQuery(UE::PoseSearch::FSearchContext& SearchContext, UE::PoseSearch::FFeatureVectorBuilder& InOutQuery) const
+void UPoseSearchFeatureChannel_Phase::BuildQuery(UE::PoseSearch::FSearchContext& SearchContext) const
 {
 	using namespace UE::PoseSearch;
 
-	const bool bIsCurrentResultValid = SearchContext.GetCurrentResult().IsValid() && SearchContext.GetCurrentResult().Database->Schema == InOutQuery.GetSchema();
-	const bool bSkip = InputQueryPose != EInputQueryPose::UseCharacterPose && bIsCurrentResultValid;
-	if (bSkip || !SearchContext.IsHistoryValid())
+	// trying to get the BuildQuery data from another schema UPoseSearchFeatureChannel_Phase already cached in the SearchContext
+	if (SearchContext.IsUseCachedChannelData())
 	{
-		if (bIsCurrentResultValid)
-		{
-			FFeatureVectorHelper::Copy(InOutQuery.EditValues(), ChannelDataOffset, ChannelCardinality, SearchContext.GetCurrentResultPoseVector());
-		}
-		else
+		// composing a unique identifier to specify this channel with all the required properties to be able to share the query data with other channels of the same type
+		uint32 UniqueIdentifier = GetClass()->GetUniqueID();
+		UniqueIdentifier = HashCombineFast(UniqueIdentifier, ::GetTypeHash(SchemaBoneIdx));
+		UniqueIdentifier = HashCombineFast(UniqueIdentifier, ::GetTypeHash(InputQueryPose));
+
+		TConstArrayView<float> CachedChannelData;
+		if (const UPoseSearchFeatureChannel* CachedChannel = SearchContext.GetCachedChannelData(UniqueIdentifier, this, CachedChannelData))
 		{
-			// we leave the InOutQuery set to zero since the SearchContext.History is invalid and it'll fail if we continue
-			UE_LOG(LogPoseSearch, Error, TEXT("UPoseSearchFeatureChannel_Phase::BuildQuery - Failed because Pose History Node is missing."));
+#if DO_CHECK
+			const UPoseSearchFeatureChannel_Phase* CachedPhaseChannel = Cast<UPoseSearchFeatureChannel_Phase>(CachedChannel);
+			check(CachedPhaseChannel);
+			check(CachedPhaseChannel->GetChannelCardinality() == ChannelCardinality);
+			check(CachedChannelData.Num() == ChannelCardinality);
+
+			// making sure there were no hash collisions
+			check(CachedPhaseChannel->SchemaBoneIdx == SchemaBoneIdx);
+			check(CachedPhaseChannel->InputQueryPose == InputQueryPose);
+#endif //DO_CHECK
+
+			// copying the CachedChannelData into this channel portion of the FeatureVectorBuilder
+			FFeatureVectorHelper::Copy(SearchContext.EditFeatureVector().Slice(ChannelDataOffset, ChannelCardinality), 0, ChannelCardinality, CachedChannelData);
+			return;
 		}
 	}
-	else
+
+	const bool bCanUseCurrentResult = SearchContext.CanUseCurrentResult();
+	const bool bSkip = InputQueryPose != EInputQueryPose::UseCharacterPose && bCanUseCurrentResult;
+	if (bSkip || !SearchContext.IsHistoryValid())
 	{
-		// @todo: Support phase in BuildQuery
-		// FFeatureVectorHelper::EncodeVector2D(InOutQuery.EditValues(), DataOffset, ???);
+		if (bCanUseCurrentResult)
+		{
+			FFeatureVectorHelper::Copy(SearchContext.EditFeatureVector(), ChannelDataOffset, ChannelCardinality, SearchContext.GetCurrentResultPoseVector());
+			return;
+		}
+		
+		// we leave the SearchContext.EditFeatureVector() set to zero since the SearchContext.History is invalid and it'll fail if we continue
+		UE_LOG(LogPoseSearch, Error, TEXT("UPoseSearchFeatureChannel_Phase::BuildQuery - Failed because Pose History Node is missing."));
+		return;
 	}
+	
+	// @todo: Support phase in BuildQuery
+	// FFeatureVectorHelper::EncodeVector2D(SearchContext.EditFeatureVector(), DataOffset, ???);
 }
 
 #if ENABLE_DRAW_DEBUG
diff --git a/Engine/Plugins/Experimental/Animation/PoseSearch/Source/Runtime/Private/PoseSearchFeatureChannel_Phase.h b/Engine/Plugins/Experimental/Animation/PoseSearch/Source/Runtime/Private/PoseSearchFeatureChannel_Phase.h
--- a/Engine/Plugins/Experimental/Animation/PoseSearch/Source/Runtime/Private/PoseSearchFeatureChannel_Phase.h
+++ b/Engine/Plugins/Experimental/Animation/PoseSearch/Source/Runtime/Private/PoseSearchFeatureChannel_Phase.h
@@ -34,7 +34,7 @@ public:
 
 	// UPoseSearchFeatureChannel interface
 	virtual void Finalize(UPoseSearchSchema* Schema) override;
-	virtual void BuildQuery(UE::PoseSearch::FSearchContext& SearchContext, UE::PoseSearch::FFeatureVectorBuilder& InOutQuery) const override;
+	virtual void BuildQuery(UE::PoseSearch::FSearchContext& SearchContext) const override;
 
 	virtual void AddDependentChannels(UPoseSearchSchema* Schema) const override;
 
diff --git a/Engine/Plugins/Experimental/Animation/PoseSearch/Source/Runtime/Private/PoseSearchFeatureChannel_Position.cpp b/Engine/Plugins/Experimental/Animation/PoseSearch/Source/Runtime/Private/PoseSearchFeatureChannel_Position.cpp
--- a/Engine/Plugins/Experimental/Animation/PoseSearch/Source/Runtime/Private/PoseSearchFeatureChannel_Position.cpp
+++ b/Engine/Plugins/Experimental/Animation/PoseSearch/Source/Runtime/Private/PoseSearchFeatureChannel_Position.cpp
@@ -64,41 +64,77 @@ void UPoseSearchFeatureChannel_Position::AddDependentChannels(UPoseSearchSchema*
 	}
 }
 
-void UPoseSearchFeatureChannel_Position::BuildQuery(UE::PoseSearch::FSearchContext& SearchContext, UE::PoseSearch::FFeatureVectorBuilder& InOutQuery) const
+void UPoseSearchFeatureChannel_Position::BuildQuery(UE::PoseSearch::FSearchContext& SearchContext) const
 {
 	using namespace UE::PoseSearch;
 
-	check(InOutQuery.GetSchema());
 	if (bUseBlueprintQueryOverride)
 	{
 		const FVector BonePositionWorld = BP_GetWorldPosition(SearchContext.GetAnimInstance());
-		const FVector BonePosition = SearchContext.GetSamplePosition(SampleTimeOffset, OriginTimeOffset, InOutQuery.GetSchema(), SchemaBoneIdx, SchemaOriginBoneIdx, PermutationTimeType, &BonePositionWorld);
-  		FFeatureVectorHelper::EncodeVector(InOutQuery.EditValues(), ChannelDataOffset, BonePosition, ComponentStripping);
+		const FVector BonePosition = SearchContext.GetSamplePosition(SampleTimeOffset, OriginTimeOffset, SchemaBoneIdx, SchemaOriginBoneIdx, PermutationTimeType, &BonePositionWorld);
+  		FFeatureVectorHelper::EncodeVector(SearchContext.EditFeatureVector(), ChannelDataOffset, BonePosition, ComponentStripping);
+		return;
 	}
-	else
+
+	// trying to get the BuildQuery data from another schema UPoseSearchFeatureChannel_Position already cached in the SearchContext
+	if (SearchContext.IsUseCachedChannelData())
 	{
-		const bool bIsCurrentResultValid = SearchContext.GetCurrentResult().IsValid() && SearchContext.GetCurrentResult().Database->Schema == InOutQuery.GetSchema();
-		const bool bSkip = InputQueryPose != EInputQueryPose::UseCharacterPose && bIsCurrentResultValid;
-		const bool bIsRootBone = SchemaBoneIdx == RootSchemaBoneIdx;
-		if (bSkip || (!SearchContext.IsHistoryValid() && !bIsRootBone))
+		// composing a unique identifier to specify this channel with all the required properties to be able to share the query data with other channels of the same type
+		uint32 UniqueIdentifier = GetClass()->GetUniqueID();
+		UniqueIdentifier = HashCombineFast(UniqueIdentifier, ::GetTypeHash(SamplingAttributeId));
+		UniqueIdentifier = HashCombineFast(UniqueIdentifier, ::GetTypeHash(SampleTimeOffset));
+		UniqueIdentifier = HashCombineFast(UniqueIdentifier, ::GetTypeHash(OriginTimeOffset));
+		UniqueIdentifier = HashCombineFast(UniqueIdentifier, ::GetTypeHash(SchemaBoneIdx));
+		UniqueIdentifier = HashCombineFast(UniqueIdentifier, ::GetTypeHash(SchemaOriginBoneIdx));
+		UniqueIdentifier = HashCombineFast(UniqueIdentifier, ::GetTypeHash(InputQueryPose));
+		UniqueIdentifier = HashCombineFast(UniqueIdentifier, ::GetTypeHash(ComponentStripping));
+		UniqueIdentifier = HashCombineFast(UniqueIdentifier, ::GetTypeHash(PermutationTimeType));
+
+		TConstArrayView<float> CachedChannelData;
+		if (const UPoseSearchFeatureChannel* CachedChannel = SearchContext.GetCachedChannelData(UniqueIdentifier, this, CachedChannelData))
 		{
-			if (bIsCurrentResultValid)
-			{
-				FFeatureVectorHelper::Copy(InOutQuery.EditValues(), ChannelDataOffset, ChannelCardinality, SearchContext.GetCurrentResultPoseVector());
-			}
-			else
-			{
-				// we leave the InOutQuery set to zero since the SearchContext.History is invalid and it'll fail if we continue
-				UE_LOG(LogPoseSearch, Error, TEXT("UPoseSearchFeatureChannel_Position::BuildQuery - Failed because Pose History Node is missing."));
-			}
+#if DO_CHECK
+			const UPoseSearchFeatureChannel_Position* CachedPositionChannel = Cast<UPoseSearchFeatureChannel_Position>(CachedChannel);
+			check(CachedPositionChannel);
+			check(CachedPositionChannel->GetChannelCardinality() == ChannelCardinality);
+			check(CachedChannelData.Num() == ChannelCardinality);
+
+			// making sure there were no hash collisions
+			check(CachedPositionChannel->SamplingAttributeId == SamplingAttributeId);
+			check(CachedPositionChannel->SampleTimeOffset == SampleTimeOffset);
+			check(CachedPositionChannel->OriginTimeOffset == OriginTimeOffset);
+			check(CachedPositionChannel->SchemaBoneIdx == SchemaBoneIdx);
+			check(CachedPositionChannel->SchemaOriginBoneIdx == SchemaOriginBoneIdx);
+			check(CachedPositionChannel->InputQueryPose == InputQueryPose);
+			check(CachedPositionChannel->ComponentStripping == ComponentStripping);
+			check(CachedPositionChannel->PermutationTimeType == PermutationTimeType);
+#endif //DO_CHECK
+
+			// copying the CachedChannelData into this channel portion of the FeatureVectorBuilder
+			FFeatureVectorHelper::Copy(SearchContext.EditFeatureVector().Slice(ChannelDataOffset, ChannelCardinality), 0, ChannelCardinality, CachedChannelData);
+			return;
 		}
-		else
+	}
+
+	const bool bCanUseCurrentResult = SearchContext.CanUseCurrentResult();
+	const bool bSkip = InputQueryPose != EInputQueryPose::UseCharacterPose && bCanUseCurrentResult;
+	const bool bIsRootBone = SchemaBoneIdx == RootSchemaBoneIdx;
+	if (bSkip || (!SearchContext.IsHistoryValid() && !bIsRootBone))
+	{
+		if (bCanUseCurrentResult)
 		{
-			// calculating the BonePosition in root bone space for the bone indexed by SchemaBoneIdx
-			const FVector BonePosition = SearchContext.GetSamplePosition(SampleTimeOffset, OriginTimeOffset, InOutQuery.GetSchema(), SchemaBoneIdx, SchemaOriginBoneIdx, PermutationTimeType);
-			FFeatureVectorHelper::EncodeVector(InOutQuery.EditValues(), ChannelDataOffset, BonePosition, ComponentStripping);
+			FFeatureVectorHelper::Copy(SearchContext.EditFeatureVector(), ChannelDataOffset, ChannelCardinality, SearchContext.GetCurrentResultPoseVector());
+			return;
 		}
+
+		// we leave the SearchContext.EditFeatureVector() set to zero since the SearchContext.History is invalid and it'll fail if we continue
+		UE_LOG(LogPoseSearch, Error, TEXT("UPoseSearchFeatureChannel_Position::BuildQuery - Failed because Pose History Node is missing."));
+		return;
 	}
+	
+	// calculating the BonePosition in root bone space for the bone indexed by SchemaBoneIdx
+	const FVector BonePosition = SearchContext.GetSamplePosition(SampleTimeOffset, OriginTimeOffset, SchemaBoneIdx, SchemaOriginBoneIdx, PermutationTimeType);
+	FFeatureVectorHelper::EncodeVector(SearchContext.EditFeatureVector(), ChannelDataOffset, BonePosition, ComponentStripping);
 }
 
 #if ENABLE_DRAW_DEBUG
diff --git a/Engine/Plugins/Experimental/Animation/PoseSearch/Source/Runtime/Private/PoseSearchFeatureChannel_Position.h b/Engine/Plugins/Experimental/Animation/PoseSearch/Source/Runtime/Private/PoseSearchFeatureChannel_Position.h
--- a/Engine/Plugins/Experimental/Animation/PoseSearch/Source/Runtime/Private/PoseSearchFeatureChannel_Position.h
+++ b/Engine/Plugins/Experimental/Animation/PoseSearch/Source/Runtime/Private/PoseSearchFeatureChannel_Position.h
@@ -70,7 +70,7 @@ public:
 
 	// UPoseSearchFeatureChannel interface
 	virtual void Finalize(UPoseSearchSchema* Schema) override;
-	virtual void BuildQuery(UE::PoseSearch::FSearchContext& SearchContext, UE::PoseSearch::FFeatureVectorBuilder& InOutQuery) const override;
+	virtual void BuildQuery(UE::PoseSearch::FSearchContext& SearchContext) const override;
 
 	virtual EPermutationTimeType GetPermutationTimeType() const override { return PermutationTimeType; }
 	virtual void AddDependentChannels(UPoseSearchSchema* Schema) const override;
diff --git a/Engine/Plugins/Experimental/Animation/PoseSearch/Source/Runtime/Private/PoseSearchFeatureChannel_SamplingTime.cpp b/Engine/Plugins/Experimental/Animation/PoseSearch/Source/Runtime/Private/PoseSearchFeatureChannel_SamplingTime.cpp
--- a/Engine/Plugins/Experimental/Animation/PoseSearch/Source/Runtime/Private/PoseSearchFeatureChannel_SamplingTime.cpp
+++ b/Engine/Plugins/Experimental/Animation/PoseSearch/Source/Runtime/Private/PoseSearchFeatureChannel_SamplingTime.cpp
@@ -12,10 +12,10 @@ void UPoseSearchFeatureChannel_SamplingTime::Finalize(UPoseSearchSchema* Schema)
 	Schema->SchemaCardinality += ChannelCardinality;
 }
 
-void UPoseSearchFeatureChannel_SamplingTime::BuildQuery(UE::PoseSearch::FSearchContext& SearchContext, UE::PoseSearch::FFeatureVectorBuilder& InOutQuery) const
+void UPoseSearchFeatureChannel_SamplingTime::BuildQuery(UE::PoseSearch::FSearchContext& SearchContext) const
 {
 	using namespace UE::PoseSearch;
-	FFeatureVectorHelper::EncodeFloat(InOutQuery.EditValues(), ChannelDataOffset, TimeToMatch);
+	FFeatureVectorHelper::EncodeFloat(SearchContext.EditFeatureVector(), ChannelDataOffset, TimeToMatch);
 }
 
 #if WITH_EDITOR
diff --git a/Engine/Plugins/Experimental/Animation/PoseSearch/Source/Runtime/Private/PoseSearchFeatureChannel_SamplingTime.h b/Engine/Plugins/Experimental/Animation/PoseSearch/Source/Runtime/Private/PoseSearchFeatureChannel_SamplingTime.h
--- a/Engine/Plugins/Experimental/Animation/PoseSearch/Source/Runtime/Private/PoseSearchFeatureChannel_SamplingTime.h
+++ b/Engine/Plugins/Experimental/Animation/PoseSearch/Source/Runtime/Private/PoseSearchFeatureChannel_SamplingTime.h
@@ -20,7 +20,7 @@ public:
 
 	// UPoseSearchFeatureChannel interface
 	virtual void Finalize(UPoseSearchSchema* Schema) override;
-	virtual void BuildQuery(UE::PoseSearch::FSearchContext& SearchContext, UE::PoseSearch::FFeatureVectorBuilder& InOutQuery) const override;
+	virtual void BuildQuery(UE::PoseSearch::FSearchContext& SearchContext) const override;
 
 #if WITH_EDITOR
 	virtual void FillWeights(TArrayView<float> Weights) const override;
diff --git a/Engine/Plugins/Experimental/Animation/PoseSearch/Source/Runtime/Private/PoseSearchFeatureChannel_Velocity.cpp b/Engine/Plugins/Experimental/Animation/PoseSearch/Source/Runtime/Private/PoseSearchFeatureChannel_Velocity.cpp
--- a/Engine/Plugins/Experimental/Animation/PoseSearch/Source/Runtime/Private/PoseSearchFeatureChannel_Velocity.cpp
+++ b/Engine/Plugins/Experimental/Animation/PoseSearch/Source/Runtime/Private/PoseSearchFeatureChannel_Velocity.cpp
@@ -33,7 +33,7 @@ void UPoseSearchFeatureChannel_Velocity::AddDependentChannels(UPoseSearchSchema*
 	}
 }
 
-void UPoseSearchFeatureChannel_Velocity::BuildQuery(UE::PoseSearch::FSearchContext& SearchContext, UE::PoseSearch::FFeatureVectorBuilder& InOutQuery) const
+void UPoseSearchFeatureChannel_Velocity::BuildQuery(UE::PoseSearch::FSearchContext& SearchContext) const
 {
 	using namespace UE::PoseSearch;
 
@@ -42,41 +42,82 @@ void UPoseSearchFeatureChannel_Velocity::BuildQuery(UE::PoseSearch::FSearchConte
 	{
 		const FVector LinearVelocityWorld = BP_GetWorldVelocity(SearchContext.GetAnimInstance());
 
-		FVector LinearVelocity = SearchContext.GetSampleVelocity(SampleTimeOffset, OriginTimeOffset, InOutQuery.GetSchema(), SchemaBoneIdx, SchemaOriginBoneIdx, bUseCharacterSpaceVelocities, EPermutationTimeType::UseSampleTime, &LinearVelocityWorld);
+		FVector LinearVelocity = SearchContext.GetSampleVelocity(SampleTimeOffset, OriginTimeOffset, SchemaBoneIdx, SchemaOriginBoneIdx, bUseCharacterSpaceVelocities, EPermutationTimeType::UseSampleTime, &LinearVelocityWorld);
 		if (bNormalize)
 		{
 			LinearVelocity = LinearVelocity.GetClampedToMaxSize(1.f);
 		}
-		FFeatureVectorHelper::EncodeVector(InOutQuery.EditValues(), ChannelDataOffset, LinearVelocity, ComponentStripping);
+		FFeatureVectorHelper::EncodeVector(SearchContext.EditFeatureVector(), ChannelDataOffset, LinearVelocity, ComponentStripping);
+		return;
 	}
-	else
+	
+	// trying to get the BuildQuery data from another schema UPoseSearchFeatureChannel_Velocity already cached in the SearchContext
+	if (SearchContext.IsUseCachedChannelData())
 	{
-		const bool bIsCurrentResultValid = SearchContext.GetCurrentResult().IsValid() && SearchContext.GetCurrentResult().Database->Schema == InOutQuery.GetSchema();
-		const bool bSkip = InputQueryPose != EInputQueryPose::UseCharacterPose && bIsCurrentResultValid;
-		if (bSkip || (!SearchContext.IsHistoryValid() && !bIsRootBone))
+		// composing a unique identifier to specify this channel with all the required properties to be able to share the query data with other channels of the same type
+		uint32 UniqueIdentifier = GetClass()->GetUniqueID();
+		UniqueIdentifier = HashCombineFast(UniqueIdentifier, ::GetTypeHash(SamplingAttributeId));
+		UniqueIdentifier = HashCombineFast(UniqueIdentifier, ::GetTypeHash(SampleTimeOffset));
+		UniqueIdentifier = HashCombineFast(UniqueIdentifier, ::GetTypeHash(OriginTimeOffset));
+		UniqueIdentifier = HashCombineFast(UniqueIdentifier, ::GetTypeHash(SchemaBoneIdx));
+		UniqueIdentifier = HashCombineFast(UniqueIdentifier, ::GetTypeHash(SchemaOriginBoneIdx));
+		UniqueIdentifier = HashCombineFast(UniqueIdentifier, ::GetTypeHash(InputQueryPose));
+		UniqueIdentifier = HashCombineFast(UniqueIdentifier, ::GetTypeHash(bUseCharacterSpaceVelocities));
+		UniqueIdentifier = HashCombineFast(UniqueIdentifier, ::GetTypeHash(bNormalize));
+		UniqueIdentifier = HashCombineFast(UniqueIdentifier, ::GetTypeHash(ComponentStripping));
+		UniqueIdentifier = HashCombineFast(UniqueIdentifier, ::GetTypeHash(PermutationTimeType));
+
+		TConstArrayView<float> CachedChannelData;
+		if (const UPoseSearchFeatureChannel* CachedChannel = SearchContext.GetCachedChannelData(UniqueIdentifier, this, CachedChannelData))
 		{
-			if (bIsCurrentResultValid)
-			{
-				FFeatureVectorHelper::Copy(InOutQuery.EditValues(), ChannelDataOffset, ChannelCardinality, SearchContext.GetCurrentResultPoseVector());
-			}
-			else
-			{
-				// we leave the InOutQuery set to zero since the SearchContext.History is invalid and it'll fail if we continue
-				UE_LOG(LogPoseSearch, Error, TEXT("UPoseSearchFeatureChannel_Velocity::BuildQuery - Failed because Pose History Node is missing."));
-			}
+#if DO_CHECK
+			const UPoseSearchFeatureChannel_Velocity* CachedVelocityChannel = Cast<UPoseSearchFeatureChannel_Velocity>(CachedChannel);
+			check(CachedVelocityChannel);
+			check(CachedVelocityChannel->GetChannelCardinality() == ChannelCardinality);
+			check(CachedChannelData.Num() == ChannelCardinality);
+
+			// making sure there were no hash collisions
+			check(CachedVelocityChannel->SamplingAttributeId == SamplingAttributeId);
+			check(CachedVelocityChannel->SampleTimeOffset == SampleTimeOffset);
+			check(CachedVelocityChannel->OriginTimeOffset == OriginTimeOffset);
+			check(CachedVelocityChannel->SchemaBoneIdx == SchemaBoneIdx);
+			check(CachedVelocityChannel->SchemaOriginBoneIdx == SchemaOriginBoneIdx);
+			check(CachedVelocityChannel->InputQueryPose == InputQueryPose);
+			check(CachedVelocityChannel->bUseCharacterSpaceVelocities == bUseCharacterSpaceVelocities);
+			check(CachedVelocityChannel->bNormalize == bNormalize);
+			check(CachedVelocityChannel->ComponentStripping == ComponentStripping);
+			check(CachedVelocityChannel->PermutationTimeType == PermutationTimeType);
+#endif //DO_CHECK
+
+			// copying the CachedChannelData into this channel portion of the FeatureVectorBuilder
+			FFeatureVectorHelper::Copy(SearchContext.EditFeatureVector().Slice(ChannelDataOffset, ChannelCardinality), 0, ChannelCardinality, CachedChannelData);
+			return;
 		}
-		else
-		{
-			// calculating the LinearVelocity for the bone indexed by SchemaBoneIdx
-			FVector LinearVelocity = SearchContext.GetSampleVelocity(SampleTimeOffset, OriginTimeOffset, InOutQuery.GetSchema(), SchemaBoneIdx, SchemaOriginBoneIdx, bUseCharacterSpaceVelocities, PermutationTimeType);
-			if (bNormalize)
-			{
-				LinearVelocity = LinearVelocity.GetClampedToMaxSize(1.f);
-			}
+	}
 
-			FFeatureVectorHelper::EncodeVector(InOutQuery.EditValues(), ChannelDataOffset, LinearVelocity, ComponentStripping);
+	const bool bCanUseCurrentResult = SearchContext.CanUseCurrentResult();
+	const bool bSkip = InputQueryPose != EInputQueryPose::UseCharacterPose && bCanUseCurrentResult;
+	if (bSkip || (!SearchContext.IsHistoryValid() && !bIsRootBone))
+	{
+		if (bCanUseCurrentResult)
+		{
+			FFeatureVectorHelper::Copy(SearchContext.EditFeatureVector(), ChannelDataOffset, ChannelCardinality, SearchContext.GetCurrentResultPoseVector());
+			return;
 		}
+
+		// we leave the SearchContext.EditFeatureVector() set to zero since the SearchContext.History is invalid and it'll fail if we continue
+		UE_LOG(LogPoseSearch, Error, TEXT("UPoseSearchFeatureChannel_Velocity::BuildQuery - Failed because Pose History Node is missing."));
+		return;
 	}
+	
+	// calculating the LinearVelocity for the bone indexed by SchemaBoneIdx
+	FVector LinearVelocity = SearchContext.GetSampleVelocity(SampleTimeOffset, OriginTimeOffset, SchemaBoneIdx, SchemaOriginBoneIdx, bUseCharacterSpaceVelocities, PermutationTimeType);
+	if (bNormalize)
+	{
+		LinearVelocity = LinearVelocity.GetClampedToMaxSize(1.f);
+	}
+
+	FFeatureVectorHelper::EncodeVector(SearchContext.EditFeatureVector(), ChannelDataOffset, LinearVelocity, ComponentStripping);
 }
 
 #if ENABLE_DRAW_DEBUG
diff --git a/Engine/Plugins/Experimental/Animation/PoseSearch/Source/Runtime/Private/PoseSearchFeatureChannel_Velocity.h b/Engine/Plugins/Experimental/Animation/PoseSearch/Source/Runtime/Private/PoseSearchFeatureChannel_Velocity.h
--- a/Engine/Plugins/Experimental/Animation/PoseSearch/Source/Runtime/Private/PoseSearchFeatureChannel_Velocity.h
+++ b/Engine/Plugins/Experimental/Animation/PoseSearch/Source/Runtime/Private/PoseSearchFeatureChannel_Velocity.h
@@ -78,7 +78,7 @@ public:
 
 	// UPoseSearchFeatureChannel interface
 	virtual void Finalize(UPoseSearchSchema* Schema) override;
-	virtual void BuildQuery(UE::PoseSearch::FSearchContext& SearchContext, UE::PoseSearch::FFeatureVectorBuilder& InOutQuery) const override;
+	virtual void BuildQuery(UE::PoseSearch::FSearchContext& SearchContext) const override;
 
 	virtual EPermutationTimeType GetPermutationTimeType() const override { return PermutationTimeType; }
 	virtual void AddDependentChannels(UPoseSearchSchema* Schema) const override;
diff --git a/Engine/Plugins/Experimental/Animation/PoseSearch/Source/Runtime/Private/PoseSearchLibrary.cpp b/Engine/Plugins/Experimental/Animation/PoseSearch/Source/Runtime/Private/PoseSearchLibrary.cpp
--- a/Engine/Plugins/Experimental/Animation/PoseSearch/Source/Runtime/Private/PoseSearchLibrary.cpp
+++ b/Engine/Plugins/Experimental/Animation/PoseSearch/Source/Runtime/Private/PoseSearchLibrary.cpp
@@ -75,6 +75,36 @@ namespace UE::PoseSearch
 			return false;
 		}
 	}
+
+	static bool ShouldUseCachedChannelData(const UPoseSearchDatabase* CurrentResultDatabase, const TArray<TObjectPtr<const UPoseSearchDatabase>>& Databases)
+	{
+		const UPoseSearchSchema* OneOfTheSchemas = nullptr;
+		if (CurrentResultDatabase)
+		{
+			OneOfTheSchemas = CurrentResultDatabase->Schema;
+		}
+
+		for (const TObjectPtr<const UPoseSearchDatabase>& Database : Databases)
+		{
+			if (ensure(Database))
+			{
+				if (OneOfTheSchemas != Database->Schema)
+				{
+					if (OneOfTheSchemas == nullptr)
+					{
+						OneOfTheSchemas = Database->Schema;
+					}
+					else
+					{
+						// we found we need to search multiple schemas
+						return true;
+					}
+				}
+			}
+		}
+
+		return false;
+	}
 }
 
 //////////////////////////////////////////////////////////////////////////
@@ -137,11 +167,11 @@ void FMotionMatchingState::UpdateWantedPlayRate(const UE::PoseSearch::FSearchCon
 		}
 		else if (!FMath::IsNearlyEqual(PlayRate.Min, PlayRate.Max, UE_KINDA_SMALL_NUMBER))
 		{
-			if (const UE::PoseSearch::FFeatureVectorBuilder* PoseSearchFeatureVectorBuilder = SearchContext.GetCachedQuery(CurrentSearchResult.Database->Schema))
+			TConstArrayView<float> QueryData = SearchContext.GetCachedQuery(CurrentSearchResult.Database->Schema);
+			if (!QueryData.IsEmpty())
 			{
 				if (const UPoseSearchFeatureChannel_Trajectory* TrajectoryChannel = CurrentSearchResult.Database->Schema->FindFirstChannelOfType<UPoseSearchFeatureChannel_Trajectory>())
 				{
-					TConstArrayView<float> QueryData = PoseSearchFeatureVectorBuilder->GetValues();
 					TConstArrayView<float> ResultData = CurrentSearchResult.Database->GetSearchIndex().GetPoseValues(CurrentSearchResult.PoseIdx);
 					const float EstimatedSpeedRatio = TrajectoryChannel->GetEstimatedSpeedRatio(QueryData, ResultData);
 
@@ -223,7 +253,7 @@ void UPoseSearchLibrary::TraceMotionMatchingState(
 		FTraceMotionMatchingStateDatabaseEntry& DbEntry = TraceState.DatabaseEntries[DbEntryIdx];
 
 		// if throttling is on, the continuing pose can be valid, but no actual search occurred, so the query will not be cached, and we need to build it
-		DbEntry.QueryVector = SearchContext.GetOrBuildQuery(Database->Schema).GetValues();
+		DbEntry.QueryVector = SearchContext.GetOrBuildQuery(Database->Schema);
 		DbEntry.DatabaseId = FTraceMotionMatchingState::GetIdFromObject(Database);
 
 		for (int32 CandidateIdx = 0; CandidateIdx < DatabaseBestPoseCandidates.Value.Num(); ++CandidateIdx)
@@ -301,6 +331,7 @@ void UPoseSearchLibrary::UpdateMotionMatchingState(
 	float YawFromAnimationTrajectoryBlendTime,
 	EPoseSearchInterruptMode InterruptMode,
 	bool bShouldSearch,
+	bool bShouldUseCachedChannelData,
 	bool bDebugDrawQuery,
 	bool bDebugDrawCurResult,
 	bool bDebugDrawPoseHistory)
@@ -363,10 +394,14 @@ void UPoseSearchLibrary::UpdateMotionMatchingState(
 	{
 		InOutMotionMatchingState.ElapsedPoseSearchTime = 0.f;
 		const bool bForceInterrupt = IsForceInterrupt(InterruptMode, CurrentResultDatabase, Databases);
+		const bool bSearchContinuingPose = !bForceInterrupt && bCanAdvance;
+
+		// calculating if it's worth bUseCachedChannelData (if we potentially have to build query with multiple schemas)
+		SearchContext.SetUseCachedChannelData(bShouldUseCachedChannelData && ShouldUseCachedChannelData(bSearchContinuingPose ? CurrentResultDatabase : nullptr, Databases));
 
-		// Evaluate continuing pose
 		FSearchResult SearchResult;
-		if (!bForceInterrupt && bCanAdvance)
+		// Evaluate continuing pose
+		if (bSearchContinuingPose)
 		{
 			SearchResult = CurrentResultDatabase->SearchContinuingPose(SearchContext);
 			SearchContext.UpdateCurrentBestCost(SearchResult.PoseCost);
@@ -455,6 +490,7 @@ void UPoseSearchLibrary::UpdateMotionMatchingState(
 		InOutMotionMatchingState.ElapsedPoseSearchTime += DeltaTime;
 	}
 
+	// @todo: consider moving this into if (bSearch) to avoid calling SearchContext.GetCachedQuery if no search is required
 	InOutMotionMatchingState.UpdateWantedPlayRate(SearchContext, PlayRate, TrajectorySpeedMultiplier);
 
 	InOutMotionMatchingState.PoseIndicesHistory.Update(InOutMotionMatchingState.CurrentSearchResult, DeltaTime, PoseReselectHistory);
@@ -489,7 +525,7 @@ void UPoseSearchLibrary::UpdateMotionMatchingState(
 			{
 				// @todo: use pose history to get the root bone!
 				UE::PoseSearch::FDebugDrawParams DrawParams(Context.AnimInstanceProxy, SearchContext.GetWorldBoneTransformAtTime(0.f), CurResultDatabase, EDebugDrawFlags::DrawQuery);
-				DrawParams.DrawFeatureVector(SearchContext.GetOrBuildQuery(CurResultDatabase->Schema).GetValues());
+				DrawParams.DrawFeatureVector(SearchContext.GetOrBuildQuery(CurResultDatabase->Schema));
 			}
 		}
 	}
@@ -692,7 +728,7 @@ void UPoseSearchLibrary::MotionMatch(
 			if (CVarAnimMotionMatchDrawQueryEnable.GetValueOnAnyThread())
 			{
 				FDebugDrawParams DrawParams(AnimInstanceProxy, SearchContext.GetWorldBoneTransformAtTime(0.f), SearchResult.Database.Get(), EDebugDrawFlags::DrawQuery);
-				DrawParams.DrawFeatureVector(SearchContext.GetOrBuildQuery(SearchResult.Database->Schema).GetValues());
+				DrawParams.DrawFeatureVector(SearchContext.GetOrBuildQuery(SearchResult.Database->Schema));
 			}
 		}
 #endif // ENABLE_DRAW_DEBUG && ENABLE_ANIM_DEBUG
@@ -782,7 +818,7 @@ UE::PoseSearch::FSearchResult UPoseSearchLibrary::MotionMatch(const FAnimationBa
 			if (CVarAnimMotionMatchDrawQueryEnable.GetValueOnAnyThread())
 			{
 				FDebugDrawParams DrawParams(Context.AnimInstanceProxy, SearchContext.GetWorldBoneTransformAtTime(0.f), SearchResult.Database.Get(), EDebugDrawFlags::DrawQuery);
-				DrawParams.DrawFeatureVector(SearchContext.GetOrBuildQuery(SearchResult.Database->Schema).GetValues());
+				DrawParams.DrawFeatureVector(SearchContext.GetOrBuildQuery(SearchResult.Database->Schema));
 			}
 		}
 #endif // ENABLE_DRAW_DEBUG && ENABLE_ANIM_DEBUG
diff --git a/Engine/Plugins/Experimental/Animation/PoseSearch/Source/Runtime/Private/PoseSearchResult.cpp b/Engine/Plugins/Experimental/Animation/PoseSearch/Source/Runtime/Private/PoseSearchResult.cpp
--- a/Engine/Plugins/Experimental/Animation/PoseSearch/Source/Runtime/Private/PoseSearchResult.cpp
+++ b/Engine/Plugins/Experimental/Animation/PoseSearch/Source/Runtime/Private/PoseSearchResult.cpp
@@ -9,17 +9,6 @@
 namespace UE::PoseSearch
 {
 
-//////////////////////////////////////////////////////////////////////////
-// FFeatureVectorBuilder
-FFeatureVectorBuilder::FFeatureVectorBuilder(const UPoseSearchSchema* InSchema)
-{
-	check(InSchema && InSchema->IsValid());
-	SchemaPtr = InSchema;
-	Values.SetNumZeroed(SchemaPtr->SchemaCardinality);
-}
-
-//////////////////////////////////////////////////////////////////////////
-// FSearchResult
 void FSearchResult::Update(float NewAssetTime)
 {
 	if (!IsValid())
diff --git a/Engine/Plugins/Experimental/Animation/PoseSearch/Source/Runtime/Private/PoseSearchSchema.cpp b/Engine/Plugins/Experimental/Animation/PoseSearch/Source/Runtime/Private/PoseSearchSchema.cpp
--- a/Engine/Plugins/Experimental/Animation/PoseSearch/Source/Runtime/Private/PoseSearchSchema.cpp
+++ b/Engine/Plugins/Experimental/Animation/PoseSearch/Source/Runtime/Private/PoseSearchSchema.cpp
@@ -2,6 +2,7 @@
 
 #include "PoseSearch/PoseSearchSchema.h"
 #include "AnimationRuntime.h"
+#include "PoseSearch/PoseSearchContext.h"
 #include "PoseSearch/PoseSearchDefines.h"
 #include "PoseSearch/PoseSearchResult.h"
 #include "PoseSearchFeatureChannel_Padding.h"
@@ -26,17 +27,18 @@ void UPoseSearchSchema::AddTemporaryChannel(UPoseSearchFeatureChannel* Temporary
 	FinalizedChannels.Add(TemporaryChannel);
 }
 
-void UPoseSearchSchema::BuildQuery(UE::PoseSearch::FSearchContext& SearchContext, UE::PoseSearch::FFeatureVectorBuilder& InOutQuery) const
+TConstArrayView<float> UPoseSearchSchema::BuildQuery(UE::PoseSearch::FSearchContext& SearchContext) const
 {
 	QUICK_SCOPE_CYCLE_COUNTER(STAT_PoseSearch_BuildQuery);
 
-	check(InOutQuery.GetSchema() == this);
-	check(InOutQuery.GetValues().Num() == SchemaCardinality);
+	SearchContext.AddNewFeatureVectorBuilder(this);
 
 	for (const TObjectPtr<UPoseSearchFeatureChannel>& ChannelPtr : GetChannels())
 	{
-		ChannelPtr->BuildQuery(SearchContext, InOutQuery);
+		ChannelPtr->BuildQuery(SearchContext);
 	}
+
+	return SearchContext.EditFeatureVector();
 }
 
 FBoneIndexType UPoseSearchSchema::GetBoneIndexType(int8 SchemaBoneIdx) const
diff --git a/Engine/Plugins/Experimental/Animation/PoseSearch/Source/Runtime/Public/PoseSearch/AnimNode_MotionMatching.h b/Engine/Plugins/Experimental/Animation/PoseSearch/Source/Runtime/Public/PoseSearch/AnimNode_MotionMatching.h
--- a/Engine/Plugins/Experimental/Animation/PoseSearch/Source/Runtime/Public/PoseSearch/AnimNode_MotionMatching.h
+++ b/Engine/Plugins/Experimental/Animation/PoseSearch/Source/Runtime/Public/PoseSearch/AnimNode_MotionMatching.h
@@ -99,6 +99,12 @@ private:
 	UPROPERTY(EditAnywhere, Category = Settings, meta = (PinHiddenByDefault))
 	bool bShouldSearch = true;
 
+	// If set to true, the search of multiple databases with different schemas will try to share pose features data calculated during query build
+	// the idea is to be able to share as much as possible the continuing pose features vector across different schemas (and potentially improve performances)
+	// defaulted to false to preserve behavior backward compatibility
+	UPROPERTY(EditAnywhere, Category = Settings, meta = (PinHiddenByDefault))
+	bool bShouldUseCachedChannelData = false;
+	
 	// blend time over which the yaw from the animation is distributed across the trajectory samples (negative values implies yaw from the animation is constant over the entire trajectory, so the trajectory will not try to recover towards the capsule orientation)
 	UPROPERTY(EditAnywhere, Category = RootMotion, meta = (PinHiddenByDefault))
 	float YawFromAnimationTrajectoryBlendTime = 0.1f;
diff --git a/Engine/Plugins/Experimental/Animation/PoseSearch/Source/Runtime/Public/PoseSearch/PoseSearchContext.h b/Engine/Plugins/Experimental/Animation/PoseSearch/Source/Runtime/Public/PoseSearch/PoseSearchContext.h
--- a/Engine/Plugins/Experimental/Animation/PoseSearch/Source/Runtime/Public/PoseSearch/PoseSearchContext.h
+++ b/Engine/Plugins/Experimental/Animation/PoseSearch/Source/Runtime/Public/PoseSearch/PoseSearchContext.h
@@ -101,7 +101,7 @@ struct FCachedTransforms
 	}
 
 private:
-	TArray<FCachedTransform<FTransformType>, TInlineAllocator<64>> CachedTransforms;
+	TArray<FCachedTransform<FTransformType>, TInlineAllocator<64, TMemStackAllocator<>>> CachedTransforms;
 };
 
 #if ENABLE_DRAW_DEBUG
@@ -141,38 +141,63 @@ private:
 
 #endif // ENABLE_DRAW_DEBUG
 
+// float buffer of features according to a UPoseSearchSchema layout. Used to build search queries at runtime
+struct FCachedQuery
+{
+public:
+	explicit FCachedQuery(const UPoseSearchSchema* InSchema);
+	const UPoseSearchSchema* GetSchema() const { return Schema; }
+	TArrayView<float> EditValues() { return Values; }
+	TConstArrayView<float> GetValues() const { return Values; }
+
+private:
+	TStackAlignedArray<float> Values;
+	
+	// no need for a TWeakObjectPtr since it doesn't persist across multiple frames (same lifespan as FSearchContext)
+	const UPoseSearchSchema* Schema;
+};
+
+// CachedChannels uses hashed unique identifiers to determine channels that can share feature vector data during the building of the query
+struct FCachedChannel
+{
+	// no need for a TWeakObjectPtr since it doesn't persist across multiple frames (same lifespan as FSearchContext)
+	const UPoseSearchFeatureChannel* Channel = nullptr;
+
+	// index of the associated query in FSearchContext::CachedQueries
+	int32 CachedQueryIndex = INDEX_NONE;
+};
+
 struct POSESEARCH_API FSearchContext
 {
 	FSearchContext(const UAnimInstance* InAnimInstance, const IPoseHistory* InHistory, TConstArrayView<const UAnimationAsset*> InAnimationsToConsider = TConstArrayView<const UAnimationAsset*>(),
 		const FPoseSearchQueryTrajectory* InTrajectory = nullptr, float InDesiredPermutationTimeOffset = 0.f, const FPoseIndicesHistory* InPoseIndicesHistory = nullptr,
-		const FSearchResult& InCurrentResult = FSearchResult(), const FFloatInterval& InPoseJumpThresholdTime = FFloatInterval(0.f, 0.f));
+		const FSearchResult& InCurrentResult = FSearchResult(), const FFloatInterval& InPoseJumpThresholdTime = FFloatInterval(0.f, 0.f), bool bInUseCachedChannelData = false);
 
 	// Returns the rotation of the bone Schema.BoneReferences[SchemaSampleBoneIdx] at an offset time of SampleTimeOffset relative to the
 	// transform of the bone Schema.BoneReferences[SchemaOriginBoneIdx] at an offset time of time OriginTimeOffset 
 	// Times will be processed by GetPermutationTimeOffsets(PermutationTimeType, ...)
-	FQuat GetSampleRotation(float SampleTimeOffset, float OriginTimeOffset, const UPoseSearchSchema* Schema, int8 SchemaSampleBoneIdx = RootSchemaBoneIdx, int8 SchemaOriginBoneIdx = RootSchemaBoneIdx, EPermutationTimeType PermutationTimeType = EPermutationTimeType::UseSampleTime, const FQuat* SampleBoneRotationWorldOverride = nullptr);
+	FQuat GetSampleRotation(float SampleTimeOffset, float OriginTimeOffset, int8 SchemaSampleBoneIdx = RootSchemaBoneIdx, int8 SchemaOriginBoneIdx = RootSchemaBoneIdx, EPermutationTimeType PermutationTimeType = EPermutationTimeType::UseSampleTime, const FQuat* SampleBoneRotationWorldOverride = nullptr);
 	
 	// Returns the position of the bone Schema.BoneReferences[SchemaSampleBoneIdx] at an offset time of SampleTimeOffset relative to the
 	// transform of the bone Schema.BoneReferences[SchemaOriginBoneIdx] at an offset time of time OriginTimeOffset 
 	// Times will be processed by GetPermutationTimeOffsets(PermutationTimeType, ...)
-	FVector GetSamplePosition(float SampleTimeOffset, float OriginTimeOffset, const UPoseSearchSchema* Schema, int8 SchemaSampleBoneIdx = RootSchemaBoneIdx, int8 SchemaOriginBoneIdx = RootSchemaBoneIdx, EPermutationTimeType PermutationTimeType = EPermutationTimeType::UseSampleTime, const FVector* SampleBonePositionWorldOverride = nullptr);
+	FVector GetSamplePosition(float SampleTimeOffset, float OriginTimeOffset, int8 SchemaSampleBoneIdx = RootSchemaBoneIdx, int8 SchemaOriginBoneIdx = RootSchemaBoneIdx, EPermutationTimeType PermutationTimeType = EPermutationTimeType::UseSampleTime, const FVector* SampleBonePositionWorldOverride = nullptr);
 	
 	// Returns the delta velocity of the velocity of the bone Schema.BoneReferences[SchemaSampleBoneIdx] at an offset time of SampleTimeOffset minus
 	// the velocity of the bone Schema.BoneReferences[SchemaOriginBoneIdx] at an offset time of time OriginTimeOffset 
 	// Times will be processed by GetPermutationTimeOffsets(PermutationTimeType, ...)
 	// if bUseCharacterSpaceVelocities is true, velocities will be computed in root bone space, rather than world space
-	FVector GetSampleVelocity(float SampleTimeOffset, float OriginTimeOffset, const UPoseSearchSchema* Schema, int8 SchemaSampleBoneIdx = RootSchemaBoneIdx, int8 SchemaOriginBoneIdx = RootSchemaBoneIdx, bool bUseCharacterSpaceVelocities = true, EPermutationTimeType PermutationTimeType = EPermutationTimeType::UseSampleTime, const FVector* SampleBoneVelocityWorldOverride = nullptr);
-
-	void ClearCachedEntries();
+	FVector GetSampleVelocity(float SampleTimeOffset, float OriginTimeOffset, int8 SchemaSampleBoneIdx = RootSchemaBoneIdx, int8 SchemaOriginBoneIdx = RootSchemaBoneIdx, bool bUseCharacterSpaceVelocities = true, EPermutationTimeType PermutationTimeType = EPermutationTimeType::UseSampleTime, const FVector* SampleBoneVelocityWorldOverride = nullptr);
 
 	void ResetCurrentBestCost();
 	void UpdateCurrentBestCost(const FPoseSearchCost& PoseSearchCost);
 	float GetCurrentBestTotalCost() const { return CurrentBestTotalCost; }
 
-	const FFeatureVectorBuilder& GetOrBuildQuery(const UPoseSearchSchema* Schema);
-	const FFeatureVectorBuilder* GetCachedQuery(const UPoseSearchSchema* Schema) const;
+	TConstArrayView<float> GetOrBuildQuery(const UPoseSearchSchema* Schema);
+	TConstArrayView<float> GetCachedQuery(const UPoseSearchSchema* Schema) const;
 
 	bool IsCurrentResultFromDatabase(const UPoseSearchDatabase* Database) const;
+	bool CanUseCurrentResult() const;
 
 	TConstArrayView<float> GetCurrentResultPoseVector() const { return CurrentResultPoseVector; }
 
@@ -188,7 +213,7 @@ struct POSESEARCH_API FSearchContext
 	TConstArrayView<const UAnimationAsset*> GetAnimationsToConsider() const { return AnimationsToConsider; }
 	
 	// returns the world space transform of the bone SchemaBoneIdx at time SampleTime
-	FTransform GetWorldBoneTransformAtTime(float SampleTime, const UPoseSearchSchema* Schema = nullptr, int8 SchemaBoneIdx = RootSchemaBoneIdx);
+	FTransform GetWorldBoneTransformAtTime(float SampleTime, int8 SchemaBoneIdx = RootSchemaBoneIdx);
 	
 #if WITH_EDITOR
 	void SetAsyncBuildIndexInProgress() { bAsyncBuildIndexInProgress = true; }
@@ -196,11 +221,18 @@ struct POSESEARCH_API FSearchContext
 	bool IsAsyncBuildIndexInProgress() const { return bAsyncBuildIndexInProgress; }
 #endif // WITH_EDITOR
 
+	void AddNewFeatureVectorBuilder(const UPoseSearchSchema* Schema) { CachedQueries.Emplace(Schema); }
+	TArrayView<float> EditFeatureVector();
+	
+	const UPoseSearchFeatureChannel* GetCachedChannelData(uint32 ChannelUniqueIdentifier, const UPoseSearchFeatureChannel* Channel, TConstArrayView<float>& CachedChannelData);
+	bool IsUseCachedChannelData() const { return bUseCachedChannelData; }
+	void SetUseCachedChannelData(bool bInUseCachedChannelData) { bUseCachedChannelData = bInUseCachedChannelData; }
+
 private:
-	FVector GetSamplePositionInternal(float SampleTime, float OriginTime, const UPoseSearchSchema* Schema, int8 SchemaSampleBoneIdx = RootSchemaBoneIdx, int8 SchemaOriginBoneIdx = RootSchemaBoneIdx, const FVector* SampleBonePositionWorldOverride = nullptr);
-	FQuat GetSampleRotationInternal(float SampleTime, float OriginTime, const UPoseSearchSchema* Schema, int8 SchemaSampleBoneIdx = RootSchemaBoneIdx, int8 SchemaOriginBoneIdx = RootSchemaBoneIdx, const FQuat* SampleBoneRotationWorldOverride = nullptr);
+	FVector GetSamplePositionInternal(float SampleTime, float OriginTime, int8 SchemaSampleBoneIdx = RootSchemaBoneIdx, int8 SchemaOriginBoneIdx = RootSchemaBoneIdx, const FVector* SampleBonePositionWorldOverride = nullptr);
+	FQuat GetSampleRotationInternal(float SampleTime, float OriginTime, int8 SchemaSampleBoneIdx = RootSchemaBoneIdx, int8 SchemaOriginBoneIdx = RootSchemaBoneIdx, const FQuat* SampleBoneRotationWorldOverride = nullptr);
 	FTransform GetWorldRootBoneTransformAtTime(float SampleTime) const;
-
+	
 	const UAnimInstance* AnimInstance = nullptr;
 	const IPoseHistory* History = nullptr;
 
@@ -213,13 +245,21 @@ private:
 	const FPoseIndicesHistory* PoseIndicesHistory = nullptr;
 	const FSearchResult& CurrentResult;
 	const FFloatInterval& PoseJumpThresholdTime;
+	bool bUseCachedChannelData = false;
 
 	TConstArrayView<float> CurrentResultPoseVector;
+
+	// @todo: use a 16 bytes aligned TInlineAllocator with overflow TMemStackAllocator
 	TStackAlignedArray<float> CurrentResultPoseVectorData;
 
 	// transforms cached in world space
 	FCachedTransforms<FTransform> CachedTransforms;
-	TArray<FFeatureVectorBuilder, TInlineAllocator<PreallocatedCachedQueriesNum>> CachedQueries;
+
+	TArray<FCachedQuery, TInlineAllocator<PreallocatedCachedQueriesNum, TMemStackAllocator<>>> CachedQueries;
+
+	// @todo add an overflow TMemStackAllocator
+	// mapping channel unique identifier (hash) to FCachedChannel
+	TMap<uint32, FCachedChannel, TInlineSetAllocator<PreallocatedCachedChannelDataNum>> CachedChannels;
 
 	float CurrentBestTotalCost = MAX_flt;
 	
diff --git a/Engine/Plugins/Experimental/Animation/PoseSearch/Source/Runtime/Public/PoseSearch/PoseSearchDefines.h b/Engine/Plugins/Experimental/Animation/PoseSearch/Source/Runtime/Public/PoseSearch/PoseSearchDefines.h
--- a/Engine/Plugins/Experimental/Animation/PoseSearch/Source/Runtime/Public/PoseSearch/PoseSearchDefines.h
+++ b/Engine/Plugins/Experimental/Animation/PoseSearch/Source/Runtime/Public/PoseSearch/PoseSearchDefines.h
@@ -28,6 +28,7 @@ static constexpr FBoneIndexType RootBoneIndexType = 0;
 static constexpr FBoneIndexType ComponentSpaceIndexType = FBoneIndexType(-1);
 
 static constexpr int32 PreallocatedCachedQueriesNum = 8;
+static constexpr int32 PreallocatedCachedChannelDataNum = 64;
 static constexpr float FiniteDelta = 1 / 60.0f; // Time delta used for computing pose derivatives
 static constexpr int32 MaxNumberOfCollectedPoseCandidatesPerDatabase = 200;
 
diff --git a/Engine/Plugins/Experimental/Animation/PoseSearch/Source/Runtime/Public/PoseSearch/PoseSearchFeatureChannel.h b/Engine/Plugins/Experimental/Animation/PoseSearch/Source/Runtime/Public/PoseSearch/PoseSearchFeatureChannel.h
--- a/Engine/Plugins/Experimental/Animation/PoseSearch/Source/Runtime/Public/PoseSearch/PoseSearchFeatureChannel.h
+++ b/Engine/Plugins/Experimental/Animation/PoseSearch/Source/Runtime/Public/PoseSearch/PoseSearchFeatureChannel.h
@@ -55,7 +55,6 @@ namespace UE::PoseSearch
 {
 
 struct FDebugDrawParams;
-struct FFeatureVectorBuilder;
 struct FSearchContext;
 struct FPoseMetadata;
 
@@ -124,7 +123,7 @@ public:
 	virtual void Finalize(UPoseSearchSchema* Schema) PURE_VIRTUAL(UPoseSearchFeatureChannel::Finalize, );
 	
 	// Called at runtime to add this channel's data to the query pose vector
-	virtual void BuildQuery(UE::PoseSearch::FSearchContext& SearchContext, UE::PoseSearch::FFeatureVectorBuilder& InOutQuery) const PURE_VIRTUAL(UPoseSearchFeatureChannel::BuildQuery, );
+	virtual void BuildQuery(UE::PoseSearch::FSearchContext& SearchContext) const PURE_VIRTUAL(UPoseSearchFeatureChannel::BuildQuery, );
 
 	// UPoseSearchFeatureChannels can hold sub channels
 	virtual TArrayView<TObjectPtr<UPoseSearchFeatureChannel>> GetSubChannels() { return TArrayView<TObjectPtr<UPoseSearchFeatureChannel>>(); }
diff --git a/Engine/Plugins/Experimental/Animation/PoseSearch/Source/Runtime/Public/PoseSearch/PoseSearchLibrary.h b/Engine/Plugins/Experimental/Animation/PoseSearch/Source/Runtime/Public/PoseSearch/PoseSearchLibrary.h
--- a/Engine/Plugins/Experimental/Animation/PoseSearch/Source/Runtime/Public/PoseSearch/PoseSearchLibrary.h
+++ b/Engine/Plugins/Experimental/Animation/PoseSearch/Source/Runtime/Public/PoseSearch/PoseSearchLibrary.h
@@ -147,10 +147,11 @@ public:
 		float SearchThrottleTime,
 		const FFloatInterval& PlayRate,
 		FMotionMatchingState& InOutMotionMatchingState,
-		float YawFromAnimationBlendRate,
-		float YawFromAnimationTrajectoryBlendTime,
+		float YawFromAnimationBlendRate = -1.f,
+		float YawFromAnimationTrajectoryBlendTime = 0.1f,
 		EPoseSearchInterruptMode InterruptMode = EPoseSearchInterruptMode::DoNotInterrupt,
 		bool bShouldSearch = true,
+		bool bShouldUseCachedChannelData = true,
 		bool bDebugDrawQuery = false,
 		bool bDebugDrawCurResult = false,
 		bool bDebugDrawPoseHistory = false);
diff --git a/Engine/Plugins/Experimental/Animation/PoseSearch/Source/Runtime/Public/PoseSearch/PoseSearchResult.h b/Engine/Plugins/Experimental/Animation/PoseSearch/Source/Runtime/Public/PoseSearch/PoseSearchResult.h
--- a/Engine/Plugins/Experimental/Animation/PoseSearch/Source/Runtime/Public/PoseSearch/PoseSearchResult.h
+++ b/Engine/Plugins/Experimental/Animation/PoseSearch/Source/Runtime/Public/PoseSearch/PoseSearchResult.h
@@ -13,24 +13,6 @@ namespace UE::PoseSearch
 {
 struct FSearchIndexAsset;
 
-/**
-* float buffer of features according to a UPoseSearchSchema layout.
-* FFeatureVectorBuilder is used to build search queries at runtime and for adding samples during search index construction.
-*/
-struct FFeatureVectorBuilder
-{
-public:
-	explicit FFeatureVectorBuilder(const UPoseSearchSchema* Schema);
-	const UPoseSearchSchema* GetSchema() const { return SchemaPtr.Get(); }
-
-	TArrayView<float> EditValues() { return Values; }
-	TConstArrayView<float> GetValues() const { return Values; }
-
-private:
-	TStackAlignedArray<float> Values;
-	TWeakObjectPtr<const UPoseSearchSchema> SchemaPtr;
-};
-	
 struct FSearchResult
 {
 	// best cost of the currently selected PoseIdx (it could be equal to ContinuingPoseCost)
diff --git a/Engine/Plugins/Experimental/Animation/PoseSearch/Source/Runtime/Public/PoseSearch/PoseSearchSchema.h b/Engine/Plugins/Experimental/Animation/PoseSearch/Source/Runtime/Public/PoseSearch/PoseSearchSchema.h
--- a/Engine/Plugins/Experimental/Animation/PoseSearch/Source/Runtime/Public/PoseSearch/PoseSearchSchema.h
+++ b/Engine/Plugins/Experimental/Animation/PoseSearch/Source/Runtime/Public/PoseSearch/PoseSearchSchema.h
@@ -130,7 +130,7 @@ public:
 	virtual void PostEditChangeProperty(struct FPropertyChangedEvent& PropertyChangedEvent) override;
 #endif
 
-	void BuildQuery(UE::PoseSearch::FSearchContext& SearchContext, UE::PoseSearch::FFeatureVectorBuilder& InOutQuery) const;
+	TConstArrayView<float> BuildQuery(UE::PoseSearch::FSearchContext& SearchContext) const;
 
 	FBoneIndexType GetBoneIndexType(int8 SchemaBoneIdx) const;
 
