motion matching - fix for UPoseSearchDatabase::SearchContinuingPose UAnimNotifyState_PoseSearchOverrideContinuingPoseCostBias extraction for blend spaces

diff --git a/Engine/Plugins/Experimental/Animation/PoseSearch/Source/Runtime/Private/PoseSearchDatabase.cpp b/Engine/Plugins/Experimental/Animation/PoseSearch/Source/Runtime/Private/PoseSearchDatabase.cpp
--- a/Engine/Plugins/Experimental/Animation/PoseSearch/Source/Runtime/Private/PoseSearchDatabase.cpp
+++ b/Engine/Plugins/Experimental/Animation/PoseSearch/Source/Runtime/Private/PoseSearchDatabase.cpp
@@ -973,7 +973,7 @@ UE::PoseSearch::FSearchResult UPoseSearchDatabase::SearchContinuingPose(UE::Pose
 	const FPoseSearchDatabaseAnimationAssetBase* DatabaseAnimationAssetBase = GetAnimationAssetStruct(SearchIndexAsset).GetPtr<FPoseSearchDatabaseAnimationAssetBase>();
 	check(DatabaseAnimationAssetBase);
 	const FAnimationAssetSampler SequenceBaseSampler(DatabaseAnimationAssetBase->GetAnimationAsset(), SearchIndexAsset.GetBlendParameters());
-	const float SampleTime = GetNormalizedAssetTime(PoseIdx);
+	const float SampleTime = GetRealAssetTime(PoseIdx);
 
 	float UpdatedContinuingPoseCostBias = ContinuingPoseCostBias;
 	SequenceBaseSampler.ExtractPoseSearchNotifyStates(SampleTime, [&UpdatedContinuingPoseCostBias](const UAnimNotifyState_PoseSearchBase* PoseSearchNotify)
