motion matching - cleanup of UPoseSearchMeshComponent, FDebuggerViewModel, FDebuggerViewModel::FSkeleton

diff --git a/Engine/Plugins/Experimental/Animation/PoseSearch/Source/Editor/Private/PoseSearchDebuggerView.cpp b/Engine/Plugins/Experimental/Animation/PoseSearch/Source/Editor/Private/PoseSearchDebuggerView.cpp
--- a/Engine/Plugins/Experimental/Animation/PoseSearch/Source/Editor/Private/PoseSearchDebuggerView.cpp
+++ b/Engine/Plugins/Experimental/Animation/PoseSearch/Source/Editor/Private/PoseSearchDebuggerView.cpp
@@ -441,15 +441,6 @@ void SDebuggerView::OnPoseSelectionChanged(const UPoseSearchDatabase* Database,
 	{
 		DetailsView->Update(*State);
 	}
-	
-	if (DbPoseIdx == INDEX_NONE)
-	{
-		Model->ClearSelectedSkeleton();
-	}
-	else
-	{
-		Model->ShowSelectedSkeleton(Database, DbPoseIdx, Time);
-	}
 }
 
 FReply SDebuggerView::OnUpdateNodeSelection(int32 InSelectedNodeId)
diff --git a/Engine/Plugins/Experimental/Animation/PoseSearch/Source/Editor/Private/PoseSearchDebuggerViewModel.cpp b/Engine/Plugins/Experimental/Animation/PoseSearch/Source/Editor/Private/PoseSearchDebuggerViewModel.cpp
--- a/Engine/Plugins/Experimental/Animation/PoseSearch/Source/Editor/Private/PoseSearchDebuggerViewModel.cpp
+++ b/Engine/Plugins/Experimental/Animation/PoseSearch/Source/Editor/Private/PoseSearchDebuggerViewModel.cpp
@@ -20,20 +20,14 @@ namespace UE::PoseSearch
 FDebuggerViewModel::FDebuggerViewModel(uint64 InAnimInstanceId)
 	: AnimInstanceId(InAnimInstanceId)
 {
-	Skeletons.AddDefaulted(ESkeletonIndex::Num);
 }
 
 FDebuggerViewModel::~FDebuggerViewModel()
 {
-	for (FSkeleton& Skeleton : Skeletons)
+	if (DebugDrawActor != nullptr)
 	{
-		if (Skeleton.Actor != nullptr)
-		{
-			Skeleton.Actor->Destroy();
-		}
+		DebugDrawActor->Destroy();
 	}
-
-	Skeletons.Empty();
 }
 
 const FTraceMotionMatchingStateMessage* FDebuggerViewModel::GetMotionMatchingState() const
@@ -54,38 +48,6 @@ const UPoseSearchDatabase* FDebuggerViewModel::GetCurrentDatabase() const
 	return nullptr;
 }
 
-void FDebuggerViewModel::ShowSelectedSkeleton(const UPoseSearchDatabase* Database, int32 DbPoseIdx, float Time)
-{
-	UPoseSearchMeshComponent* Component = Skeletons[SelectedPose].Component.Get();
-	if (!Component)
-	{
-		return;
-	}
-
-	if (!FAsyncPoseSearchDatabasesManagement::RequestAsyncBuildIndex(Database, ERequestAsyncBuildFlag::ContinueRequest))
-	{
-		return;
-	}
-
-	const FSearchIndex& SearchIndex = Database->GetSearchIndex();
-	if (const FSearchIndexAsset* IndexAsset = SearchIndex.GetAssetForPoseSafe(DbPoseIdx))
-	{
-		Component->ResetToStart();
-		bSelecting = true;
-
-		Skeletons[SelectedPose].Time = Time;
-		Skeletons[SelectedPose].bMirrored = IndexAsset->IsMirrored();
-		Skeletons[SelectedPose].SourceDatabase = Database;
-		Skeletons[SelectedPose].AssetIdx = IndexAsset->GetSourceAssetIdx();
-		Skeletons[SelectedPose].BlendParameters = IndexAsset->GetBlendParameters();
-	}
-}
-
-void FDebuggerViewModel::ClearSelectedSkeleton()
-{
-	bSelecting = false;
-}
-
 const TArray<int32>* FDebuggerViewModel::GetNodeIds() const
 {
 	return &NodeIds;
@@ -103,22 +65,20 @@ const FTransform& FDebuggerViewModel::GetRootBoneTransform() const
 
 void FDebuggerViewModel::OnUpdate()
 {
-	if (!bSkeletonsInitialized)
+	if (!bSkeletonInitialized)
 	{
 		UWorld* World = RewindDebugger.Get()->GetWorldToVisualize();
-		for (FSkeleton& Skeleton : Skeletons)
-		{
-			FActorSpawnParameters ActorSpawnParameters;
-			ActorSpawnParameters.bHideFromSceneOutliner = false;
-			ActorSpawnParameters.ObjectFlags |= RF_Transient;
-			Skeleton.Actor = World->SpawnActor<AActor>(ActorSpawnParameters);
-			Skeleton.Actor->SetActorLabel(TEXT("PoseSearch"));
-			Skeleton.Component = NewObject<UPoseSearchMeshComponent>(Skeleton.Actor.Get());
-			Skeleton.Actor->AddInstanceComponent(Skeleton.Component.Get());
-			Skeleton.Component->RegisterComponentWithWorld(World);
-		}
+		FActorSpawnParameters ActorSpawnParameters;
+		ActorSpawnParameters.bHideFromSceneOutliner = false;
+		ActorSpawnParameters.ObjectFlags |= RF_Transient;
+		DebugDrawActor = World->SpawnActor<AActor>(ActorSpawnParameters);
+		DebugDrawActor->SetActorLabel(TEXT("PoseSearch"));
+		DebugDrawMeshComponent = NewObject<UPoseSearchMeshComponent>(DebugDrawActor.Get());
+		DebugDrawActor->AddInstanceComponent(DebugDrawMeshComponent.Get());
+		DebugDrawMeshComponent->RegisterComponentWithWorld(World);
+	
 		FWorldDelegates::OnWorldCleanup.AddRaw(this, &FDebuggerViewModel::OnWorldCleanup);
-		bSkeletonsInitialized = true;
+		bSkeletonInitialized = true;
 	}
 
 	UpdateFromTimeline();
@@ -142,59 +102,6 @@ void FDebuggerViewModel::OnUpdateNodeSelection(int32 InNodeId)
 			break;
 		}
 	}
-
-	if (const FTraceMotionMatchingStateMessage* State = GetMotionMatchingState())
-	{
-		if (const UPoseSearchDatabase* CurrentDatabase = State->GetCurrentDatabase())
-		{
-			if (FAsyncPoseSearchDatabasesManagement::RequestAsyncBuildIndex(CurrentDatabase, ERequestAsyncBuildFlag::ContinueRequest))
-			{
-				const FSearchIndex& CurrentSearchIndex = CurrentDatabase->GetSearchIndex();
-				int32 CurrentPoseIdx = State->GetCurrentDatabasePoseIndex();
-				if (const FSearchIndexAsset* IndexAsset = CurrentSearchIndex.GetAssetForPoseSafe(CurrentPoseIdx))
-				{
-					Skeletons[Asset].bMirrored = IndexAsset->IsMirrored();
-					Skeletons[Asset].SourceDatabase = CurrentDatabase;
-					Skeletons[Asset].AssetIdx = IndexAsset->GetSourceAssetIdx();
-					Skeletons[Asset].BlendParameters = IndexAsset->GetBlendParameters();
-				}
-			}
-		}
-	}
-}
-
-void FDebuggerViewModel::UpdatePoseSearchContext(UPoseSearchMeshComponent::FUpdateContext& InOutContext, const FSkeleton& Skeleton) const
-{
-	const FInstancedStruct* DatabaseAssetStruct = Skeleton.GetAnimationAsset();
-	if (DatabaseAssetStruct)
-	{
-		const FPoseSearchDatabaseAnimationAssetBase* DatabaseAsset = DatabaseAssetStruct->GetPtr<FPoseSearchDatabaseAnimationAssetBase>();
-		if (DatabaseAsset)
-		{
-			InOutContext.StartTime = Skeletons[SelectedPose].Time;
-			InOutContext.Time = Skeletons[SelectedPose].Time;
-			InOutContext.MirrorDataCache = Skeletons[SelectedPose].bMirrored ? &MirrorDataCache : nullptr;
-			InOutContext.bLoop = DatabaseAsset->IsLooping();
-		}
-
-		if (const FPoseSearchDatabaseSequence* DatabaseSequence = DatabaseAssetStruct->GetPtr<FPoseSearchDatabaseSequence>())
-		{
-			InOutContext.SequenceBase = DatabaseSequence->Sequence;
-		}
-		else if (const FPoseSearchDatabaseAnimComposite* DatabaseAnimComposite = DatabaseAssetStruct->GetPtr<FPoseSearchDatabaseAnimComposite>())
-		{
-			InOutContext.SequenceBase = DatabaseAnimComposite->AnimComposite;
-		}
-		else if (const FPoseSearchDatabaseBlendSpace* DatabaseBlendSpace = DatabaseAssetStruct->GetPtr<FPoseSearchDatabaseBlendSpace>())
-		{
-			InOutContext.BlendSpace = DatabaseBlendSpace->BlendSpace;
-			InOutContext.BlendParameters = Skeletons[SelectedPose].BlendParameters;
-		}
-		else
-		{
-			checkNoEntry();
-		}
-	}
 }
 
 void FDebuggerViewModel::UpdateFromTimeline()
@@ -251,19 +158,15 @@ void FDebuggerViewModel::UpdateFromTimeline()
 				return TraceServices::EEventEnumerate::Stop;
 			}
 
-			UPoseSearchMeshComponent* ActiveComponent = Skeletons[ActivePose].Component.Get();
-			UPoseSearchMeshComponent* SelectedComponent = Skeletons[SelectedPose].Component.Get();
-			UPoseSearchMeshComponent* AssetComponent = Skeletons[Asset].Component.Get();
+			UPoseSearchMeshComponent* PoseSearchMeshComponent = DebugDrawMeshComponent.Get();
 			USkeletalMesh* SkeletalMesh = TSoftObjectPtr<USkeletalMesh>(FSoftObjectPath(SkeletalMeshObjectInfo->PathName)).LoadSynchronous();
 			if (SkeletalMesh)
 			{
-				ActiveComponent->SetSkinnedAssetAndUpdate(SkeletalMesh, true);
-				SelectedComponent->SetSkinnedAssetAndUpdate(SkeletalMesh, true);
-				AssetComponent->SetSkinnedAssetAndUpdate(SkeletalMesh, true);
+				PoseSearchMeshComponent->SetSkinnedAssetAndUpdate(SkeletalMesh, true);
 			}
 			FTransform ComponentWorldTransform;
 			// Active skeleton is simply the traced bone transforms
-			TArray<FTransform>& ComponentSpaceTransforms = ActiveComponent->GetEditableComponentSpaceTransforms();
+			TArray<FTransform>& ComponentSpaceTransforms = PoseSearchMeshComponent->GetEditableComponentSpaceTransforms();
 			AnimationProvider->GetSkeletalMeshComponentSpacePose(PoseMessage, *SkeletalMeshInfo, ComponentWorldTransform, ComponentSpaceTransforms);
 
 			check(ComponentWorldTransform.Equals(PoseMessage.ComponentToWorld));
@@ -277,12 +180,7 @@ void FDebuggerViewModel::UpdateFromTimeline()
 				RootBoneWorldTransform = ComponentWorldTransform;
 			}
 
-			ActiveComponent->Initialize(ComponentWorldTransform);
-			ActiveComponent->SetDebugDrawColor(FLinearColor::Green);
-			SelectedComponent->SetDebugDrawColor(FLinearColor::Blue);
-			SelectedComponent->Initialize(ComponentWorldTransform);
-			AssetComponent->SetDebugDrawColor(FLinearColor::Red);
-			AssetComponent->Initialize(ComponentWorldTransform);
+			PoseSearchMeshComponent->Initialize(ComponentWorldTransform);
 
 			return TraceServices::EEventEnumerate::Stop;
 		});
@@ -291,49 +189,12 @@ void FDebuggerViewModel::UpdateFromTimeline()
 
 const USkinnedMeshComponent* FDebuggerViewModel::GetMeshComponent() const
 {
-	if (Skeletons.Num() > FDebuggerViewModel::Asset)
-	{
-		return Skeletons[FDebuggerViewModel::Asset].Component.Get();
-	}
-	return nullptr;
-}
-
-void FDebuggerViewModel::FillCompactPoseAndComponentRefRotations()
-{
-	bool bResetMirrorDataCache = true;
-	if (const UPoseSearchDatabase* Database = GetCurrentDatabase())
-	{
-		if (MirrorDataCache.GetMirrorDataTable() != Database->Schema->MirrorDataTable)
-		{
-			if (UPoseSearchMeshComponent* MeshComponent = Skeletons[ActivePose].Component.Get())
-			{
-				if (MeshComponent->RequiredBones.IsValid())
-				{
-					MirrorDataCache.Init(Database->Schema->MirrorDataTable, MeshComponent->RequiredBones);
-					bResetMirrorDataCache = false;
-				}
-			}
-		}
-	}
-
-	if (bResetMirrorDataCache)
-	{
-		MirrorDataCache.Reset();
-	}
+	return DebugDrawMeshComponent.Get();
 }
 
 void FDebuggerViewModel::OnWorldCleanup(UWorld* InWorld, bool bSessionEnded, bool bCleanupResources)
 {
-	bSkeletonsInitialized = false;
-}
-
-const FInstancedStruct* FDebuggerViewModel::FSkeleton::GetAnimationAsset() const
-{
-	if (SourceDatabase.IsValid() && SourceDatabase->AnimationAssets.IsValidIndex(AssetIdx))
-	{
-		return &SourceDatabase->GetAnimationAssetStruct(AssetIdx);
-	}
-	return nullptr;
+	bSkeletonInitialized = false;
 }
 
 } // namespace UE::PoseSearch
diff --git a/Engine/Plugins/Experimental/Animation/PoseSearch/Source/Editor/Private/PoseSearchDebuggerViewModel.h b/Engine/Plugins/Experimental/Animation/PoseSearch/Source/Editor/Private/PoseSearchDebuggerViewModel.h
--- a/Engine/Plugins/Experimental/Animation/PoseSearch/Source/Editor/Private/PoseSearchDebuggerViewModel.h
+++ b/Engine/Plugins/Experimental/Animation/PoseSearch/Source/Editor/Private/PoseSearchDebuggerViewModel.h
@@ -37,12 +37,6 @@ public:
 	/** Updates active motion matching state based on node selection */
 	void OnUpdateNodeSelection(int32 InNodeId);
 
-	/** Sets the selected pose skeleton*/
-	void ShowSelectedSkeleton(const UPoseSearchDatabase* Database, int32 DbPoseIdx, float Time);
-	
-	/** Clears the selected pose skeleton */
-	void ClearSelectedSkeleton();
-	
 	void SetVerbose(bool bVerbose) { bIsVerbose = bVerbose; }
 	bool IsVerbose() const { return bIsVerbose; }
 
@@ -59,41 +53,9 @@ public:
 	const USkinnedMeshComponent* GetMeshComponent() const;
 
 private:
-	/** Debug visualization skeleton actor */
-	struct FSkeleton
-	{
-		/** Actor object for the skeleton */
-		TWeakObjectPtr<AActor> Actor;
-
-		/** Derived skeletal mesh for setting the skeleton in the scene */
-		TWeakObjectPtr<UPoseSearchMeshComponent> Component;
-
-		/** Source database for this skeleton  */
-		TWeakObjectPtr<const UPoseSearchDatabase> SourceDatabase;
-
-		/** Source asset for this skeleton */
-		int32 AssetIdx = 0;
-
-		/** Time in the sequence this skeleton is accessing */
-		float Time = 0.0f;
-
-		/** If this asset should be mirrored */
-		bool bMirrored = false;
-
-		/** Blend Parameters if asset is a BlendSpace */
-		FVector BlendParameters = FVector::Zero();
-
-		const FInstancedStruct* GetAnimationAsset() const;
-	};
-
 	/** Update the list of states for this frame */
 	void UpdateFromTimeline();
 
-	/** Populates arrays used for mirroring the animation pose */
-	void FillCompactPoseAndComponentRefRotations();
-
-	void UpdatePoseSearchContext(UPoseSearchMeshComponent::FUpdateContext& InOutContext, const FSkeleton& Skeleton) const;
-	
 	/** List of all Node IDs associated with motion matching states */
 	TArray<int32> NodeIds;
 	
@@ -115,26 +77,14 @@ private:
 	/** Anim Instance associated with this debugger instance */
 	uint64 AnimInstanceId = 0;
 
-	FMirrorDataCache MirrorDataCache;
+	/** Actor object for the skeleton */
+	TWeakObjectPtr<AActor> DebugDrawActor;
 
-	/** Index for each type of skeleton we store for debug visualization */
-	enum ESkeletonIndex
-	{
-		ActivePose = 0,
-		SelectedPose,
-		Asset,
+	/** Derived skeletal mesh for setting the skeleton in the scene */
+	TWeakObjectPtr<UPoseSearchMeshComponent> DebugDrawMeshComponent;
 
-		Num
-	};
-
-	/** Skeleton container for each type */
-	TArray<FSkeleton, TFixedAllocator<ESkeletonIndex::Num>> Skeletons;
-
-	/** Whether the skeletons have been initialized for this world */
-	bool bSkeletonsInitialized = false;
-	
-	/** If we currently have a selection active in the view */
-	bool bSelecting = false;
+	/** Whether the skeleton have been initialized for this world */
+	bool bSkeletonInitialized = false;
 	
 	bool bIsVerbose = false;
 
diff --git a/Engine/Plugins/Experimental/Animation/PoseSearch/Source/Editor/Private/PoseSearchMeshComponent.cpp b/Engine/Plugins/Experimental/Animation/PoseSearch/Source/Editor/Private/PoseSearchMeshComponent.cpp
--- a/Engine/Plugins/Experimental/Animation/PoseSearch/Source/Editor/Private/PoseSearchMeshComponent.cpp
+++ b/Engine/Plugins/Experimental/Animation/PoseSearch/Source/Editor/Private/PoseSearchMeshComponent.cpp
@@ -1,13 +1,7 @@
 // Copyright Epic Games, Inc. All Rights Reserved.
 
 #include "PoseSearchMeshComponent.h"
-#include "AnimationRuntime.h"
-#include "Animation/AnimRootMotionProvider.h"
-#include "Animation/AnimSequence.h"
-#include "Animation/BlendSpace.h"
-#include "Animation/MirrorDataTable.h"
 #include "Engine/SkinnedAsset.h"
-#include "PoseSearch/PoseSearchMirrorDataCache.h"
 
 void UPoseSearchMeshComponent::Initialize(const FTransform& InComponentToWorld)
 {
@@ -24,7 +18,6 @@ void UPoseSearchMeshComponent::Initialize(const FTransform& InComponentToWorld)
 		}
 	}
 
-	StartingTransform = InComponentToWorld;
 	Refresh();
 }
 
@@ -41,116 +34,4 @@ void UPoseSearchMeshComponent::Refresh()
 	MarkRenderTransformDirty();
 	MarkRenderDynamicDataDirty();
 	MarkRenderStateDirty();
-}
-
-void UPoseSearchMeshComponent::ResetToStart()
-{
-	SetComponentToWorld(StartingTransform);
-	Refresh();
-}
-
-void UPoseSearchMeshComponent::UpdatePose(const FUpdateContext& UpdateContext)
-{
-	// @todo: reuse asset sampler here if needed at all
-	using namespace UE::PoseSearch;
-
-	FMemMark Mark(FMemStack::Get());
-
-	FCompactPose CompactPose;
-	CompactPose.SetBoneContainer(&RequiredBones);
-	FBlendedCurve Curve;
-	Curve.InitFrom(RequiredBones);
-	UE::Anim::FStackAttributeContainer Attributes;
-	FAnimationPoseData PoseData(CompactPose, Curve, Attributes);
-
-	if (UpdateContext.SequenceBase)
-	{
-		float AdvancedTime = UpdateContext.StartTime;
-
-		FAnimationRuntime::AdvanceTime(
-			UpdateContext.bLoop,
-			UpdateContext.Time - UpdateContext.StartTime,
-			AdvancedTime,
-			UpdateContext.SequenceBase->GetPlayLength());
-
-		FAnimExtractContext ExtractionCtx;
-		ExtractionCtx.CurrentTime = AdvancedTime;
-
-		UpdateContext.SequenceBase->GetAnimationPose(PoseData, ExtractionCtx);
-	}
-	else if (UpdateContext.BlendSpace)
-	{
-		TArray<FBlendSampleData> BlendSamples;
-		int32 TriangulationIndex = 0;
-		UpdateContext.BlendSpace->GetSamplesFromBlendInput(UpdateContext.BlendParameters, BlendSamples, TriangulationIndex, true);
-		
-		const float PlayLength = UpdateContext.BlendSpace->GetAnimationLengthFromSampleData(BlendSamples);
-		const float PreviousTime = UpdateContext.StartTime * PlayLength;
-		float CurrentTime = UpdateContext.Time * PlayLength;
-
-		float AdvancedTime = PreviousTime;
-		FAnimationRuntime::AdvanceTime(UpdateContext.bLoop, CurrentTime - PreviousTime, CurrentTime, PlayLength);
-
-		FDeltaTimeRecord DeltaTimeRecord;
-		DeltaTimeRecord.Set(PreviousTime, AdvancedTime - PreviousTime);
-		FAnimExtractContext ExtractionCtx(static_cast<double>(AdvancedTime), true, DeltaTimeRecord, UpdateContext.bLoop);
-
-		for (int32 BlendSampleIdex = 0; BlendSampleIdex < BlendSamples.Num(); BlendSampleIdex++)
-		{
-			FBlendSampleData& BlendSample = BlendSamples[BlendSampleIdex];
-			const float Scale = BlendSample.Animation && PlayLength > UE_KINDA_SMALL_NUMBER ? BlendSample.Animation->GetPlayLength() / PlayLength : 1.f;
-
-			FDeltaTimeRecord BlendSampleDeltaTimeRecord;
-			BlendSampleDeltaTimeRecord.Set(DeltaTimeRecord.GetPrevious() * Scale, DeltaTimeRecord.Delta * Scale);
-
-			BlendSample.DeltaTimeRecord = BlendSampleDeltaTimeRecord;
-			BlendSample.PreviousTime = PreviousTime * Scale;
-			BlendSample.Time = AdvancedTime * Scale;
-		}
-
-		UpdateContext.BlendSpace->GetAnimationPose(BlendSamples, ExtractionCtx, PoseData);
-	}
-	else
-	{
-		checkNoEntry();
-	}
-
-	LastRootMotionDelta = FTransform::Identity;
-
-	const UE::Anim::IAnimRootMotionProvider* RootMotionProvider = UE::Anim::IAnimRootMotionProvider::Get();
-
-	if (ensureMsgf(RootMotionProvider, TEXT("Could not get Root Motion Provider.")))
-	{
-		if (ensureMsgf(RootMotionProvider->HasRootMotion(Attributes), TEXT("Blend Space had no Root Motion Attribute.")))
-		{
-			RootMotionProvider->ExtractRootMotion(Attributes, LastRootMotionDelta);
-		}
-	}
-
-	if (UpdateContext.MirrorDataCache)
-	{
-		UpdateContext.MirrorDataCache->MirrorPose(CompactPose);
-	}
-
-	FCSPose<FCompactPose> ComponentSpacePose;
-	ComponentSpacePose.InitPose(MoveTemp(CompactPose));
-
-	for (const FBoneIndexType BoneIndex : RequiredBones.GetBoneIndicesArray())
-	{
-		const FTransform BoneTransform = ComponentSpacePose.GetComponentSpaceTransform(FCompactPoseBoneIndex(BoneIndex));
-		FSkeletonPoseBoneIndex SkeletonBoneIndex = RequiredBones.GetSkeletonPoseIndexFromCompactPoseIndex(FCompactPoseBoneIndex(BoneIndex));
-		FName BoneName = RequiredBones.GetSkeletonAsset()->GetReferenceSkeleton().GetBoneName(SkeletonBoneIndex.GetInt());
-		SetBoneTransformByName(BoneName, BoneTransform, EBoneSpaces::ComponentSpace);
-	}
-
-	if (UpdateContext.MirrorDataCache)
-	{
-		LastRootMotionDelta = UpdateContext.MirrorDataCache->MirrorTransform(LastRootMotionDelta);
-	}
-
-	const FTransform ComponentTransform = LastRootMotionDelta * StartingTransform;
-
-	SetComponentToWorld(ComponentTransform);
-	FillComponentSpaceTransforms();
-	Refresh();
 }
\ No newline at end of file
diff --git a/Engine/Plugins/Experimental/Animation/PoseSearch/Source/Editor/Private/PoseSearchMeshComponent.h b/Engine/Plugins/Experimental/Animation/PoseSearch/Source/Editor/Private/PoseSearchMeshComponent.h
--- a/Engine/Plugins/Experimental/Animation/PoseSearch/Source/Editor/Private/PoseSearchMeshComponent.h
+++ b/Engine/Plugins/Experimental/Animation/PoseSearch/Source/Editor/Private/PoseSearchMeshComponent.h
@@ -2,39 +2,15 @@
 
 #pragma once
 
-#include "BonePose.h"
 #include "Components/PoseableMeshComponent.h"
-#include "PoseSearch/PoseSearchIndex.h"
 #include "PoseSearchMeshComponent.generated.h"
 
-class UBlendSpace;
-
-namespace UE::PoseSearch
-{
-	struct FMirrorDataCache;
-};
-
 UCLASS()
 class UPoseSearchMeshComponent : public UPoseableMeshComponent
 {
 	GENERATED_BODY()
 public:
 
-	struct FUpdateContext
-	{
-		const UAnimSequenceBase* SequenceBase = nullptr;
-		const UBlendSpace* BlendSpace = nullptr;
-		float StartTime = 0.0f;
-		float Time = 0.0f;
-		bool bLoop = false;
-		FVector BlendParameters = FVector::Zero();
-		const UE::PoseSearch::FMirrorDataCache* MirrorDataCache = nullptr;
-	};
-
 	void Refresh();
-	void ResetToStart();
-	void UpdatePose(const FUpdateContext& UpdateContext);
 	void Initialize(const FTransform& InComponentToWorld);
-	FTransform StartingTransform = FTransform::Identity;
-	FTransform LastRootMotionDelta = FTransform::Identity;
 };
