motion matching - making AnimationWarping (IAnimRootMotionProvider) dependency optional in PoseSearch

diff --git a/Engine/Plugins/Experimental/Animation/PoseSearch/Source/Runtime/Private/PoseSearchAssetSampler.cpp b/Engine/Plugins/Experimental/Animation/PoseSearch/Source/Runtime/Private/PoseSearchAssetSampler.cpp
--- a/Engine/Plugins/Experimental/Animation/PoseSearch/Source/Runtime/Private/PoseSearchAssetSampler.cpp
+++ b/Engine/Plugins/Experimental/Animation/PoseSearch/Source/Runtime/Private/PoseSearchAssetSampler.cpp
@@ -1,12 +1,9 @@
 // Copyright Epic Games, Inc. All Rights Reserved.
 
 #include "PoseSearch/PoseSearchAssetSampler.h"
-#include "AnimationRuntime.h"
 #include "Animation/AnimMontage.h"
 #include "Animation/AnimSequence.h"
-#include "Animation/AnimRootMotionProvider.h"
 #include "Animation/BlendSpace.h"
-#include "Animation/MirrorDataTable.h"
 #include "PoseSearch/PoseSearchAnimNotifies.h"
 #include "PoseSearch/PoseSearchDefines.h"
 
@@ -150,74 +147,29 @@ static void ProcessRootTransform(const UBlendSpace* BlendSpace, const FVector& B
 	int32 RootTransformSamplingRate, bool bIsLoopable, TArray<FTransform>& AccumulatedRootTransform)
 {
 	// Pre-compute root motion
-	int32 NumRootSamples = FMath::Max(CachedPlayLength * RootTransformSamplingRate + 1, 1);
-	AccumulatedRootTransform.SetNumUninitialized(NumRootSamples);
+	const int32 NumRootSamples = FMath::Max(CachedPlayLength * RootTransformSamplingRate + 1, 1);
+	AccumulatedRootTransform.Init(FTransform::Identity, NumRootSamples);
 
 	TArray<FBlendSampleData> BlendSamples;
 	int32 TriangulationIndex = 0;
 	if (BlendSpace->GetSamplesFromBlendInput(BlendParameters, BlendSamples, TriangulationIndex, true))
 	{
-		FTransform RootMotionAccumulation = FTransform::Identity;
-
-		AccumulatedRootTransform[0] = RootMotionAccumulation;
-
 		for (int32 SampleIdx = 1; SampleIdx < NumRootSamples; ++SampleIdx)
 		{
-			float PreviousTime = float(SampleIdx - 1) / RootTransformSamplingRate;
-			float CurrentTime = float(SampleIdx - 0) / RootTransformSamplingRate;
-
-			FDeltaTimeRecord DeltaTimeRecord;
-			DeltaTimeRecord.Set(PreviousTime, CurrentTime - PreviousTime);
-			FAnimExtractContext ExtractionCtx(static_cast<double>(CurrentTime), true, DeltaTimeRecord, bIsLoopable);
-
+			FRootMotionMovementParams RootMotionMovementParams;
 			for (int32 BlendSampleIdex = 0; BlendSampleIdex < BlendSamples.Num(); BlendSampleIdex++)
 			{
 				FBlendSampleData& BlendSample = BlendSamples[BlendSampleIdex];
-				const float Scale = BlendSample.Animation && CachedPlayLength > UE_KINDA_SMALL_NUMBER ? BlendSample.Animation->GetPlayLength() / CachedPlayLength : 1.f;
-
-				FDeltaTimeRecord BlendSampleDeltaTimeRecord;
-				BlendSampleDeltaTimeRecord.Set(DeltaTimeRecord.GetPrevious() * Scale, DeltaTimeRecord.Delta * Scale);
-
-				BlendSample.DeltaTimeRecord = BlendSampleDeltaTimeRecord;
-				BlendSample.PreviousTime = PreviousTime * Scale;
-				BlendSample.Time = CurrentTime * Scale;
-			}
-
-			FMemMark Mark(FMemStack::Get());
-			FCompactPose Pose;
-			FBlendedCurve BlendedCurve;
-			UE::Anim::FStackAttributeContainer StackAttributeContainer;
-			FAnimationPoseData AnimPoseData(Pose, BlendedCurve, StackAttributeContainer);
-
-			// copying bone container to be thread safe, since it has mutable members
-			FBoneContainer BoneContainerCopy = BoneContainer;
+				// @todo: add support for synch marker
+				const float PlayLength = BlendSample.Animation->GetPlayLength();
+				const float DeltaTime = PlayLength / (NumRootSamples - 1);
+				const float PreviousTime = (SampleIdx - 1) * DeltaTime;
 
-			Pose.SetBoneContainer(&BoneContainerCopy);
-			BlendedCurve.InitFrom(BoneContainerCopy);
-
-			BlendSpace->GetAnimationPose(BlendSamples, ExtractionCtx, AnimPoseData);
-
-			const UE::Anim::IAnimRootMotionProvider* RootMotionProvider = UE::Anim::IAnimRootMotionProvider::Get();
-
-			if (RootMotionProvider)
-			{
-				if (RootMotionProvider->HasRootMotion(StackAttributeContainer))
-				{
-					FTransform RootMotionDelta;
-					RootMotionProvider->ExtractRootMotion(StackAttributeContainer, RootMotionDelta);
-					RootMotionAccumulation = RootMotionDelta * RootMotionAccumulation;
-				}
-				else
-				{
-					UE_LOG(LogPoseSearch, Error, TEXT("ProcessRootTransform: Blend Space '%s' has no Root Motion Attribute"), *BlendSpace->GetName());
-				}
-			}
-			else
-			{
-				UE_LOG(LogPoseSearch, Error, TEXT("ProcessRootTransform: Could not get Root Motion Provider for BlendSpace '%s'"), *BlendSpace->GetName());
+				const FTransform BlendSampleRootMotion = BlendSample.Animation->ExtractRootMotion(PreviousTime, DeltaTime, BlendSpace->bLoop);
+				RootMotionMovementParams.AccumulateWithBlend(BlendSampleRootMotion, BlendSample.GetClampedWeight());
 			}
 
-			AccumulatedRootTransform[SampleIdx] = RootMotionAccumulation;
+			AccumulatedRootTransform[SampleIdx] = RootMotionMovementParams.GetRootMotionTransform() * AccumulatedRootTransform[SampleIdx - 1];
 		}
 	}
 }
