motion matching
- fix for pose history retention policy
- updated pose history collector with meaningful default property values
- implemented motion trajectory orientation collection (that fixes pose history debug lines)

diff --git a/Engine/Plugins/Experimental/Animation/MotionTrajectory/Source/MotionTrajectory/Private/MotionTrajectoryLibrary.cpp b/Engine/Plugins/Experimental/Animation/MotionTrajectory/Source/MotionTrajectory/Private/MotionTrajectoryLibrary.cpp
--- a/Engine/Plugins/Experimental/Animation/MotionTrajectory/Source/MotionTrajectory/Private/MotionTrajectoryLibrary.cpp
+++ b/Engine/Plugins/Experimental/Animation/MotionTrajectory/Source/MotionTrajectory/Private/MotionTrajectoryLibrary.cpp
@@ -188,12 +188,13 @@ void FMotionTrajectoryLibrary::UpdateHistory_TransformHistory(FPoseSearchQueryTr
 		for (int32 Index = 0; Index < SamplingData.NumHistorySamples - 1; ++Index)
 		{
 			Trajectory.Samples[Index].AccumulatedSeconds = Trajectory.Samples[Index + 1].AccumulatedSeconds;
-			TranslationHistory[Index] = TranslationHistory[Index + 1];
-			TranslationHistory[Index] += CurrentTranslation;
+			TranslationHistory[Index] = TranslationHistory[Index + 1] + CurrentTranslation;
+			Trajectory.Samples[Index].Facing = Trajectory.Samples[Index + 1].Facing;
 		}
 
 		Trajectory.Samples[SamplingData.NumHistorySamples - 1].AccumulatedSeconds = 0.f;
 		TranslationHistory[SamplingData.NumHistorySamples - 1] = CurrentTranslation;
+		Trajectory.Samples[SamplingData.NumHistorySamples - 1].Facing = CharacterTrajectoryData.Facing;
 	}
 	else
 	{
@@ -208,9 +209,6 @@ void FMotionTrajectoryLibrary::UpdateHistory_TransformHistory(FPoseSearchQueryTr
 	{
 		Trajectory.Samples[Index].AccumulatedSeconds -= DeltaSeconds;
 		Trajectory.Samples[Index].Position = CharacterTrajectoryData.Position - TranslationHistory[Index];
-
-		// @todo: Handle facing. We currently don't use facing for history in any of our content. We will need the rotation intent of the character.
-		Trajectory.Samples[Index].Facing = FQuat::Identity;
 	}
 }
 
diff --git a/Engine/Plugins/Experimental/Animation/PoseSearch/Source/Runtime/Private/AnimNode_PoseSearchHistoryCollector.cpp b/Engine/Plugins/Experimental/Animation/PoseSearch/Source/Runtime/Private/AnimNode_PoseSearchHistoryCollector.cpp
--- a/Engine/Plugins/Experimental/Animation/PoseSearch/Source/Runtime/Private/AnimNode_PoseSearchHistoryCollector.cpp
+++ b/Engine/Plugins/Experimental/Animation/PoseSearch/Source/Runtime/Private/AnimNode_PoseSearchHistoryCollector.cpp
@@ -70,7 +70,7 @@ void FAnimNode_PoseSearchHistoryCollector_Base::CacheBones_AnyThread(const FAnim
 		}
 	}
 
-	PoseHistory.Init(PoseCount, PoseDuration, RequiredBones);
+	PoseHistory.Init(PoseCount, SamplingInterval, RequiredBones);
 
 	if (bInitializeWithRefPose)
 	{
@@ -86,6 +86,8 @@ void FAnimNode_PoseSearchHistoryCollector_Base::CacheBones_AnyThread(const FAnim
 
 void FAnimNode_PoseSearchHistoryCollector_Base::Update_AnyThread(const FAnimationUpdateContext& Context)
 {
+	GetEvaluateGraphExposedInputs().Execute(Context);
+
 	const bool bNeedsReset =
 		bResetOnBecomingRelevant &&
 		UpdateCounter.HasEverBeenUpdated() &&
diff --git a/Engine/Plugins/Experimental/Animation/PoseSearch/Source/Runtime/Private/PoseSearchHistory.cpp b/Engine/Plugins/Experimental/Animation/PoseSearch/Source/Runtime/Private/PoseSearchHistory.cpp
--- a/Engine/Plugins/Experimental/Animation/PoseSearch/Source/Runtime/Private/PoseSearchHistory.cpp
+++ b/Engine/Plugins/Experimental/Animation/PoseSearch/Source/Runtime/Private/PoseSearchHistory.cpp
@@ -132,10 +132,10 @@ FTransform FPoseHistoryEntry::GetComponentSpaceTransform(int32 Index) const
 
 //////////////////////////////////////////////////////////////////////////
 // FPoseHistory
-void FPoseHistory::Init(int32 InNumPoses, float InTimeHorizon, const TArray<FBoneIndexType>& RequiredBones)
+void FPoseHistory::Init(int32 InNumPoses, float InSamplingInterval, const TArray<FBoneIndexType>& RequiredBones)
 {
-	check(InNumPoses >= 2 && InTimeHorizon > UE_KINDA_SMALL_NUMBER);
-	TimeHorizon = InTimeHorizon;
+	check(InNumPoses >= 2 && InSamplingInterval > UE_KINDA_SMALL_NUMBER);
+	SamplingInterval = InSamplingInterval;
 
 	BoneToTransformMap.Reset();
 	if (!RequiredBones.IsEmpty())
@@ -297,32 +297,16 @@ void FPoseHistory::Update(float SecondsElapsed, FCSPose<FCompactPose>& Component
 		Entry.Time += SecondsElapsed;
 	}
 
-	const int32 EntriesMax = Entries.Max();
-	if (Entries.Num() != EntriesMax)
+	if (Entries.Num() != Entries.Max())
 	{
 		// Consume every pose until the queue is full
 		Entries.Emplace();
 	}
-	else
+	else if (Entries[Entries.Num() - 2].Time >= SamplingInterval)
 	{
-		// Exercise pose retention policy. We must guarantee there is always one additional pose
-		// beyond the time horizon so we can compute derivatives at the time horizon. We also
-		// want to evenly distribute poses across the entire history buffer so we only push additional
-		// poses when enough time has elapsed.
-
-		check(EntriesMax >= 2 && TimeHorizon > UE_KINDA_SMALL_NUMBER);
-		// Reserve one pose for computing derivatives at the time horizon
-		const float SampleInterval = TimeHorizon / (EntriesMax - 1);
-
-		bool bCanEvictOldest = Entries[1].Time >= TimeHorizon + SampleInterval;
-		bool bShouldPushNewest = Entries[Entries.Num() - 2].Time >= SampleInterval;
-
-		if (bCanEvictOldest && bShouldPushNewest)
-		{
-			FPoseHistoryEntry EntryTemp = MoveTemp(Entries.First());
-			Entries.PopFront();
-			Entries.Emplace(MoveTemp(EntryTemp));
-		}
+		FPoseHistoryEntry EntryTemp = MoveTemp(Entries.First());
+		Entries.PopFront();
+		Entries.Emplace(MoveTemp(EntryTemp));
 	}
 
 	// Regardless of the retention policy, we always update the most recent Entry
diff --git a/Engine/Plugins/Experimental/Animation/PoseSearch/Source/Runtime/Public/PoseSearch/AnimNode_PoseSearchHistoryCollector.h b/Engine/Plugins/Experimental/Animation/PoseSearch/Source/Runtime/Public/PoseSearch/AnimNode_PoseSearchHistoryCollector.h
--- a/Engine/Plugins/Experimental/Animation/PoseSearch/Source/Runtime/Public/PoseSearch/AnimNode_PoseSearchHistoryCollector.h
+++ b/Engine/Plugins/Experimental/Animation/PoseSearch/Source/Runtime/Public/PoseSearch/AnimNode_PoseSearchHistoryCollector.h
@@ -16,11 +16,11 @@ public:
 	
 	// The maximum amount of poses that can be stored
 	UPROPERTY(EditAnywhere, BlueprintReadWrite, Category = Settings, meta = (ClampMin="2"))
-	int32 PoseCount = 64;
+	int32 PoseCount = 2;
 	
-	// The time horizon for how long a pose will be stored in seconds
+	// how often in seconds poses are collected
 	UPROPERTY(EditAnywhere, BlueprintReadWrite, Category = Settings, meta = (ClampMin="0.001"))
-	float PoseDuration = 1.5f;
+	float SamplingInterval = UE::PoseSearch::FiniteDelta;
 
 	UPROPERTY(EditAnywhere, Category = Settings)
 	TArray<FBoneReference> CollectedBones;
diff --git a/Engine/Plugins/Experimental/Animation/PoseSearch/Source/Runtime/Public/PoseSearch/PoseSearchHistory.h b/Engine/Plugins/Experimental/Animation/PoseSearch/Source/Runtime/Public/PoseSearch/PoseSearchHistory.h
--- a/Engine/Plugins/Experimental/Animation/PoseSearch/Source/Runtime/Public/PoseSearch/PoseSearchHistory.h
+++ b/Engine/Plugins/Experimental/Animation/PoseSearch/Source/Runtime/Public/PoseSearch/PoseSearchHistory.h
@@ -59,9 +59,8 @@ typedef TArray<FPoseHistoryEntry> FPoseHistoryFutureEntries;
 
 struct FPoseHistory : public IPoseHistory
 {
-	void Init(int32 InNumPoses, float InTimeHorizon, const TArray<FBoneIndexType>& RequiredBones);
+	void Init(int32 InNumPoses, float InSamplingInterval, const TArray<FBoneIndexType>& RequiredBones);
 	void Update(float SecondsElapsed, FCSPose<FCompactPose>& ComponentSpacePose, bool bStoreScales);
-	float GetTimeHorizon() const { return TimeHorizon; }
 
 	const FBoneToTransformMap& GetBoneToTransformMap() const { return BoneToTransformMap; }
 	const FPoseHistoryEntries& GetEntries() const { return Entries; }
@@ -90,7 +89,7 @@ private:
 
 	// ring buffer of collected bones
 	FPoseHistoryEntries Entries;
-	float TimeHorizon = 0.f;
+	float SamplingInterval = 0.f;
 };
 
 struct FExtendedPoseHistory : public IPoseHistory
