motion matching
- consolidating config / schema naming
- consolidating database naming
- removing some Experimental metadata

diff --git a/Engine/Plugins/Experimental/Animation/PoseSearch/Source/Editor/Private/PoseSearchAssetDefinitions.h b/Engine/Plugins/Experimental/Animation/PoseSearch/Source/Editor/Private/PoseSearchAssetDefinitions.h
--- a/Engine/Plugins/Experimental/Animation/PoseSearch/Source/Editor/Private/PoseSearchAssetDefinitions.h
+++ b/Engine/Plugins/Experimental/Animation/PoseSearch/Source/Editor/Private/PoseSearchAssetDefinitions.h
@@ -25,7 +25,7 @@ public:
 	virtual TConstArrayView<FAssetCategoryPath> GetAssetCategories() const override { return UE::PoseSearch::GetAssetCategories(); }
 	virtual UThumbnailInfo* LoadThumbnailInfo(const FAssetData& InAssetData) const override { return UE::PoseSearch::LoadThumbnailInfo(InAssetData); }
 
-	virtual FText GetAssetDisplayName() const override { return NSLOCTEXT("PoseSearchAssetDefinition", "DisplayName_UPoseSearchDatabase", "Motion Database"); }
+	virtual FText GetAssetDisplayName() const override { return NSLOCTEXT("PoseSearchAssetDefinition", "DisplayName_UPoseSearchDatabase", "Pose Search Database"); }
 	virtual TSoftClassPtr<UObject> GetAssetClass() const override { return UPoseSearchDatabase::StaticClass(); }
 	virtual EAssetCommandResult OpenAssets(const FAssetOpenArgs& OpenArgs) const override;
 };
@@ -40,7 +40,7 @@ public:
 	virtual TConstArrayView<FAssetCategoryPath> GetAssetCategories() const override { return UE::PoseSearch::GetAssetCategories(); }
 	virtual UThumbnailInfo* LoadThumbnailInfo(const FAssetData& InAssetData) const override { return UE::PoseSearch::LoadThumbnailInfo(InAssetData); }
 
-	virtual FText GetAssetDisplayName() const override { return NSLOCTEXT("PoseSearchAssetDefinition", "DisplayName_UPoseSearchSchema", "Motion Database Config"); }
+	virtual FText GetAssetDisplayName() const override { return NSLOCTEXT("PoseSearchAssetDefinition", "DisplayName_UPoseSearchSchema", "Pose Search Schema"); }
 	virtual TSoftClassPtr<UObject> GetAssetClass() const override { return UPoseSearchSchema::StaticClass(); }
 };
 
@@ -54,7 +54,7 @@ public:
 	virtual TConstArrayView<FAssetCategoryPath> GetAssetCategories() const override { return UE::PoseSearch::GetAssetCategories(); }
 	virtual UThumbnailInfo* LoadThumbnailInfo(const FAssetData& InAssetData) const override { return UE::PoseSearch::LoadThumbnailInfo(InAssetData); }
 
-	virtual FText GetAssetDisplayName() const override { return NSLOCTEXT("PoseSearchAssetDefinition", "DisplayName_UPoseSearchNormalizationSet", "Normalization Set"); }
+	virtual FText GetAssetDisplayName() const override { return NSLOCTEXT("PoseSearchAssetDefinition", "DisplayName_UPoseSearchNormalizationSet", "Pose Search Normalization Set"); }
 	virtual TSoftClassPtr<UObject> GetAssetClass() const override { return UPoseSearchNormalizationSet::StaticClass(); }
 };
 
diff --git a/Engine/Plugins/Experimental/Animation/PoseSearch/Source/Editor/Private/PoseSearchDatabaseAssetTreeNode.cpp b/Engine/Plugins/Experimental/Animation/PoseSearch/Source/Editor/Private/PoseSearchDatabaseAssetTreeNode.cpp
--- a/Engine/Plugins/Experimental/Animation/PoseSearch/Source/Editor/Private/PoseSearchDatabaseAssetTreeNode.cpp
+++ b/Engine/Plugins/Experimental/Animation/PoseSearch/Source/Editor/Private/PoseSearchDatabaseAssetTreeNode.cpp
@@ -73,6 +73,6 @@ namespace UE::PoseSearch
 			}
 		}
 
-		return EPoseSearchMirrorOption::Invalid;
+		return EPoseSearchMirrorOption::UnmirroredOnly;
 	}
 }
diff --git a/Engine/Plugins/Experimental/Animation/PoseSearch/Source/Editor/Private/PoseSearchDatabaseEditorReflection.cpp b/Engine/Plugins/Experimental/Animation/PoseSearch/Source/Editor/Private/PoseSearchDatabaseEditorReflection.cpp
--- a/Engine/Plugins/Experimental/Animation/PoseSearch/Source/Editor/Private/PoseSearchDatabaseEditorReflection.cpp
+++ b/Engine/Plugins/Experimental/Animation/PoseSearch/Source/Editor/Private/PoseSearchDatabaseEditorReflection.cpp
@@ -131,11 +131,11 @@ void UPoseSearchDatabaseStatistics::Initialize(const UPoseSearchDatabase* PoseSe
 		SearchableFrames = NumOfSearchablePoses;
 		SearchableTime = FText::Format(TimeFormat, static_cast<double>(NumOfSearchablePoses) / SampleRate);
 
-		ConfigCardinality = PoseSearchDatabase->Schema->SchemaCardinality;
+		SchemaCardinality = PoseSearchDatabase->Schema->SchemaCardinality;
 
-		if (ConfigCardinality > 0)
+		if (SchemaCardinality > 0)
 		{
-			const int32 TotalAnimationFeatureVectors = SearchIndex.GetNumValuesVectors(ConfigCardinality);
+			const int32 TotalAnimationFeatureVectors = SearchIndex.GetNumValuesVectors(SchemaCardinality);
 			PrunedFrames = TotalAnimationPosesInFrames - TotalAnimationFeatureVectors;
 		}
 		else
diff --git a/Engine/Plugins/Experimental/Animation/PoseSearch/Source/Editor/Private/PoseSearchDatabaseFactory.cpp b/Engine/Plugins/Experimental/Animation/PoseSearch/Source/Editor/Private/PoseSearchDatabaseFactory.cpp
--- a/Engine/Plugins/Experimental/Animation/PoseSearch/Source/Editor/Private/PoseSearchDatabaseFactory.cpp
+++ b/Engine/Plugins/Experimental/Animation/PoseSearch/Source/Editor/Private/PoseSearchDatabaseFactory.cpp
@@ -73,7 +73,7 @@ UObject* UPoseSearchDatabaseFactory::FactoryCreateNew(UClass* Class, UObject* In
 
 FString UPoseSearchDatabaseFactory::GetDefaultNewAssetName() const
 {
-	return FString(TEXT("NewMotionDatabase"));
+	return FString(TEXT("NewPoseSearchDatabase"));
 }
 
 #undef LOCTEXT_NAMESPACE
\ No newline at end of file
diff --git a/Engine/Plugins/Experimental/Animation/PoseSearch/Source/Editor/Private/PoseSearchSchemaFactory.cpp b/Engine/Plugins/Experimental/Animation/PoseSearch/Source/Editor/Private/PoseSearchSchemaFactory.cpp
--- a/Engine/Plugins/Experimental/Animation/PoseSearch/Source/Editor/Private/PoseSearchSchemaFactory.cpp
+++ b/Engine/Plugins/Experimental/Animation/PoseSearch/Source/Editor/Private/PoseSearchSchemaFactory.cpp
@@ -73,7 +73,7 @@ UObject* UPoseSearchSchemaFactory::FactoryCreateNew(UClass* Class, UObject* InPa
 
 FString UPoseSearchSchemaFactory::GetDefaultNewAssetName() const
 {
-	return FString(TEXT("NewMotionDatabaseConfig"));
+	return FString(TEXT("NewPoseSearchSchema"));
 }
 
 #undef LOCTEXT_NAMESPACE
\ No newline at end of file
diff --git a/Engine/Plugins/Experimental/Animation/PoseSearch/Source/Editor/Public/AnimGraphNode_MotionMatching.h b/Engine/Plugins/Experimental/Animation/PoseSearch/Source/Editor/Public/AnimGraphNode_MotionMatching.h
--- a/Engine/Plugins/Experimental/Animation/PoseSearch/Source/Editor/Public/AnimGraphNode_MotionMatching.h
+++ b/Engine/Plugins/Experimental/Animation/PoseSearch/Source/Editor/Public/AnimGraphNode_MotionMatching.h
@@ -10,7 +10,7 @@
 #include "AnimGraphNode_MotionMatching.generated.h"
 
 
-UCLASS(MinimalAPI, Experimental)
+UCLASS(MinimalAPI)
 class UAnimGraphNode_MotionMatching : public UAnimGraphNode_BlendStack_Base
 {
 	GENERATED_BODY()
diff --git a/Engine/Plugins/Experimental/Animation/PoseSearch/Source/Editor/Public/AnimGraphNode_PoseSearchHistoryCollector.h b/Engine/Plugins/Experimental/Animation/PoseSearch/Source/Editor/Public/AnimGraphNode_PoseSearchHistoryCollector.h
--- a/Engine/Plugins/Experimental/Animation/PoseSearch/Source/Editor/Public/AnimGraphNode_PoseSearchHistoryCollector.h
+++ b/Engine/Plugins/Experimental/Animation/PoseSearch/Source/Editor/Public/AnimGraphNode_PoseSearchHistoryCollector.h
@@ -8,7 +8,7 @@
 #include "UObject/ObjectMacros.h"
 #include "AnimGraphNode_PoseSearchHistoryCollector.generated.h"
 
-UCLASS(Abstract, BlueprintType, Experimental)
+UCLASS(Abstract, BlueprintType)
 class UAnimGraphNode_PoseSearchHistoryCollector_Base : public UAnimGraphNode_Base
 {
 	GENERATED_BODY()
@@ -17,7 +17,7 @@ class UAnimGraphNode_PoseSearchHistoryCollector_Base : public UAnimGraphNode_Bas
 	virtual FText GetMenuCategory() const override;
 };
 
-UCLASS(MinimalAPI, Experimental)
+UCLASS(MinimalAPI)
 class UAnimGraphNode_PoseSearchHistoryCollector : public UAnimGraphNode_PoseSearchHistoryCollector_Base
 {
 	GENERATED_BODY()
@@ -28,7 +28,7 @@ class UAnimGraphNode_PoseSearchHistoryCollector : public UAnimGraphNode_PoseSear
 	virtual FText GetNodeTitle(ENodeTitleType::Type TitleType) const override;
 };
 
-UCLASS(MinimalAPI, Experimental)
+UCLASS(MinimalAPI)
 class UAnimGraphNode_PoseSearchComponentSpaceHistoryCollector : public UAnimGraphNode_PoseSearchHistoryCollector_Base
 {
 	GENERATED_BODY()
diff --git a/Engine/Plugins/Experimental/Animation/PoseSearch/Source/Editor/Public/PoseSearchDatabaseEditorReflection.h b/Engine/Plugins/Experimental/Animation/PoseSearch/Source/Editor/Public/PoseSearchDatabaseEditorReflection.h
--- a/Engine/Plugins/Experimental/Animation/PoseSearch/Source/Editor/Public/PoseSearchDatabaseEditorReflection.h
+++ b/Engine/Plugins/Experimental/Animation/PoseSearch/Source/Editor/Public/PoseSearchDatabaseEditorReflection.h
@@ -206,9 +206,9 @@ public:
 	UPROPERTY(VisibleAnywhere, Category = "General Information")
 	uint32 PrunedPCAFrames;
 
-	// Cardinality for the database config (how many floats per pose to store the pose features data)
+	// Cardinality for the database schema (how many floats per pose to store the pose features data)
 	UPROPERTY(VisibleAnywhere, Category = "General Information")
-	uint32 ConfigCardinality;
+	uint32 SchemaCardinality;
 
 	// Average speed of the characters trajectory across all animations in the database.
 	UPROPERTY(VisibleAnywhere, Category = "Kinematic Information")
diff --git a/Engine/Plugins/Experimental/Animation/PoseSearch/Source/Runtime/Private/PoseSearchDatabase.cpp b/Engine/Plugins/Experimental/Animation/PoseSearch/Source/Runtime/Private/PoseSearchDatabase.cpp
--- a/Engine/Plugins/Experimental/Animation/PoseSearch/Source/Runtime/Private/PoseSearchDatabase.cpp
+++ b/Engine/Plugins/Experimental/Animation/PoseSearch/Source/Runtime/Private/PoseSearchDatabase.cpp
@@ -619,20 +619,6 @@ bool UPoseSearchDatabase::GetSkipSearchIfPossible() const
 
 void UPoseSearchDatabase::PostLoad()
 {
-#if WITH_EDITORONLY_DATA
-	for (const FPoseSearchDatabaseSequence& DatabaseSequence : Sequences_DEPRECATED)
-	{
-		AnimationAssets.Add(FInstancedStruct::Make(DatabaseSequence));
-	}
-	Sequences_DEPRECATED.Empty();
-
-	for (const FPoseSearchDatabaseBlendSpace& DatabaseBlendSpace : BlendSpaces_DEPRECATED)
-	{
-		AnimationAssets.Add(FInstancedStruct::Make(DatabaseBlendSpace));
-	}
-	BlendSpaces_DEPRECATED.Empty();
-#endif // WITH_EDITORONLY_DATA
-
 #if WITH_EDITOR
 	using namespace UE::PoseSearch;
 
diff --git a/Engine/Plugins/Experimental/Animation/PoseSearch/Source/Runtime/Private/PoseSearchFeatureChannel_Heading.h b/Engine/Plugins/Experimental/Animation/PoseSearch/Source/Runtime/Private/PoseSearchFeatureChannel_Heading.h
--- a/Engine/Plugins/Experimental/Animation/PoseSearch/Source/Runtime/Private/PoseSearchFeatureChannel_Heading.h
+++ b/Engine/Plugins/Experimental/Animation/PoseSearch/Source/Runtime/Private/PoseSearchFeatureChannel_Heading.h
@@ -13,8 +13,7 @@ enum class EHeadingAxis : uint8
 	Y,
 	Z,
 
-	Num UMETA(Hidden),
-	Invalid = Num UMETA(Hidden)
+	Num UMETA(Hidden)
 };
 
 UCLASS(EditInlineNew, Blueprintable, meta = (DisplayName = "Heading Channel"), CollapseCategories)
diff --git a/Engine/Plugins/Experimental/Animation/PoseSearch/Source/Runtime/Public/PoseSearch/MotionMatchingAnimNodeLibrary.h b/Engine/Plugins/Experimental/Animation/PoseSearch/Source/Runtime/Public/PoseSearch/MotionMatchingAnimNodeLibrary.h
--- a/Engine/Plugins/Experimental/Animation/PoseSearch/Source/Runtime/Public/PoseSearch/MotionMatchingAnimNodeLibrary.h
+++ b/Engine/Plugins/Experimental/Animation/PoseSearch/Source/Runtime/Public/PoseSearch/MotionMatchingAnimNodeLibrary.h
@@ -19,7 +19,7 @@ struct FMotionMatchingAnimNodeReference : public FAnimNodeReference
 };
 
 // Exposes operations that can be run on a Motion Matching node via Anim Node Functions such as "On Become Relevant" and "On Update".
-UCLASS(Experimental)
+UCLASS()
 class POSESEARCH_API UMotionMatchingAnimNodeLibrary : public UBlueprintFunctionLibrary
 {
 	GENERATED_BODY()
diff --git a/Engine/Plugins/Experimental/Animation/PoseSearch/Source/Runtime/Public/PoseSearch/PoseSearchAnimNotifies.h b/Engine/Plugins/Experimental/Animation/PoseSearch/Source/Runtime/Public/PoseSearch/PoseSearchAnimNotifies.h
--- a/Engine/Plugins/Experimental/Animation/PoseSearch/Source/Runtime/Public/PoseSearch/PoseSearchAnimNotifies.h
+++ b/Engine/Plugins/Experimental/Animation/PoseSearch/Source/Runtime/Public/PoseSearch/PoseSearchAnimNotifies.h
@@ -17,23 +17,23 @@ class POSESEARCH_API UAnimNotifyState_PoseSearchBase : public UAnimNotifyState
 
 // Use this notify state to remove animation segments from the database completely, they will never play or return from
 // a search result
-UCLASS(Blueprintable, meta = (DisplayName = "Pose Matching: Exclude From Database"))
+UCLASS(Blueprintable, meta = (DisplayName = "Pose Search: Exclude From Database"))
 class POSESEARCH_API UAnimNotifyState_PoseSearchExcludeFromDatabase : public UAnimNotifyState_PoseSearchBase
 {
 	GENERATED_BODY()
 };
 
-// A pose matching search will not return results that overlap this notify, but the animation segment can still play
+// A pose search search will not return results that overlap this notify, but the animation segment can still play
 // if a previous search result advances into it.
-UCLASS(Blueprintable, meta = (DisplayName = "Pose Matching: Block Transition"))
+UCLASS(Blueprintable, meta = (DisplayName = "Pose Search: Block Transition"))
 class POSESEARCH_API UAnimNotifyState_PoseSearchBlockTransition : public UAnimNotifyState_PoseSearchBase
 {
 	GENERATED_BODY()
 };
 
-// Pose matching cost will be affected by this, making the animation segment more or less likely to be selected based
+// Pose search cost will be affected by this, making the animation segment more or less likely to be selected based
 // on the notify parameters
-UCLASS(Blueprintable, meta = (DisplayName = "Pose Matching: Override Base Cost Bias"))
+UCLASS(Blueprintable, meta = (DisplayName = "Pose Search: Override Base Cost Bias"))
 // @todo: rename into UAnimNotifyState_PoseSearchOverrideBaseCostBias
 class POSESEARCH_API UAnimNotifyState_PoseSearchModifyCost : public UAnimNotifyState_PoseSearchBase
 {
@@ -47,9 +47,9 @@ public:
 	float CostAddend = -1.0f;
 };
 
-// Pose matching cost for the continuing pose will be affected by this, making the animation segment more or less 
+// Pose search cost for the continuing pose will be affected by this, making the animation segment more or less 
 // likely to be continuing playing based on the notify parameters
-UCLASS(Blueprintable, meta = (DisplayName = "Pose Matching: Override Continuing Pose Cost Bias"))
+UCLASS(Blueprintable, meta = (DisplayName = "Pose Search: Override Continuing Pose Cost Bias"))
 class POSESEARCH_API UAnimNotifyState_PoseSearchOverrideContinuingPoseCostBias : public UAnimNotifyState_PoseSearchBase
 {
 	GENERATED_BODY()
@@ -64,7 +64,7 @@ public:
 
 // UPoseSearchFeatureChannel(s) can use this UAnimNotifyState_PoseSearchSamplingAttribute as animation space position, rotation, and linear velocity provider 
 // during database indexing by specifying their SamplingAttributeId property to match UAnimNotifyState_PoseSearchSamplingAttribute::SamplingAttributeId
-UCLASS(Blueprintable, meta = (DisplayName = "Pose Matching: Sampling Attribute"))
+UCLASS(Blueprintable, meta = (DisplayName = "Pose Search: Sampling Attribute"))
 class POSESEARCH_API UAnimNotifyState_PoseSearchSamplingAttribute : public UAnimNotifyState_PoseSearchBase
 {
 	GENERATED_BODY()
@@ -91,7 +91,7 @@ public:
 #endif // WITH_EDITORONLY_DATA
 };
 
-UCLASS(Blueprintable, meta = (DisplayName = "Pose Matching: Motion Matched Branch In"))
+UCLASS(Blueprintable, meta = (DisplayName = "Pose Search: Motion Matched Branch In"))
 class POSESEARCH_API UAnimNotifyState_PoseSearchBranchIn : public UAnimNotifyState_PoseSearchBase
 {
 	GENERATED_BODY()
diff --git a/Engine/Plugins/Experimental/Animation/PoseSearch/Source/Runtime/Public/PoseSearch/PoseSearchContext.h b/Engine/Plugins/Experimental/Animation/PoseSearch/Source/Runtime/Public/PoseSearch/PoseSearchContext.h
--- a/Engine/Plugins/Experimental/Animation/PoseSearch/Source/Runtime/Public/PoseSearch/PoseSearchContext.h
+++ b/Engine/Plugins/Experimental/Animation/PoseSearch/Source/Runtime/Public/PoseSearch/PoseSearchContext.h
@@ -22,7 +22,7 @@ enum class EDebugDrawFlags : uint32
 {
 	None = 0,
 
-	// Draw using Query colors form the schema / config
+	// Draw using Query colors form the Schema
 	DrawQuery = 1 << 1,
 };
 ENUM_CLASS_FLAGS(EDebugDrawFlags);
diff --git a/Engine/Plugins/Experimental/Animation/PoseSearch/Source/Runtime/Public/PoseSearch/PoseSearchDatabase.h b/Engine/Plugins/Experimental/Animation/PoseSearch/Source/Runtime/Public/PoseSearch/PoseSearchDatabase.h
--- a/Engine/Plugins/Experimental/Animation/PoseSearch/Source/Runtime/Public/PoseSearch/PoseSearchDatabase.h
+++ b/Engine/Plugins/Experimental/Animation/PoseSearch/Source/Runtime/Public/PoseSearch/PoseSearchDatabase.h
@@ -33,11 +33,8 @@ enum class EPoseSearchMode : int32
 	// Optimized search mode: the database projects the poses into a PCA space using only the most significant "NumberOfPrincipalComponents" dimensions, and construct a kdtree to facilitate the search.
 	PCAKDTree,
 
-	// Optimized search mode using a vantage point tree (disabled/Hidden until production ready)
-	VPTree UMETA(Hidden),
-
-	Num UMETA(Hidden),
-	Invalid = Num UMETA(Hidden)
+	// Optimized search mode using a vantage point tree (Experimental)
+	VPTree
 };
 
 UENUM()
@@ -45,10 +42,7 @@ enum class EPoseSearchMirrorOption : int32
 {
 	UnmirroredOnly UMETA(DisplayName = "Original Only"),
 	MirroredOnly UMETA(DisplayName = "Mirrored Only"),
-	UnmirroredAndMirrored UMETA(DisplayName = "Original and Mirrored"),
-
-	Num UMETA(Hidden),
-	Invalid = Num UMETA(Hidden)
+	UnmirroredAndMirrored UMETA(DisplayName = "Original and Mirrored")
 };
 
 USTRUCT()
@@ -84,7 +78,7 @@ struct POSESEARCH_API FPoseSearchDatabaseAnimationAssetBase
 	bool bDisableReselection = false;
 
 	// This allows users to set if this animation is original only (no mirrored data), original and mirrored, or only the mirrored version of this animation.
-	// It requires the mirror table to be set up in the config file.
+	// It requires the mirror table to be set up in the database Schema.
 	UPROPERTY(EditAnywhere, Category = "Settings", meta = (DisplayPriority = 3))
 	EPoseSearchMirrorOption MirrorOption = EPoseSearchMirrorOption::UnmirroredOnly;
 
@@ -225,14 +219,14 @@ struct POSESEARCH_API FPoseSearchDatabaseAnimMontage : public FPoseSearchDatabas
 };
 
 /** A data asset for indexing a collection of animation sequences. */
-UCLASS(BlueprintType, Category = "Animation|Pose Search", Experimental, meta = (DisplayName = "Motion Database"))
+UCLASS(BlueprintType, Category = "Animation|Pose Search", meta = (DisplayName = "Pose Search Database"))
 class POSESEARCH_API UPoseSearchDatabase : public UDataAsset
 {
 	GENERATED_BODY()
 public:
 
-	// The Motion Database Config sets what channels this database will use to match against (bones, trajectory and what properties of those you’re interested in, such as position and velocity).
-	UPROPERTY(EditAnywhere, BlueprintReadOnly, Category="Database", DisplayName="Config")
+	// The Schema sets what channels this database will use to match against (bones, trajectory and what properties of those you’re interested in, such as position and velocity).
+	UPROPERTY(EditAnywhere, BlueprintReadOnly, Category="Database")
 	TObjectPtr<const UPoseSearchSchema> Schema;
 
 	// Cost added to the continuing pose from this database. This allows users to apply a cost bias (positive or negative) to the continuing pose.
@@ -276,13 +270,6 @@ public:
 	UPROPERTY(EditAnywhere, BlueprintReadOnly, Category = "Database")
 	TObjectPtr<const UPoseSearchNormalizationSet> NormalizationSet;
 
-	// Sequences and Blendspaces are deprecated and its data will be part of the AnimationAssets.
-	// All sequences and blend spaces will be added to the AnimationAssets in PostLoad().
-	UPROPERTY()
-	TArray<FPoseSearchDatabaseSequence> Sequences_DEPRECATED;
-	UPROPERTY()
-	TArray<FPoseSearchDatabaseBlendSpace> BlendSpaces_DEPRECATED;
-
 	// If null, the default preview mesh for the skeleton will be used. Otherwise, this will be used in preview scenes.
 	// @todo: Move this to be a setting in the Pose Search Database editor. 
 	UPROPERTY(EditAnywhere, Category = "Preview")
diff --git a/Engine/Plugins/Experimental/Animation/PoseSearch/Source/Runtime/Public/PoseSearch/PoseSearchNormalizationSet.h b/Engine/Plugins/Experimental/Animation/PoseSearch/Source/Runtime/Public/PoseSearch/PoseSearchNormalizationSet.h
--- a/Engine/Plugins/Experimental/Animation/PoseSearch/Source/Runtime/Public/PoseSearch/PoseSearchNormalizationSet.h
+++ b/Engine/Plugins/Experimental/Animation/PoseSearch/Source/Runtime/Public/PoseSearch/PoseSearchNormalizationSet.h
@@ -8,7 +8,7 @@
 
 class UPoseSearchDatabase;
 
-UCLASS(BlueprintType, Category = "Animation|Pose Search", Experimental, meta = (DisplayName = "Normalization Set"))
+UCLASS(BlueprintType, Category = "Animation|Pose Search", meta = (DisplayName = "Pose Search Normalization Set"))
 class POSESEARCH_API UPoseSearchNormalizationSet : public UDataAsset
 {
 	GENERATED_BODY()
diff --git a/Engine/Plugins/Experimental/Animation/PoseSearch/Source/Runtime/Public/PoseSearch/PoseSearchSchema.h b/Engine/Plugins/Experimental/Animation/PoseSearch/Source/Runtime/Public/PoseSearch/PoseSearchSchema.h
--- a/Engine/Plugins/Experimental/Animation/PoseSearch/Source/Runtime/Public/PoseSearch/PoseSearchSchema.h
+++ b/Engine/Plugins/Experimental/Animation/PoseSearch/Source/Runtime/Public/PoseSearch/PoseSearchSchema.h
@@ -21,16 +21,13 @@ enum class EPoseSearchDataPreprocessor : int32
 	Normalize,
 
 	// The data will be normalized against its deviation
-	NormalizeOnlyByDeviation,
-
-	Num UMETA(Hidden),
-	Invalid = Num UMETA(Hidden)
+	NormalizeOnlyByDeviation
 };
 
 /**
 * Specifies the format of a pose search index. At runtime, queries are built according to the schema for searching.
 */
-UCLASS(BlueprintType, Category = "Animation|Pose Search", Experimental, meta = (DisplayName = "Motion Database Config"), CollapseCategories)
+UCLASS(BlueprintType, Category = "Animation|Pose Search", meta = (DisplayName = "Pose Search Schema"), CollapseCategories)
 class POSESEARCH_API UPoseSearchSchema : public UDataAsset, public IBoneReferenceSkeletonProvider
 {
 	GENERATED_BODY()
@@ -45,7 +42,7 @@ public:
 	int32 SampleRate = 30;
 
 private:
-	// Channels itemize the cost breakdown of the config in simpler parts such as position or velocity of a bones, or phase of limbs. The total cost of a query against an indexed database pose will be the sum of the combined channel costs
+	// Channels itemize the cost breakdown of the Schema in simpler parts such as position or velocity of a bones, or phase of limbs. The total cost of a query against an indexed database pose will be the sum of the combined channel costs
 	UPROPERTY(EditAnywhere, Instanced, Category = "Schema")
 	TArray<TObjectPtr<UPoseSearchFeatureChannel>> Channels;
 
