motion matching - moving traj from MM node into pose history (to allow future change to generate traj into the pose history, synchronized with the captured poses)

#virtualized

diff --git a/Engine/Plugins/Experimental/Animation/PoseSearch/Source/Runtime/Private/AnimNode_MotionMatching.cpp b/Engine/Plugins/Experimental/Animation/PoseSearch/Source/Runtime/Private/AnimNode_MotionMatching.cpp
--- a/Engine/Plugins/Experimental/Animation/PoseSearch/Source/Runtime/Private/AnimNode_MotionMatching.cpp
+++ b/Engine/Plugins/Experimental/Animation/PoseSearch/Source/Runtime/Private/AnimNode_MotionMatching.cpp
@@ -22,7 +22,6 @@
 #if ENABLE_ANIM_DEBUG
 static TAutoConsoleVariable<bool> CVarAnimNodeMotionMatchingDrawQuery(TEXT("a.AnimNode.MotionMatching.DebugDrawQuery"), false, TEXT("Draw input query"));
 static TAutoConsoleVariable<bool> CVarAnimNodeMotionMatchingDrawCurResult(TEXT("a.AnimNode.MotionMatching.DebugDrawCurResult"), false, TEXT("Draw current result"));
-static TAutoConsoleVariable<bool> CVarAnimNodeMotionMatchingDrawPoseHistory(TEXT("a.AnimNode.MotionMatching.DebugDrawPoseHistory"), false, TEXT("Draw Pose History"));
 static TAutoConsoleVariable<bool> CVarAnimNodeMotionMatchingDrawInfo(TEXT("a.AnimNode.MotionMatching.DebugDrawInfo"), false, TEXT("Draw info like current databases and asset"));
 static TAutoConsoleVariable<bool> CVarAnimNodeMotionMatchingDrawInfoVerbose(TEXT("a.AnimNode.MotionMatching.DebugDrawInfoVerbose"), true, TEXT("Draw additional info like blend stack"));
 static TAutoConsoleVariable<float> CVarAnimNodeMotionMatchingDrawInfoHeight(TEXT("a.AnimNode.MotionMatching.DebugDrawInfoHeight"), 50.f, TEXT("Vertical offset for DebugDrawInfo"));
@@ -169,8 +168,6 @@ void FAnimNode_MotionMatching::UpdateAssetPlayer(const FAnimationUpdateContext&
 	UPoseSearchLibrary::UpdateMotionMatchingState(
 		Context,
 		DatabasesToSearch,
-		Trajectory,
-		TrajectorySpeedMultiplier,
 		BlendTime,
 		MaxActiveBlends,
 		PoseJumpThresholdTime,
@@ -186,7 +183,6 @@ void FAnimNode_MotionMatching::UpdateAssetPlayer(const FAnimationUpdateContext&
 		#if ENABLE_ANIM_DEBUG
 		, CVarAnimNodeMotionMatchingDrawQuery.GetValueOnAnyThread()
 		, CVarAnimNodeMotionMatchingDrawCurResult.GetValueOnAnyThread()
-		, CVarAnimNodeMotionMatchingDrawPoseHistory.GetValueOnAnyThread()
 		#endif // ENABLE_ANIM_DEBUG
 	);
 
diff --git a/Engine/Plugins/Experimental/Animation/PoseSearch/Source/Runtime/Private/AnimNode_PoseSearchHistoryCollector.cpp b/Engine/Plugins/Experimental/Animation/PoseSearch/Source/Runtime/Private/AnimNode_PoseSearchHistoryCollector.cpp
--- a/Engine/Plugins/Experimental/Animation/PoseSearch/Source/Runtime/Private/AnimNode_PoseSearchHistoryCollector.cpp
+++ b/Engine/Plugins/Experimental/Animation/PoseSearch/Source/Runtime/Private/AnimNode_PoseSearchHistoryCollector.cpp
@@ -10,10 +10,6 @@
 
 #define LOCTEXT_NAMESPACE "AnimNode_PoseSearchHistoryCollector"
 
-#if ENABLE_DRAW_DEBUG && ENABLE_ANIM_DEBUG
-TAutoConsoleVariable<bool> CVarAnimPoseHistoryDebugDraw(TEXT("a.AnimNode.PoseHistory.DebugDraw"), false, TEXT("Enable / Disable Pose History DebugDraw"));
-#endif
-
 namespace UE::PoseSearch::Private
 {
 
@@ -128,22 +124,20 @@ void FAnimNode_PoseSearchHistoryCollector::Evaluate_AnyThread(FPoseContext& Outp
 	Super::Evaluate_AnyThread(Output);
 	Source.Evaluate(Output);
 
+	PoseHistory.UpdateTrajectory(Trajectory, TrajectorySpeedMultiplier);
+
 	FCSPose<FCompactPose> ComponentSpacePose;
 	ComponentSpacePose.InitPose(Output.Pose);
 	PoseHistory.Update(Output.AnimInstanceProxy->GetDeltaSeconds(), ComponentSpacePose, bStoreScales);
 
 #if ENABLE_DRAW_DEBUG && ENABLE_ANIM_DEBUG
-	if (CVarAnimPoseHistoryDebugDraw.GetValueOnAnyThread())
-	{
 		FColor Color;
 #if WITH_EDITORONLY_DATA
 		Color = DebugColor.ToFColor(true);
 #else // WITH_EDITORONLY_DATA
 		Color = FLinearColor::Red.ToFColor(true);
 #endif // WITH_EDITORONLY_DATA
-
-		PoseHistory.DebugDraw(*Output.AnimInstanceProxy, Color, &DebugDrawTrajectory);
-	}
+	PoseHistory.DebugDraw(*Output.AnimInstanceProxy, Color);
 #endif // ENABLE_DRAW_DEBUG && ENABLE_ANIM_DEBUG
 }
 
@@ -161,7 +155,6 @@ void FAnimNode_PoseSearchHistoryCollector::GatherDebugData(FNodeDebugData& Debug
 	Source.GatherDebugData(DebugData);
 }
 
-
 /////////////////////////////////////////////////////
 // FAnimNode_PoseSearchComponentSpaceHistoryCollector
 
@@ -189,20 +182,17 @@ void FAnimNode_PoseSearchComponentSpaceHistoryCollector::EvaluateComponentSpace_
 	Super::EvaluateComponentSpace_AnyThread(Output);
 	Source.EvaluateComponentSpace(Output);
 
+	PoseHistory.UpdateTrajectory(Trajectory, TrajectorySpeedMultiplier);
 	PoseHistory.Update(Output.AnimInstanceProxy->GetDeltaSeconds(), Output.Pose, bStoreScales);
 
 #if ENABLE_DRAW_DEBUG && ENABLE_ANIM_DEBUG
-	if (CVarAnimPoseHistoryDebugDraw.GetValueOnAnyThread())
-	{
 		FColor Color;
 #if WITH_EDITORONLY_DATA
 		Color = DebugColor.ToFColor(true);
 #else // WITH_EDITORONLY_DATA
 		Color = FLinearColor::Red.ToFColor(true);
 #endif // WITH_EDITORONLY_DATA
-
-		PoseHistory.DebugDraw(*Output.AnimInstanceProxy, Color, &DebugDrawTrajectory);
-	}
+	PoseHistory.DebugDraw(*Output.AnimInstanceProxy, Color);
 #endif // ENABLE_DRAW_DEBUG && ENABLE_ANIM_DEBUG
 }
 
diff --git a/Engine/Plugins/Experimental/Animation/PoseSearch/Source/Runtime/Private/PoseSearchContext.cpp b/Engine/Plugins/Experimental/Animation/PoseSearch/Source/Runtime/Private/PoseSearchContext.cpp
--- a/Engine/Plugins/Experimental/Animation/PoseSearch/Source/Runtime/Private/PoseSearchContext.cpp
+++ b/Engine/Plugins/Experimental/Animation/PoseSearch/Source/Runtime/Private/PoseSearchContext.cpp
@@ -361,12 +361,11 @@ FCachedQuery::FCachedQuery(const UPoseSearchSchema* InSchema)
 //////////////////////////////////////////////////////////////////////////
 // FSearchContext
 FSearchContext::FSearchContext(const UAnimInstance* InAnimInstance, const IPoseHistory* InHistory, TConstArrayView<const UAnimationAsset*> InAnimationsToConsider,
-		const FPoseSearchQueryTrajectory* InTrajectory, float InDesiredPermutationTimeOffset, const FPoseIndicesHistory* InPoseIndicesHistory,
+		float InDesiredPermutationTimeOffset, const FPoseIndicesHistory* InPoseIndicesHistory,
 		const FSearchResult& InCurrentResult, const FFloatInterval& InPoseJumpThresholdTime, bool bInUseCachedChannelData)
 : AnimInstance(InAnimInstance)
 , History(InHistory)
 , AnimationsToConsider(InAnimationsToConsider)
-, Trajectory(InTrajectory)
 , DesiredPermutationTimeOffset(InDesiredPermutationTimeOffset)
 , PoseIndicesHistory(InPoseIndicesHistory)
 , CurrentResult(InCurrentResult)
@@ -440,18 +439,29 @@ FVector FSearchContext::GetSampleVelocity(float SampleTimeOffset, float OriginTi
 
 FTransform FSearchContext::GetWorldRootBoneTransformAtTime(float SampleTime) const
 {
-	if (Trajectory)
+	check(!CachedQueries.IsEmpty());
+	const UPoseSearchSchema* Schema = CachedQueries.Last().GetSchema();
+	check(Schema);
+
+	#if WITH_EDITOR
+	if (!History)
+	{
+		UE_LOG(LogPoseSearch, Error, TEXT("FSearchContext::GetWorldRootBoneTransformAtTime - Couldn't search for world space root boneTransform by %s, because no IPoseHistory has been found!"), *Schema->GetName());
+	}
+	else
+	#endif // WITH_EDITOR
 	{
-		// Trajectory is already in root bone world space (transformed in UPoseSearchLibrary::ProcessTrajectory), so we just ask for a sample at the proper SampleTime
-		return Trajectory->GetSampleAtTime(SampleTime, true).GetTransform();
+		FTransform WorldRootBoneTransform;
+		if (ensure(History) && History->GetTransformAtTime(SampleTime, WorldRootBoneTransform, Schema->Skeleton, RootBoneIndexType, WorldSpaceIndexType))
+		{
+			return WorldRootBoneTransform;
+		}
 	}
 
 	if (AnimInstance && AnimInstance->CurrentSkeleton)
 	{
-		// @todo: perhaps pass in to FSearchContext the previous frame root bone transform as we should do in ProcessTrajectory
 		const FTransform& RootBoneTransform = AnimInstance->CurrentSkeleton->GetReferenceSkeleton().GetRefBonePose()[RootSchemaBoneIdx];
 		const FTransform& ComponentToWorldTransform = AnimInstance->GetSkelMeshComponent()->GetComponentTransform();
-
 		return RootBoneTransform * ComponentToWorldTransform;
 	}
 
@@ -461,8 +471,11 @@ FTransform FSearchContext::GetWorldRootBoneTransformAtTime(float SampleTime) con
 FTransform FSearchContext::GetWorldBoneTransformAtTime(float SampleTime, int8 SchemaBoneIdx)
 {
 	// CachedQueries.Last is the query we're building 
-	const UPoseSearchSchema* Schema = !CachedQueries.IsEmpty() ? CachedQueries.Last().GetSchema() : nullptr;
-	const FBoneIndexType BoneIndexType = !Schema || SchemaBoneIdx == RootSchemaBoneIdx ? RootBoneIndexType : Schema->GetBoneIndexType(SchemaBoneIdx);
+	check(!CachedQueries.IsEmpty());
+	const UPoseSearchSchema* Schema = CachedQueries.Last().GetSchema();
+	check(Schema);
+
+	const FBoneIndexType BoneIndexType = Schema->GetBoneIndexType(SchemaBoneIdx);
 	if (const FCachedTransform<FTransform>* CachedTransform = CachedTransforms.Find(SampleTime, BoneIndexType))
 	{
 		return CachedTransform->Transform;
@@ -486,8 +499,6 @@ FTransform FSearchContext::GetWorldBoneTransformAtTime(float SampleTime, int8 Sc
 			WorldBoneTransform = GetWorldRootBoneTransformAtTime(SampleTime);
 		}
 
-		check(Schema);
-
 		// collecting the local bone transforms from the IPoseHistory
 		#if WITH_EDITOR
 		if (!History)
diff --git a/Engine/Plugins/Experimental/Animation/PoseSearch/Source/Runtime/Private/PoseSearchFeatureChannel_FilterCrashingLegs.cpp b/Engine/Plugins/Experimental/Animation/PoseSearch/Source/Runtime/Private/PoseSearchFeatureChannel_FilterCrashingLegs.cpp
--- a/Engine/Plugins/Experimental/Animation/PoseSearch/Source/Runtime/Private/PoseSearchFeatureChannel_FilterCrashingLegs.cpp
+++ b/Engine/Plugins/Experimental/Animation/PoseSearch/Source/Runtime/Private/PoseSearchFeatureChannel_FilterCrashingLegs.cpp
@@ -91,7 +91,7 @@ void UPoseSearchFeatureChannel_FilterCrashingLegs::BuildQuery(UE::PoseSearch::FS
 	// trying to get the BuildQuery data from the continuing pose
 	const bool bCanUseCurrentResult = SearchContext.CanUseCurrentResult();
 	const bool bSkip = InputQueryPose != EInputQueryPose::UseCharacterPose && bCanUseCurrentResult;
-	if (bSkip || !SearchContext.IsHistoryValid())
+	if (bSkip || !SearchContext.GetHistory())
 	{
 		if (bCanUseCurrentResult)
 		{
diff --git a/Engine/Plugins/Experimental/Animation/PoseSearch/Source/Runtime/Private/PoseSearchFeatureChannel_Heading.cpp b/Engine/Plugins/Experimental/Animation/PoseSearch/Source/Runtime/Private/PoseSearchFeatureChannel_Heading.cpp
--- a/Engine/Plugins/Experimental/Animation/PoseSearch/Source/Runtime/Private/PoseSearchFeatureChannel_Heading.cpp
+++ b/Engine/Plugins/Experimental/Animation/PoseSearch/Source/Runtime/Private/PoseSearchFeatureChannel_Heading.cpp
@@ -142,20 +142,20 @@ void UPoseSearchFeatureChannel_Heading::BuildQuery(UE::PoseSearch::FSearchContex
 
 	const bool bCanUseCurrentResult = SearchContext.CanUseCurrentResult();
 	const bool bSkip = InputQueryPose != EInputQueryPose::UseCharacterPose && bCanUseCurrentResult;
-	if (bSkip || (!SearchContext.IsHistoryValid() && !bIsRootBone))
-	{
-		if (bCanUseCurrentResult)
+	if (bSkip || (!SearchContext.GetHistory() && !bIsRootBone))
 		{
+		if (bCanUseCurrentResult)
+			{
 			FFeatureVectorHelper::Copy(SearchContext.EditFeatureVector(), ChannelDataOffset, ChannelCardinality, SearchContext.GetCurrentResultPoseVector());
 			return;
-		}
+			}
 		
 		// we leave the SearchContext.EditFeatureVector() set to zero since the SearchContext.History is invalid and it'll fail if we continue
-		UE_LOG(LogPoseSearch, Error, TEXT("UPoseSearchFeatureChannel_Heading::BuildQuery - Failed because Pose History Node is missing."));
+				UE_LOG(LogPoseSearch, Error, TEXT("UPoseSearchFeatureChannel_Heading::BuildQuery - Failed because Pose History Node is missing."));
 		return;
-	}
+			}
 	
-	// calculating the BoneRotation in component space for the bone indexed by SchemaBoneIdx
+			// calculating the BoneRotation in component space for the bone indexed by SchemaBoneIdx
 	const FQuat BoneRotation = SearchContext.GetSampleRotation(SampleTimeOffset, OriginTimeOffset, SchemaBoneIdx, SchemaOriginBoneIdx, PermutationTimeType);
 	FFeatureVectorHelper::EncodeVector(SearchContext.EditFeatureVector(), ChannelDataOffset, GetAxis(BoneRotation), ComponentStripping);
 }
diff --git a/Engine/Plugins/Experimental/Animation/PoseSearch/Source/Runtime/Private/PoseSearchFeatureChannel_Phase.cpp b/Engine/Plugins/Experimental/Animation/PoseSearch/Source/Runtime/Private/PoseSearchFeatureChannel_Phase.cpp
--- a/Engine/Plugins/Experimental/Animation/PoseSearch/Source/Runtime/Private/PoseSearchFeatureChannel_Phase.cpp
+++ b/Engine/Plugins/Experimental/Animation/PoseSearch/Source/Runtime/Private/PoseSearchFeatureChannel_Phase.cpp
@@ -334,7 +334,7 @@ void UPoseSearchFeatureChannel_Phase::BuildQuery(UE::PoseSearch::FSearchContext&
 
 	const bool bCanUseCurrentResult = SearchContext.CanUseCurrentResult();
 	const bool bSkip = InputQueryPose != EInputQueryPose::UseCharacterPose && bCanUseCurrentResult;
-	if (bSkip || !SearchContext.IsHistoryValid())
+	if (bSkip || !SearchContext.GetHistory())
 	{
 		if (bCanUseCurrentResult)
 		{
diff --git a/Engine/Plugins/Experimental/Animation/PoseSearch/Source/Runtime/Private/PoseSearchFeatureChannel_Position.cpp b/Engine/Plugins/Experimental/Animation/PoseSearch/Source/Runtime/Private/PoseSearchFeatureChannel_Position.cpp
--- a/Engine/Plugins/Experimental/Animation/PoseSearch/Source/Runtime/Private/PoseSearchFeatureChannel_Position.cpp
+++ b/Engine/Plugins/Experimental/Animation/PoseSearch/Source/Runtime/Private/PoseSearchFeatureChannel_Position.cpp
@@ -118,21 +118,21 @@ void UPoseSearchFeatureChannel_Position::BuildQuery(UE::PoseSearch::FSearchConte
 
 	const bool bCanUseCurrentResult = SearchContext.CanUseCurrentResult();
 	const bool bSkip = InputQueryPose != EInputQueryPose::UseCharacterPose && bCanUseCurrentResult;
-	const bool bIsRootBone = SchemaBoneIdx == RootSchemaBoneIdx;
-	if (bSkip || (!SearchContext.IsHistoryValid() && !bIsRootBone))
-	{
-		if (bCanUseCurrentResult)
+		const bool bIsRootBone = SchemaBoneIdx == RootSchemaBoneIdx;
+	if (bSkip || (!SearchContext.GetHistory() && !bIsRootBone))
 		{
+		if (bCanUseCurrentResult)
+			{
 			FFeatureVectorHelper::Copy(SearchContext.EditFeatureVector(), ChannelDataOffset, ChannelCardinality, SearchContext.GetCurrentResultPoseVector());
 			return;
-		}
+			}
 
 		// we leave the SearchContext.EditFeatureVector() set to zero since the SearchContext.History is invalid and it'll fail if we continue
-		UE_LOG(LogPoseSearch, Error, TEXT("UPoseSearchFeatureChannel_Position::BuildQuery - Failed because Pose History Node is missing."));
+				UE_LOG(LogPoseSearch, Error, TEXT("UPoseSearchFeatureChannel_Position::BuildQuery - Failed because Pose History Node is missing."));
 		return;
-	}
+			}
 	
-	// calculating the BonePosition in root bone space for the bone indexed by SchemaBoneIdx
+			// calculating the BonePosition in root bone space for the bone indexed by SchemaBoneIdx
 	const FVector BonePosition = SearchContext.GetSamplePosition(SampleTimeOffset, OriginTimeOffset, SchemaBoneIdx, SchemaOriginBoneIdx, PermutationTimeType);
 	FFeatureVectorHelper::EncodeVector(SearchContext.EditFeatureVector(), ChannelDataOffset, BonePosition, ComponentStripping);
 }
diff --git a/Engine/Plugins/Experimental/Animation/PoseSearch/Source/Runtime/Private/PoseSearchFeatureChannel_Velocity.cpp b/Engine/Plugins/Experimental/Animation/PoseSearch/Source/Runtime/Private/PoseSearchFeatureChannel_Velocity.cpp
--- a/Engine/Plugins/Experimental/Animation/PoseSearch/Source/Runtime/Private/PoseSearchFeatureChannel_Velocity.cpp
+++ b/Engine/Plugins/Experimental/Animation/PoseSearch/Source/Runtime/Private/PoseSearchFeatureChannel_Velocity.cpp
@@ -97,25 +97,25 @@ void UPoseSearchFeatureChannel_Velocity::BuildQuery(UE::PoseSearch::FSearchConte
 
 	const bool bCanUseCurrentResult = SearchContext.CanUseCurrentResult();
 	const bool bSkip = InputQueryPose != EInputQueryPose::UseCharacterPose && bCanUseCurrentResult;
-	if (bSkip || (!SearchContext.IsHistoryValid() && !bIsRootBone))
-	{
-		if (bCanUseCurrentResult)
+	if (bSkip || (!SearchContext.GetHistory() && !bIsRootBone))
 		{
+		if (bCanUseCurrentResult)
+			{
 			FFeatureVectorHelper::Copy(SearchContext.EditFeatureVector(), ChannelDataOffset, ChannelCardinality, SearchContext.GetCurrentResultPoseVector());
 			return;
-		}
+			}
 
 		// we leave the SearchContext.EditFeatureVector() set to zero since the SearchContext.History is invalid and it'll fail if we continue
-		UE_LOG(LogPoseSearch, Error, TEXT("UPoseSearchFeatureChannel_Velocity::BuildQuery - Failed because Pose History Node is missing."));
+				UE_LOG(LogPoseSearch, Error, TEXT("UPoseSearchFeatureChannel_Velocity::BuildQuery - Failed because Pose History Node is missing."));
 		return;
-	}
+			}
 	
-	// calculating the LinearVelocity for the bone indexed by SchemaBoneIdx
+			// calculating the LinearVelocity for the bone indexed by SchemaBoneIdx
 	FVector LinearVelocity = SearchContext.GetSampleVelocity(SampleTimeOffset, OriginTimeOffset, SchemaBoneIdx, SchemaOriginBoneIdx, bUseCharacterSpaceVelocities, PermutationTimeType);
-	if (bNormalize)
-	{
-		LinearVelocity = LinearVelocity.GetClampedToMaxSize(1.f);
-	}
+			if (bNormalize)
+			{
+				LinearVelocity = LinearVelocity.GetClampedToMaxSize(1.f);
+			}
 
 	FFeatureVectorHelper::EncodeVector(SearchContext.EditFeatureVector(), ChannelDataOffset, LinearVelocity, ComponentStripping);
 }
diff --git a/Engine/Plugins/Experimental/Animation/PoseSearch/Source/Runtime/Private/PoseSearchHistory.cpp b/Engine/Plugins/Experimental/Animation/PoseSearch/Source/Runtime/Private/PoseSearchHistory.cpp
--- a/Engine/Plugins/Experimental/Animation/PoseSearch/Source/Runtime/Private/PoseSearchHistory.cpp
+++ b/Engine/Plugins/Experimental/Animation/PoseSearch/Source/Runtime/Private/PoseSearchHistory.cpp
@@ -12,13 +12,17 @@
 #include "PoseSearch/PoseSearchResult.h"
 #include "PoseSearch/PoseSearchDatabase.h"
 #include "PoseSearch/PoseSearchDefines.h"
-#include "PoseSearch/PoseSearchTrajectoryTypes.h"
 
 IMPLEMENT_ANIMGRAPH_MESSAGE(UE::PoseSearch::IPoseHistoryProvider);
 
 namespace UE::PoseSearch
 {
 
+#if ENABLE_DRAW_DEBUG && ENABLE_ANIM_DEBUG
+TAutoConsoleVariable<bool> CVarAnimPoseHistoryDebugDrawPose(TEXT("a.AnimNode.PoseHistory.DebugDrawPose"), false, TEXT("Enable / Disable Pose History Pose DebugDraw"));
+TAutoConsoleVariable<bool> CVarAnimPoseHistoryDebugDrawTrajectory(TEXT("a.AnimNode.PoseHistory.DebugDrawTrajectory"), false, TEXT("Enable / Disable Pose History Trajectory DebugDraw"));
+#endif
+
 /**
 * Algo::LowerBound adapted to TIndexedContainerIterator for use with indexable but not necessarily contiguous containers. Used here with TRingBuffer.
 *
@@ -174,6 +178,8 @@ FBoneIndexType FPoseHistory::GetRemappedBoneIndexType(FBoneIndexType BoneIndexTy
 
 FComponentSpaceTransformIndex FPoseHistory::GetRemappedComponentSpaceTransformIndex(const USkeleton* BoneIndexSkeleton, const USkeleton* LastUpdateSkeleton, const FBoneToTransformMap& BoneToTransformMap, FBoneIndexType BoneIndexType, bool& bSuccess)
 {
+	check(BoneIndexType != WorldSpaceIndexType);
+
 	FComponentSpaceTransformIndex BoneTransformIndex = FComponentSpaceTransformIndex(BoneIndexType);
 	if (BoneIndexType != ComponentSpaceIndexType)
 	{
@@ -242,43 +248,64 @@ bool FPoseHistory::LerpEntries(float Time, bool bExtrapolate, const FPoseHistory
 
 bool FPoseHistory::GetTransformAtTime(float Time, FTransform& OutBoneTransform, const USkeleton* BoneIndexSkeleton, FBoneIndexType BoneIndexType, FBoneIndexType ReferenceBoneIndexType, bool bExtrapolate) const
 {
-	static_assert(RootBoneIndexType == 0 && ComponentSpaceIndexType == FBoneIndexType(-1)); // some assumptions
-	check(BoneIndexType != ComponentSpaceIndexType);
+	static_assert(RootBoneIndexType == 0 && ComponentSpaceIndexType == FBoneIndexType(-1) && WorldSpaceIndexType == FBoneIndexType(-2)); // some assumptions
+	check(BoneIndexType != ComponentSpaceIndexType && BoneIndexType != WorldSpaceIndexType);
 	
-	const int32 Num = Entries.Num();
+	bool bSuccess = false;
 	
-	if (Num > 0)
+	const bool bApplyComponentToWorld = ReferenceBoneIndexType == WorldSpaceIndexType;
+	FTransform ComponentToWorld = FTransform::Identity;
+	if (bApplyComponentToWorld)
+	{
+		ComponentToWorld = Trajectory.GetSampleAtTime(Time, bExtrapolate).GetTransform();
+		ReferenceBoneIndexType = ComponentSpaceIndexType;
+	}
+
+	const int32 NumEntries = Entries.Num();
+	if (NumEntries > 0)
 	{
 		const float SecondsAgo = -Time;
 
 		int32 NextIdx = 0;
 		int32 PrevIdx = 0;
 
-		if (Num > 1)
+		if (NumEntries > 1)
 		{
 			const int32 LowerBoundIdx = LowerBound(Entries.begin(), Entries.end(), SecondsAgo, [](const FPoseHistoryEntry& Entry, float Value) { return Value < Entry.Time; });
-			NextIdx = FMath::Clamp(LowerBoundIdx, 1, Num - 1);
+			NextIdx = FMath::Clamp(LowerBoundIdx, 1, NumEntries - 1);
 			PrevIdx = NextIdx - 1;
 		}
 	
 		const FPoseHistoryEntry& PrevEntry = Entries[PrevIdx];
 		const FPoseHistoryEntry& NextEntry = Entries[NextIdx];
 
-		return LerpEntries(Time, bExtrapolate, PrevEntry, NextEntry, BoneIndexSkeleton, GetLastUpdateSkeleton(), BoneToTransformMap, BoneIndexType, ReferenceBoneIndexType, OutBoneTransform);
+		bSuccess = LerpEntries(Time, bExtrapolate, PrevEntry, NextEntry, BoneIndexSkeleton, GetLastUpdateSkeleton(), BoneToTransformMap, BoneIndexType, ReferenceBoneIndexType, OutBoneTransform);
+		if (bApplyComponentToWorld)
+		{
+			OutBoneTransform *= ComponentToWorld;
+		}
+	}
+	else
+	{
+		OutBoneTransform = ComponentToWorld;
 	}
 	
-	OutBoneTransform = FTransform::Identity;
-	return false;
+	return bSuccess;
 }
 
-bool FPoseHistory::IsEmpty() const
+void FPoseHistory::UpdateTrajectory(const FPoseSearchQueryTrajectory& InTrajectory, float InTrajectorySpeedMultiplier)
 {
-	return Entries.IsEmpty();
-}
+	TrajectorySpeedMultiplier = InTrajectorySpeedMultiplier;
+	Trajectory = InTrajectory;
 
-void FPoseHistory::ClearHistory()
-{
-	Entries.Reset();
+	if (!FMath::IsNearlyEqual(TrajectorySpeedMultiplier, 1.f))
+	{
+		const float TrajectorySpeedMultiplierInv = FMath::IsNearlyZero(TrajectorySpeedMultiplier) ? 1.f : 1.f / TrajectorySpeedMultiplier;
+		for (FPoseSearchQueryTrajectorySample& Sample : Trajectory.Samples)
+		{
+			Sample.AccumulatedSeconds *= TrajectorySpeedMultiplierInv;
+		}
+	}
 }
 
 void FPoseHistory::Update(float SecondsElapsed, FCSPose<FCompactPose>& ComponentSpacePose, bool bStoreScales)
@@ -314,41 +341,36 @@ void FPoseHistory::Update(float SecondsElapsed, FCSPose<FCompactPose>& Component
 }
 
 #if ENABLE_DRAW_DEBUG && ENABLE_ANIM_DEBUG
-void FPoseHistory::DebugDraw(FAnimInstanceProxy& AnimInstanceProxy, FColor Color, const FPoseSearchQueryTrajectory* Trajectory) const
+void FPoseHistory::DebugDraw(FAnimInstanceProxy& AnimInstanceProxy, FColor Color) const
+{
+	if (CVarAnimPoseHistoryDebugDrawTrajectory.GetValueOnAnyThread())
+	{
+		Trajectory.DebugDrawTrajectory(AnimInstanceProxy);
+	}
+
+	if (Color.A > 0 && CVarAnimPoseHistoryDebugDrawPose.GetValueOnAnyThread())
 {
-	const bool bValidTrajectory = Trajectory && !Trajectory->Samples.IsEmpty();
+		const bool bValidTrajectory = !Trajectory.Samples.IsEmpty();
+		TArray<FTransform, TInlineAllocator<128>> PrevGlobalTransforms;
 
-	TArray<FTransform> PrevGlobalTransforms;
 	for (int32 EntryIndex = 0; EntryIndex < Entries.Num(); ++EntryIndex)
 	{
 		const FPoseHistoryEntry& Entry = Entries[EntryIndex];
-		if (Entry.Num() == 0)
-		{
-			PrevGlobalTransforms.Reset();
-		}
-		else if (PrevGlobalTransforms.Num() != Entry.Num())
-		{
-			PrevGlobalTransforms.SetNum(Entry.Num());
-			for (int32 i = 0; i < Entry.Num(); ++i)
-			{
-				const FTransform RootTransform = bValidTrajectory ? Trajectory->GetSampleAtTime(-Entry.Time).GetTransform() : AnimInstanceProxy.GetComponentTransform();
 
-				AnimInstanceProxy.AnimDrawDebugPoint(RootTransform.GetTranslation(), 6.f, FColor::Blue, false, 0.f, ESceneDepthPriorityGroup::SDPG_Foreground);
+			const int32 PrevGlobalTransformsNum = PrevGlobalTransforms.Num();
+			const int32 Max = FMath::Max(PrevGlobalTransformsNum, Entry.Num());
+
+			PrevGlobalTransforms.SetNum(Max, false);
 
-				PrevGlobalTransforms[i] = Entry.GetComponentSpaceTransform(i) * RootTransform;
-			}
-		}
-		else
-		{
 			for (int32 i = 0; i < Entry.Num(); ++i)
 			{
-				const FTransform RootTransform = bValidTrajectory ? Trajectory->GetSampleAtTime(-Entry.Time).GetTransform() : AnimInstanceProxy.GetComponentTransform();
-
-				AnimInstanceProxy.AnimDrawDebugPoint(RootTransform.GetTranslation(), 6.f, FColor::Blue, false, 0.f, ESceneDepthPriorityGroup::SDPG_Foreground);
-
+				const FTransform RootTransform = bValidTrajectory ? Trajectory.GetSampleAtTime(-Entry.Time).GetTransform() : AnimInstanceProxy.GetComponentTransform();
 				const FTransform GlobalTransforms = Entry.GetComponentSpaceTransform(i) * RootTransform;
 
-				AnimInstanceProxy.AnimDrawDebugLine(PrevGlobalTransforms[i].GetTranslation(), GlobalTransforms.GetTranslation(), Color, false, 0.f, ESceneDepthPriorityGroup::SDPG_Foreground);
+				if (i < PrevGlobalTransformsNum)
+				{
+					AnimInstanceProxy.AnimDrawDebugLine(PrevGlobalTransforms[i].GetTranslation(), GlobalTransforms.GetTranslation(), Color, false, 0.f, ESceneDepthPriorityGroup::SDPG_Foreground);
+				}
 
 				PrevGlobalTransforms[i] = GlobalTransforms;
 			}
@@ -359,7 +381,7 @@ void FPoseHistory::DebugDraw(FAnimInstanceProxy& AnimInstanceProxy, FColor Color
 
 //////////////////////////////////////////////////////////////////////////
 // FExtendedPoseHistory
-void FExtendedPoseHistory::Init(const FPoseHistory* InPoseHistory)
+void FExtendedPoseHistory::Init(const IPoseHistory* InPoseHistory)
 {
 	check(InPoseHistory);
 	PoseHistory = InPoseHistory;
@@ -378,80 +400,102 @@ bool FExtendedPoseHistory::GetTransformAtTime(float Time, FTransform& OutBoneTra
 		const int32 Num = FutureEntries.Num();
 		if (Num > 0)
 		{
+			const bool bApplyComponentToWorld = ReferenceBoneIndexType == WorldSpaceIndexType;
+			FTransform ComponentToWorld = FTransform::Identity;
+			if (bApplyComponentToWorld)
+			{
+				ComponentToWorld = GetTrajectory().GetSampleAtTime(Time, bExtrapolate).GetTransform();
+				ReferenceBoneIndexType = ComponentSpaceIndexType;
+			}
+
 			const float SecondsAgo = -Time;
 			const int32 LowerBoundIdx = Algo::LowerBound(FutureEntries, SecondsAgo, [](const FPoseHistoryEntry& Entry, float Value) { return Value < Entry.Time; });
 			const int32 NextIdx = FMath::Min(LowerBoundIdx, Num - 1);
-			const FPoseHistoryEntries& PastEntries = PoseHistory->GetEntries();
+			const FPoseHistoryEntries& PastEntries = GetEntries();
 			const FPoseHistoryEntry& NextEntry = FutureEntries[NextIdx];
 			const FPoseHistoryEntry& PrevEntry = NextIdx > 0 ? FutureEntries[NextIdx - 1] : !PastEntries.IsEmpty() ? PastEntries.First() : NextEntry;
 						
-			return FPoseHistory::LerpEntries(Time, bExtrapolate, PrevEntry, NextEntry, BoneIndexSkeleton, PoseHistory->GetLastUpdateSkeleton(), PoseHistory->GetBoneToTransformMap(), BoneIndexType, ReferenceBoneIndexType, OutBoneTransform);
+			const bool bSuccess = FPoseHistory::LerpEntries(Time, bExtrapolate, PrevEntry, NextEntry, BoneIndexSkeleton, GetLastUpdateSkeleton(), GetBoneToTransformMap(), BoneIndexType, ReferenceBoneIndexType, OutBoneTransform);
+			if (bApplyComponentToWorld)
+			{
+				OutBoneTransform *= ComponentToWorld;
+			}
+			return bSuccess;
 		}
 	}
 	
 	return PoseHistory->GetTransformAtTime(Time, OutBoneTransform, BoneIndexSkeleton, BoneIndexType, ReferenceBoneIndexType, bExtrapolate);
 }
 
-bool FExtendedPoseHistory::IsEmpty() const
+void FExtendedPoseHistory::AddFutureRootBone(float SecondsInTheFuture, const FTransform& FutureRootBoneTransform, bool bStoreScales)
 {
-	check(PoseHistory);
-	return PoseHistory->IsEmpty() && FutureEntries.IsEmpty();
-}
+	// we don't allow to add "past" or "present" poses to FutureEntries
+	check(SecondsInTheFuture > 0.f);
 
-void FExtendedPoseHistory::ResetFuturePoses()
-{
-	FutureEntries.Reset();
+	const float SecondsAgo = -SecondsInTheFuture;
+	const int32 LowerBoundIdx = Algo::LowerBound(FutureEntries, SecondsAgo, [](const FPoseHistoryEntry& Entry, float Value) { return Value < Entry.Time; });
+	FPoseHistoryEntry& FutureEntry = FutureEntries.InsertDefaulted_GetRef(LowerBoundIdx);
+	FutureEntry.SetNum(1, bStoreScales);
+	FutureEntry.SetComponentSpaceTransform(RootBoneIndexType, FutureRootBoneTransform);
+	FutureEntry.Time = SecondsAgo;
 }
 
-void FExtendedPoseHistory::AddFuturePose(float SecondsInTheFuture, FCSPose<FCompactPose>& ComponentSpacePose, const FTransform& ComponentTransform)
+void FExtendedPoseHistory::AddFuturePose(float SecondsInTheFuture, FCSPose<FCompactPose>& ComponentSpacePose)
 {
 	// we don't allow to add "past" or "present" poses to FutureEntries
 	check(SecondsInTheFuture > 0.f);
 	check(PoseHistory);	
 	const float SecondsAgo = -SecondsInTheFuture;
 	const int32 LowerBoundIdx = Algo::LowerBound(FutureEntries, SecondsAgo, [](const FPoseHistoryEntry& Entry, float Value) { return Value < Entry.Time; });
-	FutureEntries.InsertDefaulted_GetRef(LowerBoundIdx).Update(SecondsAgo, ComponentSpacePose, PoseHistory->GetBoneToTransformMap(), true);
+	FutureEntries.InsertDefaulted_GetRef(LowerBoundIdx).Update(SecondsAgo, ComponentSpacePose, GetBoneToTransformMap(), true);
 }
 
 #if ENABLE_DRAW_DEBUG && ENABLE_ANIM_DEBUG
-void FExtendedPoseHistory::DebugDraw(FAnimInstanceProxy& AnimInstanceProxy, FColor Color, const FPoseSearchQueryTrajectory* Trajectory) const
+void FExtendedPoseHistory::DebugDraw(FAnimInstanceProxy& AnimInstanceProxy, FColor Color) const
 {
 	check(PoseHistory);
 
-	const bool bValidTrajectory = Trajectory && !Trajectory->Samples.IsEmpty();
-
-	TArray<FTransform> PrevGlobalTransforms;
-	for (int32 EntryIndex = 0; EntryIndex < FutureEntries.Num(); ++EntryIndex)
+	if (Color.A > 0 && !FutureEntries.IsEmpty() && CVarAnimPoseHistoryDebugDrawPose.GetValueOnAnyThread())
 	{
-		const FPoseHistoryEntry& Entry = FutureEntries[EntryIndex];
-		if (Entry.Num() == 0)
-		{
-			PrevGlobalTransforms.Reset();
-		}
-		else if (PrevGlobalTransforms.Num() != Entry.Num())
+		const FPoseSearchQueryTrajectory& Trajectory = GetTrajectory();
+		const bool bValidTrajectory = !Trajectory.Samples.IsEmpty();
+		TArray<FTransform, TInlineAllocator<128>> PrevGlobalTransforms;
+
+		const FPoseHistoryEntries& PastEntries = GetEntries();
+
+		int32 EntriesNum = FutureEntries.Num();
+		if (!PastEntries.IsEmpty())
 		{
-			PrevGlobalTransforms.SetNum(Entry.Num());
-			for (int32 i = 0; i < Entry.Num(); ++i)
-			{
-				const FTransform RootTransform = bValidTrajectory ? Trajectory->GetSampleAtTime(-Entry.Time).GetTransform() : AnimInstanceProxy.GetComponentTransform();
-				PrevGlobalTransforms[i] = Entry.GetComponentSpaceTransform(i) * RootTransform;
-			}
+			// connecting the future entries with the past entries
+			++EntriesNum;
 		}
-		else
+
+		for (int32 EntryIndex = 0; EntryIndex < EntriesNum; ++EntryIndex)
 		{
+			const FPoseHistoryEntry& Entry = (EntryIndex == FutureEntries.Num()) ? PastEntries.Last() : FutureEntries[EntryIndex];
+
+			const int32 PrevGlobalTransformsNum = PrevGlobalTransforms.Num();
+			const int32 Max = FMath::Max(PrevGlobalTransformsNum, Entry.Num());
+
+			PrevGlobalTransforms.SetNum(Max, false);
+
 			for (int32 i = 0; i < Entry.Num(); ++i)
 			{
-				const FTransform RootTransform = bValidTrajectory ? Trajectory->GetSampleAtTime(-Entry.Time).GetTransform() : AnimInstanceProxy.GetComponentTransform();
+				const FTransform RootTransform = bValidTrajectory ? Trajectory.GetSampleAtTime(-Entry.Time).GetTransform() : AnimInstanceProxy.GetComponentTransform();
 				const FTransform GlobalTransforms = Entry.GetComponentSpaceTransform(i) * RootTransform;
 
-				AnimInstanceProxy.AnimDrawDebugLine(PrevGlobalTransforms[i].GetTranslation(), GlobalTransforms.GetTranslation(), Color, false, 0.f, ESceneDepthPriorityGroup::SDPG_Foreground);
+				if (i < PrevGlobalTransformsNum)
+				{
+					AnimInstanceProxy.AnimDrawDebugLine(PrevGlobalTransforms[i].GetTranslation(), GlobalTransforms.GetTranslation(), Color, false, 0.f, ESceneDepthPriorityGroup::SDPG_Foreground);
+				}
 
 				PrevGlobalTransforms[i] = GlobalTransforms;
 			}
 		}
-	}
 
-	PoseHistory->DebugDraw(AnimInstanceProxy, Color, Trajectory);
+		// no need to DebugDraw PoseHistory since it'll be drawn anyways by the history collectors
+		//PoseHistory->DebugDraw(AnimInstanceProxy, Color);
+	}
 }
 #endif // ENABLE_DRAW_DEBUG && ENABLE_ANIM_DEBUG
 
diff --git a/Engine/Plugins/Experimental/Animation/PoseSearch/Source/Runtime/Private/PoseSearchLibrary.cpp b/Engine/Plugins/Experimental/Animation/PoseSearch/Source/Runtime/Private/PoseSearchLibrary.cpp
--- a/Engine/Plugins/Experimental/Animation/PoseSearch/Source/Runtime/Private/PoseSearchLibrary.cpp
+++ b/Engine/Plugins/Experimental/Animation/PoseSearch/Source/Runtime/Private/PoseSearchLibrary.cpp
@@ -29,7 +29,6 @@
 #if ENABLE_DRAW_DEBUG && ENABLE_ANIM_DEBUG
 TAutoConsoleVariable<int32> CVarAnimMotionMatchDrawQueryEnable(TEXT("a.MotionMatch.DrawQuery.Enable"), 0, TEXT("Enable / Disable MotionMatch Draw Query"));
 TAutoConsoleVariable<int32> CVarAnimMotionMatchDrawMatchEnable(TEXT("a.MotionMatch.DrawMatch.Enable"), 0, TEXT("Enable / Disable MotionMatch Draw Match"));
-TAutoConsoleVariable<int32> CVarAnimMotionMatchDrawHistoryEnable(TEXT("a.MotionMatch.DrawHistory.Enable"), 0, TEXT("Enable / Disable MotionMatch Draw History"));
 #endif
 
 namespace UE::PoseSearch
@@ -225,9 +224,6 @@ void FMotionMatchingState::UpdateRootBoneControl(const FAnimationUpdateContext&
 
 #if UE_POSE_SEARCH_TRACE_ENABLED
 void UPoseSearchLibrary::TraceMotionMatchingState(
-	// Trajectory is the trajectory in mesh component space prior ProcessTrajectory
-	const FPoseSearchQueryTrajectory& Trajectory, 
-	// SearchContext.Trajectory is the trajectory in root bone component space after ProcessTrajectory
 	UE::PoseSearch::FSearchContext& SearchContext,
 	const UE::PoseSearch::FSearchResult& CurrentResult,
 	float ElapsedPoseSearchTime,
@@ -285,12 +281,11 @@ void UPoseSearchLibrary::TraceMotionMatchingState(
 	if (DeltaTime > SMALL_NUMBER)
 	{
 		// simulation
-		if (SearchContext.IsTrajectoryValid())
+		if (SearchContext.AnyCachedQuery())
 		{
 			const FTransform PrevRoot = SearchContext.GetWorldBoneTransformAtTime(-DeltaTime);
 			const FTransform CurrRoot = SearchContext.GetWorldBoneTransformAtTime(0.f);
 			const FTransform SimDelta = CurrRoot.GetRelativeTransform(PrevRoot);
-
 			TraceState.SimLinearVelocity = SimDelta.GetTranslation().Size() / DeltaTime;
 			TraceState.SimAngularVelocity = FMath::RadiansToDegrees(SimDelta.GetRotation().GetAngle()) / DeltaTime;
 		}
@@ -309,7 +304,14 @@ void UPoseSearchLibrary::TraceMotionMatchingState(
 	TraceState.SearchBruteForceCost = CurrentResult.BruteForcePoseCost.GetTotalCost();
 	TraceState.SearchBestPosePos = CurrentResult.BestPosePos;
 
-	TraceState.Trajectory = Trajectory;
+	if (SearchContext.GetHistory())
+	{
+		TraceState.Trajectory = SearchContext.GetHistory()->GetTrajectory();
+	}
+	else
+	{
+		TraceState.Trajectory.Samples.Reset();
+	}
 
 	TraceState.Output(AnimInstance, NodeId);
 }
@@ -318,8 +320,6 @@ void UPoseSearchLibrary::TraceMotionMatchingState(
 void UPoseSearchLibrary::UpdateMotionMatchingState(
 	const FAnimationUpdateContext& Context,
 	const TArray<TObjectPtr<const UPoseSearchDatabase>>& Databases,
-	const FPoseSearchQueryTrajectory& Trajectory,
-	float TrajectorySpeedMultiplier,
 	float BlendTime,
 	int32 MaxActiveBlends,
 	const FFloatInterval& PoseJumpThresholdTime,
@@ -333,8 +333,7 @@ void UPoseSearchLibrary::UpdateMotionMatchingState(
 	bool bShouldSearch,
 	bool bShouldUseCachedChannelData,
 	bool bDebugDrawQuery,
-	bool bDebugDrawCurResult,
-	bool bDebugDrawPoseHistory)
+	bool bDebugDrawCurResult)
 {
 	QUICK_SCOPE_CYCLE_COUNTER(STAT_PoseSearch_Update);
 
@@ -356,22 +355,27 @@ void UPoseSearchLibrary::UpdateMotionMatchingState(
 
 	InOutMotionMatchingState.bJumpedToPose = false;
 
-	const IPoseHistory* History = nullptr;
+	// used when YawFromAnimationBlendRate is greater than zero, by setting a future (YawFromAnimationTrajectoryBlendTime seconds ahead) root bone to the skeleton default
+	FExtendedPoseHistory ExtendedPoseHistory;
+	const IPoseHistory* PoseHistory = nullptr;
 	if (IPoseHistoryProvider* PoseHistoryProvider = Context.GetMessage<IPoseHistoryProvider>())
 	{
-		History = &PoseHistoryProvider->GetPoseHistory();
+		PoseHistory = &PoseHistoryProvider->GetPoseHistory();
 
-#if ENABLE_DRAW_DEBUG && ENABLE_ANIM_DEBUG
-		if (bDebugDrawPoseHistory)
+		if (YawFromAnimationBlendRate >= 0.f && YawFromAnimationTrajectoryBlendTime > 0.f)
 		{
-			History->DebugDraw(*Context.AnimInstanceProxy, FColor::Orange, &Trajectory);
-		}
+			ExtendedPoseHistory.Init(PoseHistory);
+
+			const FTransform& RefRootBoneTransform = Context.AnimInstanceProxy->GetSkeleton()->GetReferenceSkeleton().GetRefBonePose()[RootBoneIndexType];
+			ExtendedPoseHistory.AddFutureRootBone(YawFromAnimationTrajectoryBlendTime, RefRootBoneTransform, false);
+
+#if ENABLE_DRAW_DEBUG && ENABLE_ANIM_DEBUG
+			ExtendedPoseHistory.DebugDraw(*Context.AnimInstanceProxy, FColor::Emerald);
 #endif // ENABLE_DRAW_DEBUG && ENABLE_ANIM_DEBUG
-	}
 
-	// @todo: perhaps store the previous frame root bone transform to use in ProcessTrajectory
-	const FTransform& RootBoneTransform = Context.AnimInstanceProxy->GetSkeleton()->GetReferenceSkeleton().GetRefBonePose()[RootSchemaBoneIdx];
-	const FPoseSearchQueryTrajectory TrajectoryRootSpace = ProcessTrajectory(Trajectory, RootBoneTransform, InOutMotionMatchingState.ComponentDeltaYaw, YawFromAnimationTrajectoryBlendTime, TrajectorySpeedMultiplier);
+			PoseHistory = &ExtendedPoseHistory;
+		}
+	}
 
 	FMemMark Mark(FMemStack::Get());
 	const UAnimInstance* AnimInstance = Cast<const UAnimInstance>(Context.AnimInstanceProxy->GetAnimInstanceObject());
@@ -383,7 +387,7 @@ void UPoseSearchLibrary::UpdateMotionMatchingState(
 		InOutMotionMatchingState.CurrentSearchResult.Reset();
 	}
 
-	FSearchContext SearchContext(AnimInstance, History, TConstArrayView<const UAnimationAsset*>(), &TrajectoryRootSpace, 0.f,
+	FSearchContext SearchContext(AnimInstance, PoseHistory, TConstArrayView<const UAnimationAsset*>(), 0.f,
 		&InOutMotionMatchingState.PoseIndicesHistory, InOutMotionMatchingState.CurrentSearchResult, PoseJumpThresholdTime);
 
 	const bool bCanAdvance = InOutMotionMatchingState.CurrentSearchResult.CanAdvance(DeltaTime);
@@ -491,7 +495,7 @@ void UPoseSearchLibrary::UpdateMotionMatchingState(
 	}
 
 	// @todo: consider moving this into if (bSearch) to avoid calling SearchContext.GetCachedQuery if no search is required
-	InOutMotionMatchingState.UpdateWantedPlayRate(SearchContext, PlayRate, TrajectorySpeedMultiplier);
+	InOutMotionMatchingState.UpdateWantedPlayRate(SearchContext, PlayRate, PoseHistory ? PoseHistory->GetTrajectorySpeedMultiplier() : 1.f);
 
 	InOutMotionMatchingState.PoseIndicesHistory.Update(InOutMotionMatchingState.CurrentSearchResult, DeltaTime, PoseReselectHistory);
 
@@ -499,7 +503,7 @@ void UPoseSearchLibrary::UpdateMotionMatchingState(
 	// Record debugger details
 	if (IsTracing(Context))
 	{
-		TraceMotionMatchingState(Trajectory, SearchContext, InOutMotionMatchingState.CurrentSearchResult, InOutMotionMatchingState.ElapsedPoseSearchTime,
+		TraceMotionMatchingState(SearchContext, InOutMotionMatchingState.CurrentSearchResult, InOutMotionMatchingState.ElapsedPoseSearchTime,
 			InOutMotionMatchingState.RootMotionTransformDelta, Context.AnimInstanceProxy->GetAnimInstanceObject(), Context.GetCurrentNodeId(), DeltaTime, bSearch,
 			AnimInstance ? FObjectTrace::GetWorldElapsedTime(AnimInstance->GetWorld()) : 0.f);
 	}
@@ -532,50 +536,9 @@ void UPoseSearchLibrary::UpdateMotionMatchingState(
 #endif
 }
 
-// transforms Trajectory from SkeletalMeshComponent world space to root bone world space, and scale it by TrajectorySpeedMultiplier
-FPoseSearchQueryTrajectory UPoseSearchLibrary::ProcessTrajectory(const FPoseSearchQueryTrajectory& Trajectory, const FTransform& RootBoneTransform, float RootBoneDeltaYaw, float YawFromAnimationTrajectoryBlendTime, float TrajectorySpeedMultiplier)
-{
-	QUICK_SCOPE_CYCLE_COUNTER(STAT_PoseSearch_ProcessTrajectory);
-
-	const float TrajectorySpeedMultiplierInv = FMath::IsNearlyZero(TrajectorySpeedMultiplier) ? 1.f : 1.f / TrajectorySpeedMultiplier;
-
-	FPoseSearchQueryTrajectory TrajectoryRootSpace = Trajectory;
-
-	if (!RootBoneTransform.Equals(FTransform::Identity))
-	{
-		// @todo: optimize me
-		for (FPoseSearchQueryTrajectorySample& Sample : TrajectoryRootSpace.Samples)
-		{
-			Sample.SetTransform(RootBoneTransform * Sample.GetTransform());
-		}
-	}
-
-	if (!FMath::IsNearlyEqual(TrajectorySpeedMultiplierInv, 1.f))
-	{
-		for (FPoseSearchQueryTrajectorySample& Sample : TrajectoryRootSpace.Samples)
-		{
-			Sample.AccumulatedSeconds *= TrajectorySpeedMultiplierInv;
-		}
-	}
-
-	if (!FMath::IsNearlyZero(RootBoneDeltaYaw))
-	{
-		for (FPoseSearchQueryTrajectorySample& Sample : TrajectoryRootSpace.Samples)
-		{
-			const float BlendParam = YawFromAnimationTrajectoryBlendTime < UE_KINDA_SMALL_NUMBER ? 1 : FMath::Clamp(1.f - (Sample.AccumulatedSeconds - YawFromAnimationTrajectoryBlendTime) / YawFromAnimationTrajectoryBlendTime, 0.f, 1.f);
-			const FQuat RootBoneDelta(FRotator(0.f, RootBoneDeltaYaw * BlendParam, 0.f));
-			Sample.Facing = RootBoneDelta * Sample.Facing;
-		}
-	}
-
-	return TrajectoryRootSpace;
-}
-
 void UPoseSearchLibrary::MotionMatch(
 	UAnimInstance* AnimInstance,
 	const UPoseSearchDatabase* Database,
-	const FPoseSearchQueryTrajectory Trajectory,
-	float TrajectorySpeedMultiplier,
 	const FName PoseHistoryName,
 	FPoseSearchBlueprintResult& Result,
 	const UAnimationAsset* FutureAnimation,
@@ -616,10 +579,6 @@ void UPoseSearchLibrary::MotionMatch(
 		FColor HistoryCollectorColor = FColor::Red;
 #endif // ENABLE_DRAW_DEBUG && ENABLE_ANIM_DEBUG
 
-		// @todo: perhaps store the previous frame root bone transform to use in ProcessTrajectory
-		const FTransform& RootBoneTransform = AnimInstance->CurrentSkeleton->GetReferenceSkeleton().GetRefBonePose()[RootSchemaBoneIdx];
-		const FPoseSearchQueryTrajectory TrajectoryRootSpace = ProcessTrajectory(Trajectory, RootBoneTransform, 0.f, 0.f, TrajectorySpeedMultiplier);
-
 		// ExtendedPoseHistory will hold future poses to match AssetSamplerBase (at FutureAnimationStartTime) TimeToFutureAnimationStart seconds in the future
 		FExtendedPoseHistory ExtendedPoseHistory;
 		if (IAnimClassInterface* AnimBlueprintClass = IAnimClassInterface::GetFromClass(AnimInstance->GetClass()))
@@ -629,6 +588,7 @@ void UPoseSearchLibrary::MotionMatch(
 				if (const FAnimNode_PoseSearchHistoryCollector_Base* PoseHistoryNode = TagSubsystem->FindNodeByTag<FAnimNode_PoseSearchHistoryCollector_Base>(PoseHistoryName, AnimInstance))
 				{
 					ExtendedPoseHistory.Init(&PoseHistoryNode->GetPoseHistory());
+
 #if ENABLE_DRAW_DEBUG && ENABLE_ANIM_DEBUG && WITH_EDITORONLY_DATA
 					HistoryCollectorColor = PoseHistoryNode->DebugColor.ToFColor(true);
 #endif // ENABLE_DRAW_DEBUG && ENABLE_ANIM_DEBUG && WITH_EDITORONLY_DATA
@@ -679,25 +639,18 @@ void UPoseSearchLibrary::MotionMatch(
 				FCSPose<FCompactPose> ComponentSpacePose;
 				ComponentSpacePose.InitPose(Pose);
 
-				const FPoseSearchQueryTrajectorySample TrajectorySample = TrajectoryRootSpace.GetSampleAtTime(ExtractionTime);
-				const FTransform& ComponentTransform = AnimInstance->GetOwningComponent()->GetComponentTransform();
-				const FTransform FutureComponentTransform = TrajectorySample.GetTransform() * ComponentTransform;
-
-				ExtendedPoseHistory.AddFuturePose(FutureAnimationTime, ComponentSpacePose, FutureComponentTransform);
+				ExtendedPoseHistory.AddFuturePose(FutureAnimationTime, ComponentSpacePose);
 			}
 
 #if ENABLE_DRAW_DEBUG && ENABLE_ANIM_DEBUG
-			if (CVarAnimMotionMatchDrawHistoryEnable.GetValueOnAnyThread())
-			{
 				if (FAnimInstanceProxy* AnimInstanceProxy = UAnimInstanceProxyProvider::GetAnimInstanceProxy(AnimInstance))
 				{
-					ExtendedPoseHistory.DebugDraw(*AnimInstanceProxy, HistoryCollectorColor, &TrajectoryRootSpace);
+				ExtendedPoseHistory.DebugDraw(*AnimInstanceProxy, HistoryCollectorColor);
 				}
-			}
 #endif // ENABLE_DRAW_DEBUG && ENABLE_ANIM_DEBUG
 		}
 
-		FSearchContext SearchContext(AnimInstance, ExtendedPoseHistory.IsInitialized() ? &ExtendedPoseHistory : nullptr, TConstArrayView<const UAnimationAsset*>(), &TrajectoryRootSpace, TimeToFutureAnimationStart);
+		FSearchContext SearchContext(AnimInstance, ExtendedPoseHistory.IsInitialized() ? &ExtendedPoseHistory : nullptr, TConstArrayView<const UAnimationAsset*>(), TimeToFutureAnimationStart);
 
 		FSearchResult SearchResult = Database->Search(SearchContext);
 		if (SearchResult.IsValid())
@@ -735,7 +688,7 @@ void UPoseSearchLibrary::MotionMatch(
 
 #if UE_POSE_SEARCH_TRACE_ENABLED
 		
-		TraceMotionMatchingState(Trajectory, SearchContext, SearchResult, 0.f, FTransform::Identity, AnimInstance, DebugSessionUniqueIdentifier,
+		TraceMotionMatchingState(SearchContext, SearchResult, 0.f, FTransform::Identity, AnimInstance, DebugSessionUniqueIdentifier,
 			AnimInstance->GetDeltaSeconds(), true, FObjectTrace::GetWorldElapsedTime(AnimInstance->GetWorld()));
 #endif // UE_POSE_SEARCH_TRACE_ENABLED
 	}
@@ -826,7 +779,7 @@ UE::PoseSearch::FSearchResult UPoseSearchLibrary::MotionMatch(const FAnimationBa
 #if UE_POSE_SEARCH_TRACE_ENABLED
 		const float SearchBestCost = SearchResult.PoseCost.GetTotalCost();
 		const float SearchBruteForceCost = SearchResult.BruteForcePoseCost.GetTotalCost();
-		TraceMotionMatchingState(FPoseSearchQueryTrajectory(), SearchContext, SearchResult, 0.f, FTransform::Identity, AnimInstance, Context.GetCurrentNodeId(),
+		TraceMotionMatchingState(SearchContext, SearchResult, 0.f, FTransform::Identity, AnimInstance, Context.GetCurrentNodeId(),
 			AnimInstance->GetDeltaSeconds(), true, FObjectTrace::GetWorldElapsedTime(AnimInstance->GetWorld()));
 #endif // UE_POSE_SEARCH_TRACE_ENABLED
 	}
diff --git a/Engine/Plugins/Experimental/Animation/PoseSearch/Source/Runtime/Private/PoseSearchTrajectoryTypes.cpp b/Engine/Plugins/Experimental/Animation/PoseSearch/Source/Runtime/Private/PoseSearchTrajectoryTypes.cpp
--- a/Engine/Plugins/Experimental/Animation/PoseSearch/Source/Runtime/Private/PoseSearchTrajectoryTypes.cpp
+++ b/Engine/Plugins/Experimental/Animation/PoseSearch/Source/Runtime/Private/PoseSearchTrajectoryTypes.cpp
@@ -1,7 +1,7 @@
 // Copyright Epic Games, Inc. All Rights Reserved.
 
 #include "PoseSearch/PoseSearchTrajectoryTypes.h"
-
+#include "Animation/AnimInstanceProxy.h"
 #include "Animation/AnimTypes.h"
 #include "Components/SceneComponent.h"
 #include "DrawDebugHelpers.h"
@@ -81,6 +81,28 @@ void FPoseSearchQueryTrajectory::DebugDrawTrajectory(const UWorld* World, float
 		}
 	}
 }
+
+void FPoseSearchQueryTrajectory::DebugDrawTrajectory(FAnimInstanceProxy& AnimInstanceProxy, float HeightOffset) const
+{
+	FVector OffsetVector = FVector::UpVector * HeightOffset;
+
+	const int32 LastIndex = Samples.Num() - 1;
+	if (LastIndex >= 0)
+	{
+		for (int32 Index = 0; ; ++Index)
+		{
+			AnimInstanceProxy.AnimDrawDebugSphere(Samples[Index].Position + OffsetVector, 2.f /*Radius*/, 4 /*Segments*/, FColor::Black);
+			AnimInstanceProxy.AnimDrawDebugCoordinateSystem(Samples[Index].Position + OffsetVector, FRotator(Samples[Index].Facing), 12.f /*Scale*/);
+
+			if (Index == LastIndex)
+			{
+				break;
+			}
+			
+			AnimInstanceProxy.AnimDrawDebugLine(Samples[Index].Position + OffsetVector, Samples[Index + 1].Position + OffsetVector, FColor::Black);
+		}
+	}
+}
 #endif // ENABLE_ANIM_DEBUG
 
 FArchive& operator<<(FArchive& Ar, FPoseSearchQueryTrajectorySample& TrajectorySample)
diff --git a/Engine/Plugins/Experimental/Animation/PoseSearch/Source/Runtime/Public/PoseSearch/AnimNode_MotionMatching.h b/Engine/Plugins/Experimental/Animation/PoseSearch/Source/Runtime/Public/PoseSearch/AnimNode_MotionMatching.h
--- a/Engine/Plugins/Experimental/Animation/PoseSearch/Source/Runtime/Public/PoseSearch/AnimNode_MotionMatching.h
+++ b/Engine/Plugins/Experimental/Animation/PoseSearch/Source/Runtime/Public/PoseSearch/AnimNode_MotionMatching.h
@@ -6,7 +6,6 @@
 #include "GameplayTagContainer.h"
 #include "BlendStack/AnimNode_BlendStack.h"
 #include "PoseSearch/PoseSearchLibrary.h"
-#include "PoseSearch/PoseSearchTrajectoryTypes.h"
 #include "AnimNode_MotionMatching.generated.h"
 
 class UPoseSearchDatabase;
@@ -52,14 +51,6 @@ private:
 	UPROPERTY(EditAnywhere, Category = Settings, meta = (PinShownByDefault))
 	TObjectPtr<const UPoseSearchDatabase> Database = nullptr;
 
-	// Motion Trajectory samples for pose search queries in Motion Matching. These are expected to be in the world space of the SkeletalMeshComponent. This is provided with the CharacterMovementTrajectory Component output.
-	UPROPERTY(EditAnywhere, Category = Settings, meta = (PinShownByDefault))
-	FPoseSearchQueryTrajectory Trajectory;
-
-	// Input Trajectory velocity will be multiplied by TrajectorySpeedMultiplier: values below 1 will result in selecting animation slower than requested from the original Trajectory
-	UPROPERTY(EditAnywhere, Category = Settings, meta = (PinHiddenByDefault, ClampMin="0"))
-	float TrajectorySpeedMultiplier = 1.f;
-
 	// Time in seconds to blend out to the new pose. Uses either inertial blending, requiring an Inertialization node after this node, or the internal blend stack, if MaxActiveBlends is greater than zero.
 	UPROPERTY(EditAnywhere, Category = Settings, meta = (PinHiddenByDefault, ClampMin="0"))
 	float BlendTime = 0.2f;
@@ -106,14 +97,14 @@ private:
 	bool bShouldUseCachedChannelData = false;
 	
 	// blend time over which the yaw from the animation is distributed across the trajectory samples (negative values implies yaw from the animation is constant over the entire trajectory, so the trajectory will not try to recover towards the capsule orientation)
-	UPROPERTY(EditAnywhere, Category = RootMotion, meta = (PinHiddenByDefault))
+	UPROPERTY(EditAnywhere, Category = "Root Motion (Experimental)", meta = (PinHiddenByDefault))
 	float YawFromAnimationTrajectoryBlendTime = 0.1f;
 
 	// rate at which the root bone orientation catches up to the capsule orientation after being controlled by animation
 	// (negative values mean the capsule is authoritative over the root bone orientation and the root bone is always synchronized with the capsule,
 	// 0 means the orientation will be fully controlled by animation, and potentially never converge over the capsule orientation,
 	// positive values represent the rate at which the orientation drifts towards the capsule orientation after being controlled by animation)
-	UPROPERTY(EditAnywhere, Category = RootMotion, meta = (PinHiddenByDefault))
+	UPROPERTY(EditAnywhere, Category = "Root Motion (Experimental)", meta = (PinHiddenByDefault))
 	float YawFromAnimationBlendRate = -1.f;
 
 	// Encapsulated motion matching algorithm and internal state
diff --git a/Engine/Plugins/Experimental/Animation/PoseSearch/Source/Runtime/Public/PoseSearch/AnimNode_PoseSearchHistoryCollector.h b/Engine/Plugins/Experimental/Animation/PoseSearch/Source/Runtime/Public/PoseSearch/AnimNode_PoseSearchHistoryCollector.h
--- a/Engine/Plugins/Experimental/Animation/PoseSearch/Source/Runtime/Public/PoseSearch/AnimNode_PoseSearchHistoryCollector.h
+++ b/Engine/Plugins/Experimental/Animation/PoseSearch/Source/Runtime/Public/PoseSearch/AnimNode_PoseSearchHistoryCollector.h
@@ -45,16 +45,20 @@ public:
 	FLinearColor DebugColor = FLinearColor::Red;
 #endif // WITH_EDITORONLY_DATA
 
-	// Optional Trajectory for debug drawing purposes only: if not provided, all the collected pose transforms will be drawn using the Graph context
-	UPROPERTY(EditAnywhere, Category = Debug, meta = (PinHiddenByDefault))
-	FPoseSearchQueryTrajectory DebugDrawTrajectory;
+	// Trajectory samples for pose search queries in Motion Matching. These are expected to be in the world space of the SkeletalMeshComponent. This is provided with the CharacterMovementTrajectory Component output.
+	UPROPERTY(EditAnywhere, Category = Settings, meta = (PinShownByDefault))
+	FPoseSearchQueryTrajectory Trajectory;
+
+	// Input Trajectory velocity will be multiplied by TrajectorySpeedMultiplier: values below 1 will result in selecting animation slower than requested from the original Trajectory
+	UPROPERTY(EditAnywhere, Category = Settings, meta = (PinHiddenByDefault, ClampMin="0"))
+	float TrajectorySpeedMultiplier = 1.f;
 
 	// FAnimNode_Base interface
 	virtual void CacheBones_AnyThread(const FAnimationCacheBonesContext& Context) override;
 	virtual void Update_AnyThread(const FAnimationUpdateContext& Context) override;
 	// End of FAnimNode_Base interface
 
-	const UE::PoseSearch::FPoseHistory& GetPoseHistory() const { return PoseHistory; }
+	const UE::PoseSearch::IPoseHistory& GetPoseHistory() const { return PoseHistory; }
 
 protected:
 	UE::PoseSearch::FPoseHistory PoseHistory;
diff --git a/Engine/Plugins/Experimental/Animation/PoseSearch/Source/Runtime/Public/PoseSearch/PoseSearchContext.h b/Engine/Plugins/Experimental/Animation/PoseSearch/Source/Runtime/Public/PoseSearch/PoseSearchContext.h
--- a/Engine/Plugins/Experimental/Animation/PoseSearch/Source/Runtime/Public/PoseSearch/PoseSearchContext.h
+++ b/Engine/Plugins/Experimental/Animation/PoseSearch/Source/Runtime/Public/PoseSearch/PoseSearchContext.h
@@ -170,7 +170,7 @@ struct FCachedChannel
 struct POSESEARCH_API FSearchContext
 {
 	FSearchContext(const UAnimInstance* InAnimInstance, const IPoseHistory* InHistory, TConstArrayView<const UAnimationAsset*> InAnimationsToConsider = TConstArrayView<const UAnimationAsset*>(),
-		const FPoseSearchQueryTrajectory* InTrajectory = nullptr, float InDesiredPermutationTimeOffset = 0.f, const FPoseIndicesHistory* InPoseIndicesHistory = nullptr,
+		float InDesiredPermutationTimeOffset = 0.f, const FPoseIndicesHistory* InPoseIndicesHistory = nullptr,
 		const FSearchResult& InCurrentResult = FSearchResult(), const FFloatInterval& InPoseJumpThresholdTime = FFloatInterval(0.f, 0.f), bool bInUseCachedChannelData = false);
 
 	// Returns the rotation of the bone Schema.BoneReferences[SchemaSampleBoneIdx] at an offset time of SampleTimeOffset relative to the
@@ -204,8 +204,7 @@ struct POSESEARCH_API FSearchContext
 	const FSearchResult& GetCurrentResult() const { return CurrentResult; }
 	const FFloatInterval& GetPoseJumpThresholdTime() const { return PoseJumpThresholdTime; }
 	const FPoseIndicesHistory* GetPoseIndicesHistory() const { return PoseIndicesHistory; }
-	bool IsHistoryValid() const { return History != nullptr; }
-	bool IsTrajectoryValid() const { return Trajectory != nullptr; }
+	const IPoseHistory* GetHistory() const { return History; }
 	float GetDesiredPermutationTimeOffset() const { return DesiredPermutationTimeOffset; }
 	const UAnimInstance* GetAnimInstance() const { return AnimInstance; }
 
@@ -221,6 +220,7 @@ struct POSESEARCH_API FSearchContext
 	bool IsAsyncBuildIndexInProgress() const { return bAsyncBuildIndexInProgress; }
 #endif // WITH_EDITOR
 
+	bool AnyCachedQuery() const { return !CachedQueries.IsEmpty(); }
 	void AddNewFeatureVectorBuilder(const UPoseSearchSchema* Schema) { CachedQueries.Emplace(Schema); }
 	TArrayView<float> EditFeatureVector();
 	
@@ -239,8 +239,6 @@ private:
 	// if AnimationsToConsider is not empty, we'll search only for poses from UAnimationAsset(s) that are in the AnimationsToConsider
 	TConstArrayView<const UAnimationAsset*> AnimationsToConsider;
 
-	// Trajectory has been transformed in root bone world space reference system
-	const FPoseSearchQueryTrajectory* Trajectory = nullptr;
 	const float DesiredPermutationTimeOffset = 0.f;
 	const FPoseIndicesHistory* PoseIndicesHistory = nullptr;
 	const FSearchResult& CurrentResult;
diff --git a/Engine/Plugins/Experimental/Animation/PoseSearch/Source/Runtime/Public/PoseSearch/PoseSearchDefines.h b/Engine/Plugins/Experimental/Animation/PoseSearch/Source/Runtime/Public/PoseSearch/PoseSearchDefines.h
--- a/Engine/Plugins/Experimental/Animation/PoseSearch/Source/Runtime/Public/PoseSearch/PoseSearchDefines.h
+++ b/Engine/Plugins/Experimental/Animation/PoseSearch/Source/Runtime/Public/PoseSearch/PoseSearchDefines.h
@@ -26,6 +26,7 @@ namespace UE::PoseSearch
 static constexpr int8 RootSchemaBoneIdx = 0;
 static constexpr FBoneIndexType RootBoneIndexType = 0;
 static constexpr FBoneIndexType ComponentSpaceIndexType = FBoneIndexType(-1);
+static constexpr FBoneIndexType WorldSpaceIndexType = FBoneIndexType(-2);
 
 static constexpr int32 PreallocatedCachedQueriesNum = 8;
 static constexpr int32 PreallocatedCachedChannelDataNum = 64;
diff --git a/Engine/Plugins/Experimental/Animation/PoseSearch/Source/Runtime/Public/PoseSearch/PoseSearchHistory.h b/Engine/Plugins/Experimental/Animation/PoseSearch/Source/Runtime/Public/PoseSearch/PoseSearchHistory.h
--- a/Engine/Plugins/Experimental/Animation/PoseSearch/Source/Runtime/Public/PoseSearch/PoseSearchHistory.h
+++ b/Engine/Plugins/Experimental/Animation/PoseSearch/Source/Runtime/Public/PoseSearch/PoseSearchHistory.h
@@ -7,36 +7,26 @@
 #include "Containers/RingBuffer.h"
 #include "DrawDebugHelpers.h"
 #include "PoseSearch/PoseSearchDefines.h"
+#include "PoseSearch/PoseSearchTrajectoryTypes.h"
 #include "UObject/ObjectKey.h"
 
 struct FAnimInstanceProxy;
-struct FPoseSearchQueryTrajectory;
 class USkeleton;
 class UWorld;
 
 namespace UE::PoseSearch
 {
 
+#if ENABLE_DRAW_DEBUG && ENABLE_ANIM_DEBUG
+extern POSESEARCH_API TAutoConsoleVariable<bool> CVarAnimPoseHistoryDebugDrawPose;
+extern POSESEARCH_API TAutoConsoleVariable<bool> CVarAnimPoseHistoryDebugDrawTrajectory;
+#endif
+
 struct FSearchResult;
 typedef uint16 FComponentSpaceTransformIndex;
 typedef TPair<FBoneIndexType, FComponentSpaceTransformIndex> FBoneToTransformPair;
 typedef TMap<FBoneIndexType, FComponentSpaceTransformIndex> FBoneToTransformMap;
 
-struct POSESEARCH_API IPoseHistory
-{
-	virtual ~IPoseHistory() {}
-	
-	// returns the BoneIndexType transform relative to ReferenceBoneIndexType: 
-	// if ReferenceBoneIndexType is 0 (RootBoneIndexType), OutBoneTransform is in root bone space
-	// if ReferenceBoneIndexType is FBoneIndexType(-1) (ComponentSpaceIndexType), OutBoneTransform is in component space
-	virtual bool GetTransformAtTime(float Time, FTransform& OutBoneTransform, const USkeleton* BoneIndexSkeleton = nullptr, FBoneIndexType BoneIndexType = RootBoneIndexType, FBoneIndexType ReferenceBoneIndexType = ComponentSpaceIndexType, bool bExtrapolate = true) const = 0;
-	virtual bool IsEmpty() const = 0;
-#if ENABLE_DRAW_DEBUG && ENABLE_ANIM_DEBUG
-	// debug draw the pose history. If Trajectory is not provided all the poses will be drawn at the AnimInstanceProxy transform
-	virtual void DebugDraw(FAnimInstanceProxy& AnimInstanceProxy, FColor Color, const FPoseSearchQueryTrajectory* Trajectory = nullptr) const = 0;
-#endif
-};
-
 struct FPoseHistoryEntry
 {
 	// collected bones transforms in component space
@@ -57,25 +47,49 @@ struct FPoseHistoryEntry
 typedef TRingBuffer<FPoseHistoryEntry> FPoseHistoryEntries;
 typedef TArray<FPoseHistoryEntry> FPoseHistoryFutureEntries;
 
+struct POSESEARCH_API IPoseHistory
+{
+public:
+	virtual ~IPoseHistory() {}
+	
+	// returns the BoneIndexType transform relative to ReferenceBoneIndexType: 
+	// if ReferenceBoneIndexType is 0 (RootBoneIndexType), OutBoneTransform is in root bone space
+	// if ReferenceBoneIndexType is FBoneIndexType(-1) (ComponentSpaceIndexType), OutBoneTransform is in component space
+	// if ReferenceBoneIndexType is FBoneIndexType(-2) (WorldSpaceIndexType), OutBoneTransform is in world space
+	virtual bool GetTransformAtTime(float Time, FTransform& OutBoneTransform, const USkeleton* BoneIndexSkeleton = nullptr, FBoneIndexType BoneIndexType = RootBoneIndexType, FBoneIndexType ReferenceBoneIndexType = ComponentSpaceIndexType, bool bExtrapolate = true) const = 0;
+	virtual const FPoseSearchQueryTrajectory& GetTrajectory() const = 0;
+	
+	// @todo: deprecate this API. TrajectorySpeedMultiplier should be a global query scaling value passed as input parameter of FSearchContext during config BuildQuery
+	virtual float GetTrajectorySpeedMultiplier() const = 0;
+	virtual bool IsEmpty() const = 0;
+
+	virtual const FBoneToTransformMap& GetBoneToTransformMap() const = 0;
+	virtual const FPoseHistoryEntries& GetEntries() const = 0;
+	virtual const USkeleton* GetLastUpdateSkeleton() const = 0;
+};
+
 struct FPoseHistory : public IPoseHistory
 {
 	void Init(int32 InNumPoses, float InSamplingInterval, const TArray<FBoneIndexType>& RequiredBones);
 	void Update(float SecondsElapsed, FCSPose<FCompactPose>& ComponentSpacePose, bool bStoreScales);
 
-	const FBoneToTransformMap& GetBoneToTransformMap() const { return BoneToTransformMap; }
-	const FPoseHistoryEntries& GetEntries() const { return Entries; }
-	const USkeleton* GetLastUpdateSkeleton() const { return LastUpdateSkeleton.Get(); }
+#if ENABLE_DRAW_DEBUG && ENABLE_ANIM_DEBUG
+	void DebugDraw(FAnimInstanceProxy& AnimInstanceProxy, FColor Color) const;
+#endif
+
+	void ClearHistory() { Entries.Reset(); }
+	void UpdateTrajectory(const FPoseSearchQueryTrajectory& InTrajectory, float InTrajectorySpeedMultiplier);
 
 	// IPoseHistory interface
 	virtual bool GetTransformAtTime(float Time, FTransform& OutBoneTransform, const USkeleton* BoneIndexSkeleton = nullptr, FBoneIndexType BoneIndexType = RootBoneIndexType, FBoneIndexType ReferenceBoneIndexType = ComponentSpaceIndexType, bool bExtrapolate = true) const override;
-	virtual bool IsEmpty() const override;
-#if ENABLE_DRAW_DEBUG && ENABLE_ANIM_DEBUG
-	virtual void DebugDraw(FAnimInstanceProxy& AnimInstanceProxy, FColor Color, const FPoseSearchQueryTrajectory* Trajectory = nullptr) const override;
-#endif
+	virtual const FPoseSearchQueryTrajectory& GetTrajectory() const override { return Trajectory; }
+	virtual float GetTrajectorySpeedMultiplier() const override { return TrajectorySpeedMultiplier; }
+	virtual bool IsEmpty() const override { return Entries.IsEmpty(); }
+	virtual const FBoneToTransformMap& GetBoneToTransformMap() const override { return BoneToTransformMap; }
+	virtual const FPoseHistoryEntries& GetEntries() const override { return Entries; }
+	virtual const USkeleton* GetLastUpdateSkeleton() const override { return LastUpdateSkeleton.Get(); }
 	// End of IPoseHistory interface
 
-	void ClearHistory();
-
 	static FBoneIndexType GetRemappedBoneIndexType(FBoneIndexType BoneIndexType, const USkeleton* BoneIndexSkeleton, const USkeleton* LastUpdateSkeleton);
 	static FComponentSpaceTransformIndex GetRemappedComponentSpaceTransformIndex(const USkeleton* BoneIndexSkeleton, const USkeleton* LastUpdateSkeleton, const FBoneToTransformMap& BoneToTransformMap, FBoneIndexType BoneIndexType, bool& bSuccess);
 	static bool LerpEntries(float Time, bool bExtrapolate, const FPoseHistoryEntry& PrevEntry, const FPoseHistoryEntry& NextEntry, const USkeleton* BoneIndexSkeleton, const USkeleton* LastUpdateSkeleton, const FBoneToTransformMap& BoneToTransformMap, FBoneIndexType BoneIndexType, FBoneIndexType ReferenceBoneIndexType, FTransform& OutBoneTransform);
@@ -90,27 +104,39 @@ private:
 	// ring buffer of collected bones
 	FPoseHistoryEntries Entries;
 	float SamplingInterval = 0.f;
+
+	// @todo: embed past Trajectory into Entries to have consistent timing of poses and trajectory. necessary during trajectory discontinuities
+	FPoseSearchQueryTrajectory Trajectory;
+
+	// @todo: deprecate this member and expose it via blue print logic or as global query scaling multiplier
+	float TrajectorySpeedMultiplier = 1.f;
 };
 
 struct FExtendedPoseHistory : public IPoseHistory
 {
-	void Init(const FPoseHistory* InPoseHistory);
+	void Init(const IPoseHistory* InPoseHistory);
 
 	bool IsInitialized() const;
 
 	// IPoseHistory interface
 	virtual bool GetTransformAtTime(float Time, FTransform& OutBoneTransform, const USkeleton* BoneIndexSkeleton = nullptr, FBoneIndexType BoneIndexType = RootBoneIndexType, FBoneIndexType ReferenceBoneIndexType = ComponentSpaceIndexType, bool bExtrapolate = true) const override;
-	virtual bool IsEmpty() const override;
+	virtual const FPoseSearchQueryTrajectory& GetTrajectory() const override { check(PoseHistory); return PoseHistory->GetTrajectory(); }
+	virtual float GetTrajectorySpeedMultiplier() const override { check(PoseHistory); return PoseHistory->GetTrajectorySpeedMultiplier(); }
+	virtual bool IsEmpty() const override { check(PoseHistory); return PoseHistory->IsEmpty() && FutureEntries.IsEmpty(); }
+	virtual const FBoneToTransformMap& GetBoneToTransformMap() const override { check(PoseHistory); return PoseHistory->GetBoneToTransformMap(); }
+	virtual const FPoseHistoryEntries& GetEntries() const override { check(PoseHistory); return PoseHistory->GetEntries(); }
+	virtual const USkeleton* GetLastUpdateSkeleton() const override { check(PoseHistory); return PoseHistory->GetLastUpdateSkeleton(); }
+	// End of IPoseHistory interface
+	
 #if ENABLE_DRAW_DEBUG && ENABLE_ANIM_DEBUG
-	virtual void DebugDraw(FAnimInstanceProxy& AnimInstanceProxy, FColor Color, const FPoseSearchQueryTrajectory* Trajectory = nullptr) const override;	
+	void DebugDraw(FAnimInstanceProxy& AnimInstanceProxy, FColor Color) const;
 #endif
-	// End of IPoseHistory interface
 
-	void ResetFuturePoses();
-	void AddFuturePose(float SecondsInTheFuture, FCSPose<FCompactPose>& ComponentSpacePose, const FTransform& ComponentTransform);
+	void AddFutureRootBone(float SecondsInTheFuture, const FTransform& FutureRootBoneTransform, bool bStoreScales);
+	void AddFuturePose(float SecondsInTheFuture, FCSPose<FCompactPose>& ComponentSpacePose);
 
 private:
-	const FPoseHistory* PoseHistory = nullptr;
+	const IPoseHistory* PoseHistory = nullptr;
 	FPoseHistoryFutureEntries FutureEntries;
 };
 
diff --git a/Engine/Plugins/Experimental/Animation/PoseSearch/Source/Runtime/Public/PoseSearch/PoseSearchLibrary.h b/Engine/Plugins/Experimental/Animation/PoseSearch/Source/Runtime/Public/PoseSearch/PoseSearchLibrary.h
--- a/Engine/Plugins/Experimental/Animation/PoseSearch/Source/Runtime/Public/PoseSearch/PoseSearchLibrary.h
+++ b/Engine/Plugins/Experimental/Animation/PoseSearch/Source/Runtime/Public/PoseSearch/PoseSearchLibrary.h
@@ -41,7 +41,6 @@ enum class EPoseSearchInterruptMode : uint8
 };
 
 struct FAnimationUpdateContext;
-struct FPoseSearchQueryTrajectory;
 
 struct FMotionMatchingState
 {
@@ -97,7 +96,6 @@ class POSESEARCH_API UPoseSearchLibrary : public UBlueprintFunctionLibrary
 
 #if UE_POSE_SEARCH_TRACE_ENABLED
 	static void TraceMotionMatchingState(
-		const FPoseSearchQueryTrajectory& Trajectory,
 		UE::PoseSearch::FSearchContext& SearchContext,
 		const UE::PoseSearch::FSearchResult& CurrentResult,
 		float ElapsedPoseSearchTime,
@@ -109,20 +107,12 @@ class POSESEARCH_API UPoseSearchLibrary : public UBlueprintFunctionLibrary
 		float RecordingTime);
 #endif // UE_POSE_SEARCH_TRACE_ENABLED
 
-	static FPoseSearchQueryTrajectory ProcessTrajectory(
-		const FPoseSearchQueryTrajectory& Trajectory,
-		const FTransform& RootBoneTransform,
-		float RootBoneDeltaYaw,
-		float YawFromAnimationTrajectoryBlendTime,
-		float TrajectorySpeedMultiplier);
-
 public:
 	/**
 	* Implementation of the core motion matching algorithm
 	*
 	* @param Context						Input animation update context providing access to the proxy and delta time
 	* @param Databases						Input array of databases to search
-	* @param Trajectory						Input motion trajectory samples for pose search queries. Expected to be in the space of the SkeletalMeshComponent. This is provided with the CharacterMovementTrajectory Component output.
 	* @param BlendTime						Input time in seconds to blend out to the new pose. Uses either inertial blending, requiring an Inertialization node after this node, or the internal blend stack, if MaxActiveBlends is greater than zero.
 	* @param MaxActiveBlends				Input number of max active animation segments being blended together in the blend stack. If MaxActiveBlends is zero then the blend stack is disabled.
 	* @param PoseJumpThresholdTime			Input don't jump to poses of the same segment that are within the interval this many seconds away from the continuing pose.
@@ -138,8 +128,6 @@ public:
 	static void UpdateMotionMatchingState(
 		const FAnimationUpdateContext& Context,
 		const TArray<TObjectPtr<const UPoseSearchDatabase>>& Databases,
-		const FPoseSearchQueryTrajectory& Trajectory,
-		float TrajectorySpeedMultiplier,
 		float BlendTime,
 		int32 MaxActiveBlends,
 		const FFloatInterval& PoseJumpThresholdTime,
@@ -153,16 +141,13 @@ public:
 		bool bShouldSearch = true,
 		bool bShouldUseCachedChannelData = true,
 		bool bDebugDrawQuery = false,
-		bool bDebugDrawCurResult = false,
-		bool bDebugDrawPoseHistory = false);
+		bool bDebugDrawCurResult = false);
 
 	/**
 	* Implementation of the core motion matching algorithm
 	*
 	* @param AnimInstance					Input animation instance
 	* @param Database						Input database to search
-	* @param Trajectory						Input motion trajectory samples for pose search queries. Expected to be in the space of the SkeletalMeshComponent. This is provided with the CharacterMovementTrajectory Component output.
-	* @param TrajectorySpeedMultiplier		Input Trajectory velocity will be multiplied by TrajectorySpeedMultiplier: values below 1 will result in selecting animation slower than requested from the original Trajectory
 	* @param PoseHistoryName				Input tag of the associated PoseSearchHistoryCollector node in the anim graph
 	* @param SelectedAnimation				Output selected animation from the Database asset
 	* @param SelectedTime					Output selected animation time
@@ -179,8 +164,6 @@ public:
 	static void MotionMatch(
 		UAnimInstance* AnimInstance,
 		const UPoseSearchDatabase* Database,
-		const FPoseSearchQueryTrajectory Trajectory,
-		float TrajectorySpeedMultiplier,
 		const FName PoseHistoryName,
 		FPoseSearchBlueprintResult& Result,
 		const UAnimationAsset* FutureAnimation = nullptr,
diff --git a/Engine/Plugins/Experimental/Animation/PoseSearch/Source/Runtime/Public/PoseSearch/PoseSearchTrajectoryTypes.h b/Engine/Plugins/Experimental/Animation/PoseSearch/Source/Runtime/Public/PoseSearch/PoseSearchTrajectoryTypes.h
--- a/Engine/Plugins/Experimental/Animation/PoseSearch/Source/Runtime/Public/PoseSearch/PoseSearchTrajectoryTypes.h
+++ b/Engine/Plugins/Experimental/Animation/PoseSearch/Source/Runtime/Public/PoseSearch/PoseSearchTrajectoryTypes.h
@@ -7,6 +7,8 @@
 
 #include "PoseSearchTrajectoryTypes.generated.h"
 
+struct FAnimInstanceProxy;
+
 USTRUCT(BlueprintType, Category="Pose Search Trajectory")
 struct POSESEARCH_API FPoseSearchQueryTrajectorySample
 {
@@ -40,6 +42,7 @@ struct POSESEARCH_API FPoseSearchQueryTrajectory
 	
 #if ENABLE_ANIM_DEBUG
 	void DebugDrawTrajectory(const UWorld* World, float HeightOffset = 0.f) const;
+	void DebugDrawTrajectory(FAnimInstanceProxy& AnimInstanceProxy, float HeightOffset = 0.f) const;
 #endif // ENABLE_ANIM_DEBUG
 };
 POSESEARCH_API FArchive& operator<<(FArchive& Ar, FPoseSearchQueryTrajectory& Trajectory);
