motion matching - fix for updating the trajectory in the poise history one frame late

diff --git a/Engine/Plugins/Experimental/Animation/PoseSearch/Source/Runtime/Private/AnimNode_PoseSearchHistoryCollector.cpp b/Engine/Plugins/Experimental/Animation/PoseSearch/Source/Runtime/Private/AnimNode_PoseSearchHistoryCollector.cpp
--- a/Engine/Plugins/Experimental/Animation/PoseSearch/Source/Runtime/Private/AnimNode_PoseSearchHistoryCollector.cpp
+++ b/Engine/Plugins/Experimental/Animation/PoseSearch/Source/Runtime/Private/AnimNode_PoseSearchHistoryCollector.cpp
@@ -84,6 +84,8 @@ void FAnimNode_PoseSearchHistoryCollector_Base::Update_AnyThread(const FAnimatio
 {
 	GetEvaluateGraphExposedInputs().Execute(Context);
 
+	PoseHistory.UpdateTrajectory(Trajectory, TrajectorySpeedMultiplier);
+
 	const bool bNeedsReset =
 		bResetOnBecomingRelevant &&
 		UpdateCounter.HasEverBeenUpdated() &&
@@ -124,8 +126,6 @@ void FAnimNode_PoseSearchHistoryCollector::Evaluate_AnyThread(FPoseContext& Outp
 	Super::Evaluate_AnyThread(Output);
 	Source.Evaluate(Output);
 
-	PoseHistory.UpdateTrajectory(Trajectory, TrajectorySpeedMultiplier);
-
 	FCSPose<FCompactPose> ComponentSpacePose;
 	ComponentSpacePose.InitPose(Output.Pose);
 	PoseHistory.Update(Output.AnimInstanceProxy->GetDeltaSeconds(), ComponentSpacePose, bStoreScales);
@@ -182,7 +182,6 @@ void FAnimNode_PoseSearchComponentSpaceHistoryCollector::EvaluateComponentSpace_
 	Super::EvaluateComponentSpace_AnyThread(Output);
 	Source.EvaluateComponentSpace(Output);
 
-	PoseHistory.UpdateTrajectory(Trajectory, TrajectorySpeedMultiplier);
 	PoseHistory.Update(Output.AnimInstanceProxy->GetDeltaSeconds(), Output.Pose, bStoreScales);
 
 #if ENABLE_DRAW_DEBUG && ENABLE_ANIM_DEBUG
