motion matching - generating trajectory in pose history node (Experimental)

diff --git a/Engine/Plugins/Experimental/Animation/PoseSearch/Source/Runtime/Private/AnimNode_PoseSearchHistoryCollector.cpp b/Engine/Plugins/Experimental/Animation/PoseSearch/Source/Runtime/Private/AnimNode_PoseSearchHistoryCollector.cpp
--- a/Engine/Plugins/Experimental/Animation/PoseSearch/Source/Runtime/Private/AnimNode_PoseSearchHistoryCollector.cpp
+++ b/Engine/Plugins/Experimental/Animation/PoseSearch/Source/Runtime/Private/AnimNode_PoseSearchHistoryCollector.cpp
@@ -66,7 +66,7 @@ void FAnimNode_PoseSearchHistoryCollector_Base::CacheBones_AnyThread(const FAnim
 		}
 	}
 
-	PoseHistory.Init(PoseCount, SamplingInterval, RequiredBones);
+	PoseHistory.CacheBones_AnyThread(PoseCount, SamplingInterval, RequiredBones);
 
 	if (bInitializeWithRefPose)
 	{
@@ -76,7 +76,7 @@ void FAnimNode_PoseSearchHistoryCollector_Base::CacheBones_AnyThread(const FAnim
 		Pose.SetBoneContainer(&Context.AnimInstanceProxy->GetRequiredBones());
 		FCSPose<FCompactPose> ComponentSpacePose;
 		ComponentSpacePose.InitPose(Pose);
-		PoseHistory.Update(0.f, ComponentSpacePose, bStoreScales);
+		PoseHistory.EvaluateComponentSpace_AnyThread(0.f, ComponentSpacePose, bGenerateTrajectory, bStoreScales);
 	}
 }
 
@@ -84,17 +84,15 @@ void FAnimNode_PoseSearchHistoryCollector_Base::Update_AnyThread(const FAnimatio
 {
 	GetEvaluateGraphExposedInputs().Execute(Context);
 
-	PoseHistory.UpdateTrajectory(Trajectory, TrajectorySpeedMultiplier);
+	const bool bNeedsReset = bResetOnBecomingRelevant && UpdateCounter.HasEverBeenUpdated() && !UpdateCounter.WasSynchronizedCounter(Context.AnimInstanceProxy->GetUpdateCounter());
 
-	const bool bNeedsReset =
-		bResetOnBecomingRelevant &&
-		UpdateCounter.HasEverBeenUpdated() &&
-		!UpdateCounter.WasSynchronizedCounter(Context.AnimInstanceProxy->GetUpdateCounter());
+	FPoseSearchTrajectoryData::FSampling TrajectoryDataSampling;
+	TrajectoryDataSampling.NumHistorySamples = FMath::Max(PoseCount, TrajectoryHistoryCount);
+	TrajectoryDataSampling.SecondsPerHistorySample = SamplingInterval;
+	TrajectoryDataSampling.NumPredictionSamples = TrajectoryPredictionCount;
+	TrajectoryDataSampling.SecondsPerPredictionSample = PredictionSamplingInterval;
 
-	if (bNeedsReset)
-	{
-		PoseHistory.ClearHistory();
-	}
+	PoseHistory.Update_AnyThread(Context.GetDeltaTime(), Trajectory, TrajectorySpeedMultiplier, bGenerateTrajectory, *Context.AnimInstanceProxy, TrajectoryData, TrajectoryDataSampling, bNeedsReset);
 
 	UpdateCounter.SynchronizeWith(Context.AnimInstanceProxy->GetUpdateCounter());
 }
@@ -128,7 +126,7 @@ void FAnimNode_PoseSearchHistoryCollector::Evaluate_AnyThread(FPoseContext& Outp
 
 	FCSPose<FCompactPose> ComponentSpacePose;
 	ComponentSpacePose.InitPose(Output.Pose);
-	PoseHistory.Update(Output.AnimInstanceProxy->GetDeltaSeconds(), ComponentSpacePose, bStoreScales);
+	PoseHistory.EvaluateComponentSpace_AnyThread(Output.AnimInstanceProxy->GetDeltaSeconds(), ComponentSpacePose, bGenerateTrajectory, bStoreScales);
 
 #if ENABLE_DRAW_DEBUG && ENABLE_ANIM_DEBUG
 		FColor Color;
@@ -182,7 +180,7 @@ void FAnimNode_PoseSearchComponentSpaceHistoryCollector::EvaluateComponentSpace_
 	Super::EvaluateComponentSpace_AnyThread(Output);
 	Source.EvaluateComponentSpace(Output);
 
-	PoseHistory.Update(Output.AnimInstanceProxy->GetDeltaSeconds(), Output.Pose, bStoreScales);
+	PoseHistory.EvaluateComponentSpace_AnyThread(Output.AnimInstanceProxy->GetDeltaSeconds(), Output.Pose, bGenerateTrajectory, bStoreScales);
 
 #if ENABLE_DRAW_DEBUG && ENABLE_ANIM_DEBUG
 		FColor Color;
diff --git a/Engine/Plugins/Experimental/Animation/PoseSearch/Source/Runtime/Private/PoseSearchHistory.cpp b/Engine/Plugins/Experimental/Animation/PoseSearch/Source/Runtime/Private/PoseSearchHistory.cpp
--- a/Engine/Plugins/Experimental/Animation/PoseSearch/Source/Runtime/Private/PoseSearchHistory.cpp
+++ b/Engine/Plugins/Experimental/Animation/PoseSearch/Source/Runtime/Private/PoseSearchHistory.cpp
@@ -71,9 +71,9 @@ FORCEINLINE auto LowerBound(IteratorType First, IteratorType Last, const ValueTy
 
 //////////////////////////////////////////////////////////////////////////
 // FPoseHistoryEntry
-void FPoseHistoryEntry::Update(float InTime, FCSPose<FCompactPose>& ComponentSpacePose, const FBoneToTransformMap& BoneToTransformMap, bool bStoreScales)
+void FPoseHistoryEntry::Update(float Time, FCSPose<FCompactPose>& ComponentSpacePose, const FBoneToTransformMap& BoneToTransformMap, bool bStoreScales)
 {
-	Time = InTime;
+	AccumulatedSeconds = Time;
 
 	const FBoneContainer& BoneContainer = ComponentSpacePose.GetPose().GetBoneContainer();
 	const USkeleton* SkeletonAsset = BoneContainer.GetSkeletonAsset();
@@ -134,11 +134,52 @@ FTransform FPoseHistoryEntry::GetComponentSpaceTransform(int32 Index) const
 	return FTransform(Quat, ComponentSpacePositions[Index], Scale);
 }
 
+//////////////////////////////////////////////////////////////////////////
+// IPoseHistory
+#if ENABLE_DRAW_DEBUG && ENABLE_ANIM_DEBUG
+void IPoseHistory::DebugDraw(FAnimInstanceProxy& AnimInstanceProxy, FColor Color, float Time, float PointSize, bool bExtrapolate) const
+{
+	const FBoneContainer& BoneContainer = AnimInstanceProxy.GetRequiredBones();
+	if (BoneContainer.IsValid())
+	{
+		const USkeleton* Skeleton = BoneContainer.GetSkeletonAsset();
+		FTransform OutBoneTransform;
+
+		const FBoneToTransformMap& BoneToTransformMap = GetBoneToTransformMap();
+		if (BoneToTransformMap.IsEmpty())
+		{
+			for (FSkeletonPoseBoneIndex SkeletonBoneIdx(0); SkeletonBoneIdx != BoneContainer.GetNumBones(); ++SkeletonBoneIdx)
+			{
+				const FCompactPoseBoneIndex CompactBoneIdx = BoneContainer.GetCompactPoseIndexFromSkeletonPoseIndex(SkeletonBoneIdx);
+				if (GetTransformAtTime(Time, OutBoneTransform, Skeleton, CompactBoneIdx.GetInt(), WorldSpaceIndexType, bExtrapolate))
+				{
+					AnimInstanceProxy.AnimDrawDebugPoint(OutBoneTransform.GetTranslation(), PointSize, Color, false, 0.f, ESceneDepthPriorityGroup::SDPG_Foreground);
+				}
+			}
+		}
+		else
+		{
+			for (const FBoneToTransformPair& BoneToTransformPair : BoneToTransformMap)
+			{
+				const FSkeletonPoseBoneIndex SkeletonBoneIdx(BoneToTransformPair.Key);
+				const FCompactPoseBoneIndex CompactBoneIdx = BoneContainer.GetCompactPoseIndexFromSkeletonPoseIndex(SkeletonBoneIdx);
+				if (GetTransformAtTime(Time, OutBoneTransform, Skeleton, CompactBoneIdx.GetInt(), WorldSpaceIndexType, bExtrapolate))
+				{
+					AnimInstanceProxy.AnimDrawDebugPoint(OutBoneTransform.GetTranslation(), PointSize, Color, false, 0.f, ESceneDepthPriorityGroup::SDPG_Foreground);
+				}
+			}
+		}
+	}
+}
+#endif // ENABLE_DRAW_DEBUG && ENABLE_ANIM_DEBUG
+
 //////////////////////////////////////////////////////////////////////////
 // FPoseHistory
-void FPoseHistory::Init(int32 InNumPoses, float InSamplingInterval, const TArray<FBoneIndexType>& RequiredBones)
+void FPoseHistory::CacheBones_AnyThread(int32 InNumPoses, float InSamplingInterval, const TArray<FBoneIndexType>& RequiredBones)
 {
 	check(InNumPoses >= 2 && InSamplingInterval > UE_KINDA_SMALL_NUMBER);
+
+	MaxNumPoses = InNumPoses;
 	SamplingInterval = InSamplingInterval;
 
 	BoneToTransformMap.Reset();
@@ -158,7 +199,7 @@ void FPoseHistory::Init(int32 InNumPoses, float InSamplingInterval, const TArray
 	}
 
 	Entries.Reset();
-	Entries.Reserve(InNumPoses);
+	Entries.Reserve(MaxNumPoses);
 }
 
 FBoneIndexType FPoseHistory::GetRemappedBoneIndexType(FBoneIndexType BoneIndexType, const USkeleton* BoneIndexSkeleton, const USkeleton* LastUpdateSkeleton)
@@ -206,12 +247,11 @@ bool FPoseHistory::LerpEntries(float Time, bool bExtrapolate, const FPoseHistory
 {
 	bool bSuccess = true;
 
-	const float Denominator = NextEntry.Time - PrevEntry.Time;
+	const float Denominator = NextEntry.AccumulatedSeconds - PrevEntry.AccumulatedSeconds;
 	float LerpValue = 0.f;
 	if (!FMath::IsNearlyZero(Denominator))
 	{
-		const float SecondsAgo = -Time;
-		const float Numerator = SecondsAgo - PrevEntry.Time;
+		const float Numerator = Time - PrevEntry.AccumulatedSeconds;
 		LerpValue = bExtrapolate ? Numerator / Denominator : FMath::Clamp(Numerator / Denominator, 0.f, 1.f);
 	}
 
@@ -264,14 +304,12 @@ bool FPoseHistory::GetTransformAtTime(float Time, FTransform& OutBoneTransform,
 	const int32 NumEntries = Entries.Num();
 	if (NumEntries > 0)
 	{
-		const float SecondsAgo = -Time;
-
 		int32 NextIdx = 0;
 		int32 PrevIdx = 0;
 
 		if (NumEntries > 1)
 		{
-			const int32 LowerBoundIdx = LowerBound(Entries.begin(), Entries.end(), SecondsAgo, [](const FPoseHistoryEntry& Entry, float Value) { return Value < Entry.Time; });
+			const int32 LowerBoundIdx = LowerBound(Entries.begin(), Entries.end(), Time, [](const FPoseHistoryEntry& Entry, float Value) { return Value > Entry.AccumulatedSeconds; });
 			NextIdx = FMath::Clamp(LowerBoundIdx, 1, NumEntries - 1);
 			PrevIdx = NextIdx - 1;
 		}
@@ -293,11 +331,31 @@ bool FPoseHistory::GetTransformAtTime(float Time, FTransform& OutBoneTransform,
 	return bSuccess;
 }
 
-void FPoseHistory::UpdateTrajectory(const FPoseSearchQueryTrajectory& InTrajectory, float InTrajectorySpeedMultiplier)
+void FPoseHistory::Update_AnyThread(float DeltaTime, const FPoseSearchQueryTrajectory& InTrajectory, float InTrajectorySpeedMultiplier, bool bGenerateTrajectory, const FAnimInstanceProxy& AnimInstanceProxy, const FPoseSearchTrajectoryData& TrajectoryData, const FPoseSearchTrajectoryData::FSampling& TrajectoryDataSampling, bool bNeedsReset)
 {
-	TrajectorySpeedMultiplier = InTrajectorySpeedMultiplier;
+	if (bNeedsReset)
+	{
+		Entries.Reset();
+	}
+
+	if (bGenerateTrajectory)
+	{
+		// @todo: Synchronize the FPoseSearchQueryTrajectorySample::AccumulatedSeconds of the generated trajectory with the FPoseHistoryEntry::AccumulatedSeconds of the captured poses
+		FPoseSearchTrajectoryData::FDerived TrajectoryDataDerived;
+		TrajectoryData.UpdateData(DeltaTime, AnimInstanceProxy, TrajectoryDataDerived, TrajectoryDataState);
+		FPoseSearchTrajectoryLibrary::InitTrajectorySamples(Trajectory, TrajectoryData, TrajectoryDataDerived, TrajectoryDataSampling);
+		FPoseSearchTrajectoryLibrary::UpdateHistory_TransformHistory(Trajectory, TrajectoryData, TrajectoryDataDerived, TrajectoryDataSampling, DeltaTime);
+		FPoseSearchTrajectoryLibrary::UpdatePrediction_SimulateCharacterMovement(Trajectory, TrajectoryData, TrajectoryDataDerived, TrajectoryDataSampling);
+
+		// @todo: support TrajectorySpeedMultiplier
+		TrajectorySpeedMultiplier = 1.f;
+	}
+	else
+	{
+		TrajectoryDataState = FPoseSearchTrajectoryData::FState();
 	Trajectory = InTrajectory;
 
+		TrajectorySpeedMultiplier = InTrajectorySpeedMultiplier;
 	if (!FMath::IsNearlyEqual(TrajectorySpeedMultiplier, 1.f))
 	{
 		const float TrajectorySpeedMultiplierInv = FMath::IsNearlyZero(TrajectorySpeedMultiplier) ? 1.f : 1.f / TrajectorySpeedMultiplier;
@@ -307,8 +365,9 @@ void FPoseHistory::UpdateTrajectory(const FPoseSearchQueryTrajectory& InTrajecto
 		}
 	}
 }
+}
 
-void FPoseHistory::Update(float SecondsElapsed, FCSPose<FCompactPose>& ComponentSpacePose, bool bStoreScales)
+void FPoseHistory::EvaluateComponentSpace_AnyThread(float DeltaTime, FCSPose<FCompactPose>& ComponentSpacePose, bool bGenerateTrajectory, bool bStoreScales)
 {
 	const USkeleton* Skeleton = ComponentSpacePose.GetPose().GetBoneContainer().GetSkeletonAsset();
 	if (LastUpdateSkeleton != Skeleton)
@@ -321,15 +380,15 @@ void FPoseHistory::Update(float SecondsElapsed, FCSPose<FCompactPose>& Component
 	// Age our elapsed times
 	for (FPoseHistoryEntry& Entry : Entries)
 	{
-		Entry.Time += SecondsElapsed;
+		Entry.AccumulatedSeconds -= DeltaTime;
 	}
 
-	if (Entries.Num() != Entries.Max())
+	if (Entries.Num() != MaxNumPoses)
 	{
 		// Consume every pose until the queue is full
 		Entries.Emplace();
 	}
-	else if (Entries[Entries.Num() - 2].Time >= SamplingInterval)
+	else if (Entries[Entries.Num() - 2].AccumulatedSeconds <= -SamplingInterval)
 	{
 		FPoseHistoryEntry EntryTemp = MoveTemp(Entries.First());
 		Entries.PopFront();
@@ -364,7 +423,7 @@ void FPoseHistory::DebugDraw(FAnimInstanceProxy& AnimInstanceProxy, FColor Color
 
 			for (int32 i = 0; i < Entry.Num(); ++i)
 			{
-				const FTransform RootTransform = bValidTrajectory ? Trajectory.GetSampleAtTime(-Entry.Time).GetTransform() : AnimInstanceProxy.GetComponentTransform();
+				const FTransform RootTransform = bValidTrajectory ? Trajectory.GetSampleAtTime(Entry.AccumulatedSeconds).GetTransform() : AnimInstanceProxy.GetComponentTransform();
 				const FTransform GlobalTransforms = Entry.GetComponentSpaceTransform(i) * RootTransform;
 
 				if (i < PrevGlobalTransformsNum)
@@ -408,8 +467,7 @@ bool FExtendedPoseHistory::GetTransformAtTime(float Time, FTransform& OutBoneTra
 				ReferenceBoneIndexType = ComponentSpaceIndexType;
 			}
 
-			const float SecondsAgo = -Time;
-			const int32 LowerBoundIdx = Algo::LowerBound(FutureEntries, SecondsAgo, [](const FPoseHistoryEntry& Entry, float Value) { return Value < Entry.Time; });
+			const int32 LowerBoundIdx = Algo::LowerBound(FutureEntries, Time, [](const FPoseHistoryEntry& Entry, float Value) { return Value > Entry.AccumulatedSeconds; });
 			const int32 NextIdx = FMath::Min(LowerBoundIdx, Num - 1);
 			const FPoseHistoryEntries& PastEntries = GetEntries();
 			const FPoseHistoryEntry& NextEntry = FutureEntries[NextIdx];
@@ -427,27 +485,25 @@ bool FExtendedPoseHistory::GetTransformAtTime(float Time, FTransform& OutBoneTra
 	return PoseHistory->GetTransformAtTime(Time, OutBoneTransform, BoneIndexSkeleton, BoneIndexType, ReferenceBoneIndexType, bExtrapolate);
 }
 
-void FExtendedPoseHistory::AddFutureRootBone(float SecondsInTheFuture, const FTransform& FutureRootBoneTransform, bool bStoreScales)
+void FExtendedPoseHistory::AddFutureRootBone(float Time, const FTransform& FutureRootBoneTransform, bool bStoreScales)
 {
 	// we don't allow to add "past" or "present" poses to FutureEntries
-	check(SecondsInTheFuture > 0.f);
+	check(Time > 0.f);
 
-	const float SecondsAgo = -SecondsInTheFuture;
-	const int32 LowerBoundIdx = Algo::LowerBound(FutureEntries, SecondsAgo, [](const FPoseHistoryEntry& Entry, float Value) { return Value < Entry.Time; });
+	const int32 LowerBoundIdx = Algo::LowerBound(FutureEntries, Time, [](const FPoseHistoryEntry& Entry, float Value) { return Value > Entry.AccumulatedSeconds; });
 	FPoseHistoryEntry& FutureEntry = FutureEntries.InsertDefaulted_GetRef(LowerBoundIdx);
 	FutureEntry.SetNum(1, bStoreScales);
 	FutureEntry.SetComponentSpaceTransform(RootBoneIndexType, FutureRootBoneTransform);
-	FutureEntry.Time = SecondsAgo;
+	FutureEntry.AccumulatedSeconds = Time;
 }
 
-void FExtendedPoseHistory::AddFuturePose(float SecondsInTheFuture, FCSPose<FCompactPose>& ComponentSpacePose)
+void FExtendedPoseHistory::AddFuturePose(float Time, FCSPose<FCompactPose>& ComponentSpacePose)
 {
 	// we don't allow to add "past" or "present" poses to FutureEntries
-	check(SecondsInTheFuture > 0.f);
+	check(Time > 0.f);
 	check(PoseHistory);	
-	const float SecondsAgo = -SecondsInTheFuture;
-	const int32 LowerBoundIdx = Algo::LowerBound(FutureEntries, SecondsAgo, [](const FPoseHistoryEntry& Entry, float Value) { return Value < Entry.Time; });
-	FutureEntries.InsertDefaulted_GetRef(LowerBoundIdx).Update(SecondsAgo, ComponentSpacePose, GetBoneToTransformMap(), true);
+	const int32 LowerBoundIdx = Algo::LowerBound(FutureEntries, Time, [](const FPoseHistoryEntry& Entry, float Value) { return Value > Entry.AccumulatedSeconds; });
+	FutureEntries.InsertDefaulted_GetRef(LowerBoundIdx).Update(Time, ComponentSpacePose, GetBoneToTransformMap(), true);
 }
 
 #if ENABLE_DRAW_DEBUG && ENABLE_ANIM_DEBUG
@@ -481,7 +537,7 @@ void FExtendedPoseHistory::DebugDraw(FAnimInstanceProxy& AnimInstanceProxy, FCol
 
 			for (int32 i = 0; i < Entry.Num(); ++i)
 			{
-				const FTransform RootTransform = bValidTrajectory ? Trajectory.GetSampleAtTime(-Entry.Time).GetTransform() : AnimInstanceProxy.GetComponentTransform();
+				const FTransform RootTransform = bValidTrajectory ? Trajectory.GetSampleAtTime(Entry.AccumulatedSeconds).GetTransform() : AnimInstanceProxy.GetComponentTransform();
 				const FTransform GlobalTransforms = Entry.GetComponentSpaceTransform(i) * RootTransform;
 
 				if (i < PrevGlobalTransformsNum)
diff --git a/Engine/Plugins/Experimental/Animation/PoseSearch/Source/Runtime/Private/PoseSearchTrajectoryLibrary.cpp b/Engine/Plugins/Experimental/Animation/PoseSearch/Source/Runtime/Private/PoseSearchTrajectoryLibrary.cpp
new file mode 100644
--- /dev/null
+++ b/Engine/Plugins/Experimental/Animation/PoseSearch/Source/Runtime/Private/PoseSearchTrajectoryLibrary.cpp
@@ -0,0 +1,305 @@
+// Copyright Epic Games, Inc. All Rights Reserved.
+
+#include "PoseSearch/PoseSearchTrajectoryLibrary.h"
+#include "Animation/AnimInstanceProxy.h"
+#include "Components/SkeletalMeshComponent.h"
+#include "GameFramework/Character.h"
+#include "GameFramework/CharacterMovementComponent.h"
+
+void FPoseSearchTrajectoryData::UpdateData(float DeltaTime, const FAnimInstanceProxy& AnimInstanceProxy, FDerived& TrajectoryDataDerived, FState& TrajectoryDataState) const
+{
+	// An AnimInstance might call this during an AnimBP recompile with 0 delta time.
+	if (DeltaTime <= 0.f)
+	{
+		return;
+	}
+
+	const UAnimInstance* AnimInstance = Cast<const UAnimInstance>(AnimInstanceProxy.GetAnimInstanceObject());
+	if (!AnimInstance)
+	{
+		return;
+	}
+
+	const ACharacter* Character = Cast<ACharacter>(AnimInstance->GetOwningActor());
+	if (!Character)
+	{
+		return;
+	}
+
+	if (const UCharacterMovementComponent* MoveComp = Character->GetCharacterMovement())
+	{
+		TrajectoryDataDerived.MaxSpeed = FMath::Max(MoveComp->GetMaxSpeed() * MoveComp->GetAnalogInputModifier(), MoveComp->GetMinAnalogSpeed());
+		TrajectoryDataDerived.BrakingDeceleration = FMath::Max(0.f, MoveComp->GetMaxBrakingDeceleration());
+		TrajectoryDataDerived.bOrientRotationToMovement = MoveComp->bOrientRotationToMovement;
+
+		TrajectoryDataDerived.Velocity = MoveComp->Velocity;
+		TrajectoryDataDerived.Acceleration = MoveComp->GetCurrentAcceleration();
+
+		if (TrajectoryDataDerived.Acceleration.IsZero())
+		{
+			TrajectoryDataDerived.Friction = MoveComp->bUseSeparateBrakingFriction ? MoveComp->BrakingFriction : MoveComp->GroundFriction;
+			const float FrictionFactor = FMath::Max(0.f, MoveComp->BrakingFrictionFactor);
+			TrajectoryDataDerived.Friction = FMath::Max(0.f, TrajectoryDataDerived.Friction * FrictionFactor);
+		}
+		else
+		{
+			TrajectoryDataDerived.Friction = MoveComp->GroundFriction;
+		}
+	}
+
+	// @todo: Simulated proxies don't have controllers, so they'll need some other mechanism to account for controller rotation rate.
+	const AController* Controller = Character->Controller;
+	if (Controller)
+	{
+		const float DesiredControllerYaw = Controller->GetDesiredRotation().Yaw;
+		
+		const float DesiredYawDelta = DesiredControllerYaw - TrajectoryDataState.DesiredControllerYawLastUpdate;
+		TrajectoryDataState.DesiredControllerYawLastUpdate = DesiredControllerYaw;
+
+		TrajectoryDataDerived.ControllerYawRate = FRotator::NormalizeAxis(DesiredYawDelta) * (1.f / DeltaTime);
+		if (MaxControllerYawRate >= 0.f)
+		{
+			TrajectoryDataDerived.ControllerYawRate = FMath::Sign(TrajectoryDataDerived.ControllerYawRate) * FMath::Min(FMath::Abs(TrajectoryDataDerived.ControllerYawRate), MaxControllerYawRate);
+		}
+	}
+
+	if (const USkeletalMeshComponent* MeshComp = Character->GetMesh())
+	{
+		TrajectoryDataDerived.Position = MeshComp->GetComponentLocation();
+		TrajectoryDataDerived.Facing = MeshComp->GetComponentRotation().Quaternion();
+		TrajectoryDataDerived.MeshCompRelativeRotation = MeshComp->GetRelativeRotation().Quaternion();
+	}
+}
+
+FVector FPoseSearchTrajectoryData::StepCharacterMovementGroundPrediction(float DeltaTime, const FVector& InVelocity, const FVector& InAcceleration, const FDerived& TrajectoryDataDerived) const
+{
+	FVector OutVelocity = InVelocity;
+
+	// Braking logic is copied from UCharacterMovementComponent::ApplyVelocityBraking()
+	if (InAcceleration.IsZero())
+	{
+		if (InVelocity.IsZero())
+		{
+			return FVector::ZeroVector;
+		}
+
+		const bool bZeroFriction = (TrajectoryDataDerived.Friction == 0.f);
+		const bool bZeroBraking = (TrajectoryDataDerived.BrakingDeceleration == 0.f);
+
+		if (bZeroFriction && bZeroBraking)
+		{
+			return InVelocity;
+		}
+
+		static const float MaxTimeStep = 1.f / 60.f;
+		float RemainingTime = DeltaTime;
+
+		const FVector PrevLinearVelocity = OutVelocity;
+		const FVector RevAccel = (bZeroBraking ? FVector::ZeroVector : (-TrajectoryDataDerived.BrakingDeceleration * OutVelocity.GetSafeNormal()));
+
+		// Decelerate to brake to a stop
+		while (RemainingTime >= UCharacterMovementComponent::MIN_TICK_TIME)
+		{
+			// Zero friction uses constant deceleration, so no need for iteration.
+			const float dt = ((RemainingTime > MaxTimeStep && !bZeroFriction) ? FMath::Min(MaxTimeStep, RemainingTime * 0.5f) : RemainingTime);
+			RemainingTime -= dt;
+
+			// apply friction and braking
+			OutVelocity = OutVelocity + ((-TrajectoryDataDerived.Friction) * OutVelocity + RevAccel) * dt;
+
+			// Don't reverse direction
+			if ((OutVelocity | PrevLinearVelocity) <= 0.f)
+			{
+				OutVelocity = FVector::ZeroVector;
+				return OutVelocity;
+			}
+		}
+
+		// Clamp to zero if nearly zero, or if below min threshold and braking
+		const float VSizeSq = OutVelocity.SizeSquared();
+		if (VSizeSq <= KINDA_SMALL_NUMBER || (!bZeroBraking && VSizeSq <= FMath::Square(UCharacterMovementComponent::BRAKE_TO_STOP_VELOCITY)))
+		{
+			OutVelocity = FVector::ZeroVector;
+		}
+	}
+	// Acceleration logic is copied from  UCharacterMovementComponent::CalcVelocity
+	else
+	{
+		const FVector AccelDir = InAcceleration.GetSafeNormal();
+		const float VelSize = OutVelocity.Size();
+
+		OutVelocity = OutVelocity - (OutVelocity - AccelDir * VelSize) * FMath::Min(DeltaTime * TrajectoryDataDerived.Friction, 1.f);
+
+		OutVelocity += InAcceleration * DeltaTime;
+		OutVelocity = OutVelocity.GetClampedToMaxSize(TrajectoryDataDerived.MaxSpeed);
+	}
+
+	return OutVelocity;
+}
+
+void FPoseSearchTrajectoryLibrary::InitTrajectorySamples(FPoseSearchQueryTrajectory& Trajectory, const FPoseSearchTrajectoryData& TrajectoryData, const FPoseSearchTrajectoryData::FDerived& TrajectoryDataDerived, const FPoseSearchTrajectoryData::FSampling& TrajectoryDataSampling)
+{
+	const int32 NumHistorySamples = TrajectoryDataSampling.NumHistorySamples;
+	const int32 NumPredictionSamples = TrajectoryDataSampling.NumPredictionSamples;
+
+	// History + current sample + prediction
+	const int32 TotalNumSamples = NumHistorySamples + 1 + NumPredictionSamples;
+
+	if (Trajectory.Samples.Num() != TotalNumSamples)
+	{
+		Trajectory.Samples.SetNumUninitialized(TotalNumSamples);
+
+		// Initialize history samples
+		const float SecondsPerHistorySample = TrajectoryDataSampling.SecondsPerHistorySample;
+		for (int32 i = 0; i < NumHistorySamples; ++i)
+		{
+			Trajectory.Samples[i].Position = TrajectoryDataDerived.Position;
+			Trajectory.Samples[i].Facing = TrajectoryDataDerived.Facing;
+			Trajectory.Samples[i].AccumulatedSeconds = SecondsPerHistorySample * (i - NumHistorySamples);
+		}
+
+		// Initialize current sample and prediction
+		const float SecondsPerPredictionSample = TrajectoryDataSampling.SecondsPerPredictionSample;
+		for (int32 i = NumHistorySamples; i < Trajectory.Samples.Num(); ++i)
+		{
+			Trajectory.Samples[i].Position = TrajectoryDataDerived.Position;
+			Trajectory.Samples[i].Facing = TrajectoryDataDerived.Facing;
+			Trajectory.Samples[i].AccumulatedSeconds = SecondsPerPredictionSample * (i - NumHistorySamples);
+		}
+	}
+}
+
+void FPoseSearchTrajectoryLibrary::UpdateHistory_TransformHistory(FPoseSearchQueryTrajectory& Trajectory, const FPoseSearchTrajectoryData& TrajectoryData, const FPoseSearchTrajectoryData::FDerived& TrajectoryDataDerived, const FPoseSearchTrajectoryData::FSampling& TrajectoryDataSampling, float DeltaTime)
+{
+	const int32 NumHistorySamples = TrajectoryDataSampling.NumHistorySamples;
+	const float SecondsPerHistorySample = TrajectoryDataSampling.SecondsPerHistorySample;
+
+	check(NumHistorySamples <= Trajectory.Samples.Num());
+
+	// converting all the history samples relative to the previous character position (Trajectory.Samples[NumHistorySamples].Position)
+	for (int32 Index = 0; Index < NumHistorySamples; ++Index)
+	{
+		Trajectory.Samples[Index].Position = Trajectory.Samples[NumHistorySamples].Position - Trajectory.Samples[Index].Position;
+	}
+
+	FVector CurrentTranslation = TrajectoryDataDerived.Velocity * DeltaTime;
+
+	// Shift history Samples when it's time to record a new one.
+	if (NumHistorySamples > 0 && FMath::Abs(Trajectory.Samples[NumHistorySamples - 1].AccumulatedSeconds) >= SecondsPerHistorySample)
+	{
+		for (int32 Index = 0; Index < NumHistorySamples - 1; ++Index)
+		{
+			Trajectory.Samples[Index].AccumulatedSeconds = Trajectory.Samples[Index + 1].AccumulatedSeconds;
+			Trajectory.Samples[Index].Position = Trajectory.Samples[Index + 1].Position + CurrentTranslation;
+			Trajectory.Samples[Index].Facing = Trajectory.Samples[Index + 1].Facing;
+		}
+
+		Trajectory.Samples[NumHistorySamples - 1].AccumulatedSeconds = 0.f;
+		Trajectory.Samples[NumHistorySamples - 1].Position = CurrentTranslation;
+		Trajectory.Samples[NumHistorySamples - 1].Facing = TrajectoryDataDerived.Facing;
+	}
+	else
+	{
+		for (int32 Index = 0; Index < NumHistorySamples; ++Index)
+		{
+			Trajectory.Samples[Index].Position += CurrentTranslation;
+		}
+	}
+
+	// converting the history sample positions in world space by applying the current world position.
+	for (int32 Index = 0; Index < NumHistorySamples; ++Index)
+	{
+		Trajectory.Samples[Index].AccumulatedSeconds -= DeltaTime;
+		Trajectory.Samples[Index].Position = TrajectoryDataDerived.Position - Trajectory.Samples[Index].Position;
+	}
+}
+
+FVector FPoseSearchTrajectoryLibrary::RemapVectorMagnitudeWithCurve(const FVector& Vector, bool bUseCurve, const FRuntimeFloatCurve& Curve)
+{
+	if (bUseCurve)
+	{
+		const float Length = Vector.Length();
+		if (Length > UE_KINDA_SMALL_NUMBER)
+		{
+			const float RemappedLength = Curve.GetRichCurveConst()->Eval(Length);
+			return Vector * (RemappedLength / Length);
+		}
+	}
+
+	return Vector;
+}
+
+void FPoseSearchTrajectoryLibrary::UpdatePrediction_SimulateCharacterMovement(FPoseSearchQueryTrajectory& Trajectory, const FPoseSearchTrajectoryData& TrajectoryData, const FPoseSearchTrajectoryData::FDerived& TrajectoryDataDerived, const FPoseSearchTrajectoryData::FSampling& TrajectoryDataSampling)
+{
+	FVector CurrentPositionWS = TrajectoryDataDerived.Position;
+	FVector CurrentVelocityWS = FPoseSearchTrajectoryLibrary::RemapVectorMagnitudeWithCurve(TrajectoryDataDerived.Velocity, TrajectoryData.bUseSpeedRemappingCurve, TrajectoryData.SpeedRemappingCurve);
+	FVector CurrentAccelerationWS = FPoseSearchTrajectoryLibrary::RemapVectorMagnitudeWithCurve(TrajectoryDataDerived.Acceleration, TrajectoryData.bUseAccelerationRemappingCurve, TrajectoryData.AccelerationRemappingCurve);
+
+	// bending CurrentVelocityWS towards CurrentAccelerationWS
+	if (TrajectoryData.BendVelocityTowardsAcceleration > UE_KINDA_SMALL_NUMBER && !CurrentAccelerationWS.IsNearlyZero())
+	{
+		const float CurrentSpeed = CurrentVelocityWS.Length();
+		const FVector VelocityWSAlongAcceleration = CurrentAccelerationWS.GetUnsafeNormal() * CurrentSpeed;
+		if (TrajectoryData.BendVelocityTowardsAcceleration < 1.f - UE_KINDA_SMALL_NUMBER)
+		{
+			CurrentVelocityWS = FMath::Lerp(CurrentVelocityWS, VelocityWSAlongAcceleration, TrajectoryData.BendVelocityTowardsAcceleration);
+
+			const float NewLength = CurrentVelocityWS.Length();
+			if (NewLength > UE_KINDA_SMALL_NUMBER)
+			{
+				CurrentVelocityWS *= CurrentSpeed / NewLength;
+			}
+			else
+			{
+				// @todo: consider setting the CurrentVelocityWS = VelocityWSAlongAcceleration if vel and acc are in opposite directions
+			}
+		}
+		else
+		{
+			CurrentVelocityWS = VelocityWSAlongAcceleration;
+		}
+	}
+
+	FQuat CurrentFacingWS = TrajectoryDataDerived.Facing;
+	
+	const int32 NumHistorySamples = TrajectoryDataSampling.NumHistorySamples;
+	const float SecondsPerPredictionSample = TrajectoryDataSampling.SecondsPerPredictionSample;
+	const FQuat ControllerRotationPerStep = FQuat::MakeFromEuler(FVector(0.f, 0.f, TrajectoryDataDerived.ControllerYawRate * SecondsPerPredictionSample));
+
+	float AccumulatedSeconds = 0.f;
+
+	const int32 LastIndex = Trajectory.Samples.Num() - 1;
+	if (NumHistorySamples <= LastIndex)
+	{
+		for (int32 Index = NumHistorySamples; ; ++Index)
+		{
+			Trajectory.Samples[Index].Position = CurrentPositionWS;
+			Trajectory.Samples[Index].Facing = CurrentFacingWS;
+			Trajectory.Samples[Index].AccumulatedSeconds = AccumulatedSeconds;
+
+			if (Index == LastIndex)
+			{
+				break;
+			}
+
+			CurrentPositionWS += CurrentVelocityWS * SecondsPerPredictionSample;
+			AccumulatedSeconds += SecondsPerPredictionSample;
+
+			// Account for the controller (e.g. the camera) rotating.
+			CurrentFacingWS = ControllerRotationPerStep * CurrentFacingWS;
+			CurrentAccelerationWS = FPoseSearchTrajectoryLibrary::RemapVectorMagnitudeWithCurve(ControllerRotationPerStep * CurrentAccelerationWS,
+				TrajectoryData.bUseAccelerationRemappingCurve, TrajectoryData.AccelerationRemappingCurve);
+
+			const FVector NewVelocityWS = TrajectoryData.StepCharacterMovementGroundPrediction(SecondsPerPredictionSample, CurrentVelocityWS, CurrentAccelerationWS, TrajectoryDataDerived);
+
+			CurrentVelocityWS = FPoseSearchTrajectoryLibrary::RemapVectorMagnitudeWithCurve(NewVelocityWS, TrajectoryData.bUseSpeedRemappingCurve, TrajectoryData.SpeedRemappingCurve);
+
+			if (TrajectoryDataDerived.bOrientRotationToMovement && !CurrentAccelerationWS.IsNearlyZero())
+			{
+				// Rotate towards acceleration.
+				const FVector CurrentAccelerationCS = TrajectoryDataDerived.MeshCompRelativeRotation.RotateVector(CurrentAccelerationWS);
+				CurrentFacingWS = FMath::QInterpConstantTo(CurrentFacingWS, CurrentAccelerationCS.ToOrientationQuat(), SecondsPerPredictionSample, TrajectoryData.RotateTowardsMovementSpeed);
+			}
+		}
+	}
+}
\ No newline at end of file
diff --git a/Engine/Plugins/Experimental/Animation/PoseSearch/Source/Runtime/Public/PoseSearch/AnimNode_PoseSearchHistoryCollector.h b/Engine/Plugins/Experimental/Animation/PoseSearch/Source/Runtime/Public/PoseSearch/AnimNode_PoseSearchHistoryCollector.h
--- a/Engine/Plugins/Experimental/Animation/PoseSearch/Source/Runtime/Public/PoseSearch/AnimNode_PoseSearchHistoryCollector.h
+++ b/Engine/Plugins/Experimental/Animation/PoseSearch/Source/Runtime/Public/PoseSearch/AnimNode_PoseSearchHistoryCollector.h
@@ -20,7 +20,7 @@ public:
 	
 	// how often in seconds poses are collected
 	UPROPERTY(EditAnywhere, BlueprintReadWrite, Category = Settings, meta = (ClampMin="0.001"))
-	float SamplingInterval = UE::PoseSearch::FiniteDelta;
+	float SamplingInterval = 0.04f;
 
 	UPROPERTY(EditAnywhere, Category = Settings)
 	TArray<FBoneReference> CollectedBones;
@@ -35,7 +35,7 @@ public:
 
 	// if true pose scales will be cached, otherwise implied to be unitary scales
 	UPROPERTY(EditAnywhere, Category = Settings)
-	bool bStoreScales = true;
+	bool bStoreScales = false;
 
 	// Update Counter for detecting being relevant
 	FGraphTraversalCounter UpdateCounter;
@@ -45,14 +45,35 @@ public:
 	FLinearColor DebugColor = FLinearColor::Red;
 #endif // WITH_EDITORONLY_DATA
 
-	// Trajectory samples for pose search queries in Motion Matching. These are expected to be in the world space of the SkeletalMeshComponent. This is provided with the CharacterMovementTrajectory Component output.
-	UPROPERTY(EditAnywhere, Category = Settings, meta = (PinShownByDefault))
+	// if true Trajectory the pose history node will generate the trajectory using the TrajectoryData parameters instead of relying on the input Trajectory (Experimental)
+	UPROPERTY(EditAnywhere, Category = Settings)
+	bool bGenerateTrajectory = false;
+
+	// input Trajectory samples for pose search queries in Motion Matching. These are expected to be in the world space of the SkeletalMeshComponent.
+	// the trajectory sample with AccumulatedSeconds equals to zero (Trajectory.Samples[i].AccumulatedSeconds) is the sample of the previous frame of simulation (since MM works by matching the previous character pose)
+	UPROPERTY(EditAnywhere, Category = Settings, meta = (PinShownByDefault, EditCondition="!bGenerateTrajectory", EditConditionHides))
 	FPoseSearchQueryTrajectory Trajectory;
 
 	// Input Trajectory velocity will be multiplied by TrajectorySpeedMultiplier: values below 1 will result in selecting animation slower than requested from the original Trajectory
-	UPROPERTY(EditAnywhere, Category = Settings, meta = (PinHiddenByDefault, ClampMin="0"))
+	UPROPERTY(EditAnywhere, Category = Settings, meta = (PinHiddenByDefault, ClampMin="0", EditCondition="!bGenerateTrajectory", EditConditionHides))
 	float TrajectorySpeedMultiplier = 1.f;
 
+	// if bGenerateTrajectory is true, this is the number of trajectory past (collected) samples
+	UPROPERTY(EditAnywhere, BlueprintReadWrite, Category = Settings, meta = (ClampMin = "2", EditCondition = "bGenerateTrajectory", EditConditionHides))
+	int32 TrajectoryHistoryCount = 10;
+
+	// if bGenerateTrajectory is true, this is the number of trajectory future (prediction) samples
+	UPROPERTY(EditAnywhere, BlueprintReadWrite, Category = Settings, meta = (ClampMin="2", EditCondition="bGenerateTrajectory", EditConditionHides))
+	int32 TrajectoryPredictionCount = 8;
+
+	// if bGenerateTrajectory is true, this is the sampling interval between trajectory future (prediction) samples
+	UPROPERTY(EditAnywhere, BlueprintReadWrite, Category = Settings, meta = (ClampMin="0.001", EditCondition="bGenerateTrajectory", EditConditionHides))
+	float PredictionSamplingInterval = 0.4f;
+
+	// if bGenerateTrajectory is true, TrajectoryData contains the tuning parameters to generate the trajectory
+	UPROPERTY(EditAnywhere, Category = Settings, meta=(EditCondition="bGenerateTrajectory", EditConditionHides))
+	FPoseSearchTrajectoryData TrajectoryData;
+
 	// FAnimNode_Base interface
 	virtual void CacheBones_AnyThread(const FAnimationCacheBonesContext& Context) override;
 	virtual void Update_AnyThread(const FAnimationUpdateContext& Context) override;
@@ -70,7 +91,7 @@ struct POSESEARCH_API FAnimNode_PoseSearchHistoryCollector : public FAnimNode_Po
 	GENERATED_BODY()
 
 public:
-	UPROPERTY(EditAnywhere, BlueprintReadWrite, Category = Links)
+	UPROPERTY(EditAnywhere, BlueprintReadWrite, Category = Links, meta = (DisplayPriority = 0))
 	FPoseLink Source;
 
 	// FAnimNode_Base interface
@@ -88,7 +109,7 @@ struct POSESEARCH_API FAnimNode_PoseSearchComponentSpaceHistoryCollector : publi
 	GENERATED_BODY()
 
 public:
-	UPROPERTY(EditAnywhere, BlueprintReadWrite, Category = Links)
+	UPROPERTY(EditAnywhere, BlueprintReadWrite, Category = Links, meta = (DisplayPriority = 0))
 	FComponentSpacePoseLink Source;
 
 	// FAnimNode_Base interface
diff --git a/Engine/Plugins/Experimental/Animation/PoseSearch/Source/Runtime/Public/PoseSearch/PoseSearchHistory.h b/Engine/Plugins/Experimental/Animation/PoseSearch/Source/Runtime/Public/PoseSearch/PoseSearchHistory.h
--- a/Engine/Plugins/Experimental/Animation/PoseSearch/Source/Runtime/Public/PoseSearch/PoseSearchHistory.h
+++ b/Engine/Plugins/Experimental/Animation/PoseSearch/Source/Runtime/Public/PoseSearch/PoseSearchHistory.h
@@ -7,7 +7,7 @@
 #include "Containers/RingBuffer.h"
 #include "DrawDebugHelpers.h"
 #include "PoseSearch/PoseSearchDefines.h"
-#include "PoseSearch/PoseSearchTrajectoryTypes.h"
+#include "PoseSearch/PoseSearchTrajectoryLibrary.h"
 #include "UObject/ObjectKey.h"
 
 struct FAnimInstanceProxy;
@@ -33,9 +33,9 @@ struct FPoseHistoryEntry
 	TArray<FQuat4f> ComponentSpaceRotations;
 	TArray<FVector> ComponentSpacePositions;
 	TArray<FVector3f> ComponentSpaceScales;
-	float Time = 0.f;
+	float AccumulatedSeconds = 0.f;
 
-	void Update(float InTime, FCSPose<FCompactPose>& ComponentSpacePose, const FBoneToTransformMap& BoneToTransformMap, bool bStoreScales);
+	void Update(float Time, FCSPose<FCompactPose>& ComponentSpacePose, const FBoneToTransformMap& BoneToTransformMap, bool bStoreScales);
 
 	void SetNum(int32 Num, bool bStoreScales);
 	int32 Num() const;
@@ -66,20 +66,27 @@ public:
 	virtual const FBoneToTransformMap& GetBoneToTransformMap() const = 0;
 	virtual const FPoseHistoryEntries& GetEntries() const = 0;
 	virtual const USkeleton* GetLastUpdateSkeleton() const = 0;
+
+#if ENABLE_DRAW_DEBUG && ENABLE_ANIM_DEBUG
+	virtual void DebugDraw(FAnimInstanceProxy& AnimInstanceProxy, FColor Color) const = 0;
+	virtual void DebugDraw(FAnimInstanceProxy& AnimInstanceProxy, FColor Color, float Time, float PointSize = 6.f, bool bExtrapolate = true) const;
+#endif
 };
 
 struct FPoseHistory : public IPoseHistory
 {
-	void Init(int32 InNumPoses, float InSamplingInterval, const TArray<FBoneIndexType>& RequiredBones);
-	void Update(float SecondsElapsed, FCSPose<FCompactPose>& ComponentSpacePose, bool bStoreScales);
+	void CacheBones_AnyThread(int32 InNumPoses, float InSamplingInterval, const TArray<FBoneIndexType>& RequiredBones);
+
+	// if bGenerateTrajectory is true, this method will generate trajectory
+	// NoTe: InTrajectory.Samples[i].AccumulatedSeconds == 0 is the sample of the previous frame of simulation (since MM works by matching the previous character pose)
+	void Update_AnyThread(float DeltaTime, const FPoseSearchQueryTrajectory& InTrajectory, float InTrajectorySpeedMultiplier, bool bGenerateTrajectory, const FAnimInstanceProxy& AnimInstanceProxy, const FPoseSearchTrajectoryData& TrajectoryData, const FPoseSearchTrajectoryData::FSampling& TrajectoryDataSampling, bool bNeedsReset);
+
+	void EvaluateComponentSpace_AnyThread(float DeltaTime, FCSPose<FCompactPose>& ComponentSpacePose, bool bGenerateTrajectory, bool bStoreScales);
 
 #if ENABLE_DRAW_DEBUG && ENABLE_ANIM_DEBUG
-	void DebugDraw(FAnimInstanceProxy& AnimInstanceProxy, FColor Color) const;
+	virtual void DebugDraw(FAnimInstanceProxy& AnimInstanceProxy, FColor Color) const override;
 #endif
 
-	void ClearHistory() { Entries.Reset(); }
-	void UpdateTrajectory(const FPoseSearchQueryTrajectory& InTrajectory, float InTrajectorySpeedMultiplier);
-
 	// IPoseHistory interface
 	virtual bool GetTransformAtTime(float Time, FTransform& OutBoneTransform, const USkeleton* BoneIndexSkeleton = nullptr, FBoneIndexType BoneIndexType = RootBoneIndexType, FBoneIndexType ReferenceBoneIndexType = ComponentSpaceIndexType, bool bExtrapolate = true) const override;
 	virtual const FPoseSearchQueryTrajectory& GetTrajectory() const override { return Trajectory; }
@@ -103,10 +110,14 @@ private:
 
 	// ring buffer of collected bones
 	FPoseHistoryEntries Entries;
+	
+	// caching MaxNumPoses, since Entries.Max() is a padded number
+	int32 MaxNumPoses = 0;
+
 	float SamplingInterval = 0.f;
 
-	// @todo: embed past Trajectory into Entries to have consistent timing of poses and trajectory. necessary during trajectory discontinuities
 	FPoseSearchQueryTrajectory Trajectory;
+	FPoseSearchTrajectoryData::FState TrajectoryDataState;
 
 	// @todo: deprecate this member and expose it via blue print logic or as global query scaling multiplier
 	float TrajectorySpeedMultiplier = 1.f;
@@ -129,11 +140,11 @@ struct FExtendedPoseHistory : public IPoseHistory
 	// End of IPoseHistory interface
 	
 #if ENABLE_DRAW_DEBUG && ENABLE_ANIM_DEBUG
-	void DebugDraw(FAnimInstanceProxy& AnimInstanceProxy, FColor Color) const;
+	virtual void DebugDraw(FAnimInstanceProxy& AnimInstanceProxy, FColor Color) const override;
 #endif
 
-	void AddFutureRootBone(float SecondsInTheFuture, const FTransform& FutureRootBoneTransform, bool bStoreScales);
-	void AddFuturePose(float SecondsInTheFuture, FCSPose<FCompactPose>& ComponentSpacePose);
+	void AddFutureRootBone(float Time, const FTransform& FutureRootBoneTransform, bool bStoreScales);
+	void AddFuturePose(float Time, FCSPose<FCompactPose>& ComponentSpacePose);
 
 private:
 	const IPoseHistory* PoseHistory = nullptr;
diff --git a/Engine/Plugins/Experimental/Animation/PoseSearch/Source/Runtime/Public/PoseSearch/PoseSearchTrajectoryLibrary.h b/Engine/Plugins/Experimental/Animation/PoseSearch/Source/Runtime/Public/PoseSearch/PoseSearchTrajectoryLibrary.h
new file mode 100644
--- /dev/null
+++ b/Engine/Plugins/Experimental/Animation/PoseSearch/Source/Runtime/Public/PoseSearch/PoseSearchTrajectoryLibrary.h
@@ -0,0 +1,97 @@
+// Copyright Epic Games, Inc. All Rights Reserved.
+
+#pragma once
+
+#include "Curves/CurveFloat.h"
+#include "PoseSearch/PoseSearchTrajectoryTypes.h"
+#include "PoseSearchTrajectoryLibrary.generated.h"
+
+USTRUCT(BlueprintType)
+struct POSESEARCH_API FPoseSearchTrajectoryData
+{
+public:
+	GENERATED_BODY()
+
+	struct FState
+	{
+		float DesiredControllerYawLastUpdate = 0.f;
+	};
+
+	struct FDerived
+	{
+		float ControllerYawRate = 0.f;
+
+		float MaxSpeed = 0.f;
+		float BrakingDeceleration = 0.f;
+		float Friction = 0.f;
+
+		FVector Velocity = FVector::ZeroVector;
+		FVector Acceleration = FVector::ZeroVector;
+
+		FVector Position = FVector::ZeroVector;
+		FQuat Facing = FQuat::Identity;
+		FQuat MeshCompRelativeRotation = FQuat::Identity;
+		bool bOrientRotationToMovement = false;
+	};
+
+	struct FSampling
+	{
+		int32 NumHistorySamples = 0;
+		float SecondsPerHistorySample = 0.f;
+
+		int32 NumPredictionSamples = 0;
+		float SecondsPerPredictionSample = 0.f;
+	};
+
+	void UpdateData(float DeltaTime, const FAnimInstanceProxy& AnimInstanceProxy, FDerived& TrajectoryDataDerived, FState& TrajectoryDataState) const;
+	FVector StepCharacterMovementGroundPrediction(float DeltaTime, const FVector& InVelocity, const FVector& InAcceleration, const FDerived& TrajectoryDataDerived) const;
+	
+	// If the character is forward facing (i.e. bOrientRotationToMovement is true), this controls how quickly the trajectory will rotate
+	// to face acceleration. It's common for this to differ from the rotation rate of the character, because animations are often authored 
+	// with different rotation speeds than the character. This is especially true in cases where the character rotation snaps to movement.
+	UPROPERTY(EditDefaultsOnly, Category = "Trajectory Settings")
+	float RotateTowardsMovementSpeed = 10.f;
+
+	// Maximum controller yaw  rate in degrees per second used to clamp the character owner controller desired yaw to generate the prediction trajectory.
+	// Negative values disable the clamping behavior
+	UPROPERTY(EditDefaultsOnly, Category = "Trajectory Settings")
+	float MaxControllerYawRate = 70.f;
+
+	// artificially bend character velocity towards acceleration direction to compute trajectory prediction, to get sharper turns
+	// 0: character velocity is used with no alteration, 1: the acceleration direction is used as velocity direction
+	UPROPERTY(EditDefaultsOnly, Category = "Trajectory Settings", meta = (ClampMin = "0", ClampMax = "1"))
+	float BendVelocityTowardsAcceleration = 0.f;
+
+	UPROPERTY(EditDefaultsOnly, Category = "Trajectory Settings", meta = (InlineEditConditionToggle))
+	bool bUseSpeedRemappingCurve = false;
+
+	UPROPERTY(EditDefaultsOnly, Category = "Trajectory Settings", meta = (EditCondition = "bUseSpeedRemappingCurve"))
+	FRuntimeFloatCurve SpeedRemappingCurve;
+
+	UPROPERTY(EditDefaultsOnly, Category = "Trajectory Settings", meta = (InlineEditConditionToggle))
+	bool bUseAccelerationRemappingCurve = false;
+
+	UPROPERTY(EditDefaultsOnly, Category = "Trajectory Settings", meta = (EditCondition = "bUseAccelerationRemappingCurve"))
+	FRuntimeFloatCurve AccelerationRemappingCurve;
+};
+
+/**
+ * Set of functions to help populate a FPoseSearchQueryTrajectory for motion matching.
+ * UCharacterTrajectoryComponent uses these functions, but they can also be used by a UAnimInstance to avoid the component.
+ */
+struct POSESEARCH_API FPoseSearchTrajectoryLibrary
+{
+public:
+	static void InitTrajectorySamples(FPoseSearchQueryTrajectory& Trajectory, const FPoseSearchTrajectoryData& TrajectoryData, const FPoseSearchTrajectoryData::FDerived& TrajectoryDataDerived, const FPoseSearchTrajectoryData::FSampling& TrajectoryDataSampling);
+
+	// Update history by tracking offsets that result from character intent (e.g. movement component velocity) and applying
+	// that to the current world transform. This works well on moving platforms as it only stores a history of movement
+	// that results from character intent, not movement from platforms.
+	static void UpdateHistory_TransformHistory(FPoseSearchQueryTrajectory& Trajectory, const FPoseSearchTrajectoryData& TrajectoryData, const FPoseSearchTrajectoryData::FDerived& TrajectoryDataDerived, const FPoseSearchTrajectoryData::FSampling& TrajectoryDataSampling, float DeltaTime);
+
+	// Update prediction by simulating the movement math for ground locomotion from UCharacterMovementComponent.
+	static void UpdatePrediction_SimulateCharacterMovement(FPoseSearchQueryTrajectory& Trajectory, const FPoseSearchTrajectoryData& TrajectoryData, const FPoseSearchTrajectoryData::FDerived& TrajectoryDataDerived, const FPoseSearchTrajectoryData::FSampling& TrajectoryDataSampling);
+
+private:
+	static FVector RemapVectorMagnitudeWithCurve(const FVector& Vector, bool bUseCurve, const FRuntimeFloatCurve& Curve);
+};
\ No newline at end of file
