motion matching - adding continuing pose support to IPoseSearchProvider::Search and FAnimNode_ChooserPlayer::ChooseAsset

diff --git a/Engine/Plugins/Experimental/Animation/PoseSearch/Source/Runtime/Private/PoseSearchContext.cpp b/Engine/Plugins/Experimental/Animation/PoseSearch/Source/Runtime/Private/PoseSearchContext.cpp
--- a/Engine/Plugins/Experimental/Animation/PoseSearch/Source/Runtime/Private/PoseSearchContext.cpp
+++ b/Engine/Plugins/Experimental/Animation/PoseSearch/Source/Runtime/Private/PoseSearchContext.cpp
@@ -373,6 +373,11 @@ FSearchContext::FSearchContext(const UAnimInstance* InAnimInstance, const IPoseH
 , bUseCachedChannelData(bInUseCachedChannelData)
 {
 	check(AnimInstance);
+	UpdateCurrentResultPoseVector();
+}
+
+void FSearchContext::UpdateCurrentResultPoseVector()
+{
 	if (CurrentResult.IsValid())
 	{
 		const FSearchIndex& SearchIndex = CurrentResult.Database->GetSearchIndex();
diff --git a/Engine/Plugins/Experimental/Animation/PoseSearch/Source/Runtime/Private/PoseSearchLibrary.cpp b/Engine/Plugins/Experimental/Animation/PoseSearch/Source/Runtime/Private/PoseSearchLibrary.cpp
--- a/Engine/Plugins/Experimental/Animation/PoseSearch/Source/Runtime/Private/PoseSearchLibrary.cpp
+++ b/Engine/Plugins/Experimental/Animation/PoseSearch/Source/Runtime/Private/PoseSearchLibrary.cpp
@@ -694,12 +694,26 @@ void UPoseSearchLibrary::MotionMatch(
 	}
 }
 
-UE::PoseSearch::FSearchResult UPoseSearchLibrary::MotionMatch(const FAnimationBaseContext& Context, TConstArrayView<UAnimationAsset*> AnimationAssets)
+UE::PoseSearch::FSearchResult UPoseSearchLibrary::MotionMatch(const FAnimationBaseContext& Context, TConstArrayView<UAnimationAsset*> AnimationAssets,
+	const UAnimationAsset* PlayingAnimationAsset, float PlayingAnimationAssetAccumulatedTime)
 {
 	using namespace UE::PoseSearch;
 
 	FSearchResult SearchResult;
 
+	const IPoseHistory* History = nullptr;
+	if (IPoseHistoryProvider* PoseHistoryProvider = Context.GetMessage<IPoseHistoryProvider>())
+	{
+		History = &PoseHistoryProvider->GetPoseHistory();
+	}
+
+	const UAnimInstance* AnimInstance = Cast<const UAnimInstance>(Context.AnimInstanceProxy->GetAnimInstanceObject());
+	check(AnimInstance);
+
+	FMemMark Mark(FMemStack::Get());
+	FSearchResult ReconstructedPreviousSearchResult;
+	FSearchContext SearchContext(AnimInstance, History, TConstArrayView<const UAnimationAsset*>(), 0.f, nullptr, ReconstructedPreviousSearchResult);
+
 	// budgeting some stack allocations for simple use cases. bigger requests of AnimationAssets contining 
 	// UAnimNotifyState_PoseSearchBranchIn referencing multiple datbases will default to slower heap allocations
 	enum { MAX_STACK_ALLOCATED_ANIMATIONS = 16 };
@@ -709,8 +723,7 @@ UE::PoseSearch::FSearchResult UPoseSearchLibrary::MotionMatch(const FAnimationBa
 	typedef TPair<const UPoseSearchDatabase*, TDbAnims> FPerDbAnimPair;
 	FPerDbAnimMap PerDbAnimMap;
 	
-	// colecting all the UAnimSequenceBase to consider for each database
-	for (const UAnimationAsset* AnimationAsset : AnimationAssets)
+	auto AddToPerDbAnimMap = [](FPerDbAnimMap& PerDbAnimMap, const UAnimationAsset* AnimationAsset)
 	{
 		if (const UAnimSequenceBase* SequenceBase = Cast<const UAnimSequenceBase>(AnimationAsset))
 		{
@@ -729,60 +742,101 @@ UE::PoseSearch::FSearchResult UPoseSearchLibrary::MotionMatch(const FAnimationBa
 				}
 			}
 		}
-	}
+	};
 
-	if (!PerDbAnimMap.IsEmpty())
+	// collecting all the possible continuing pose search (it could be multiple searches, but most likely only one)
+	if (PlayingAnimationAsset)
 	{
-		const IPoseHistory* History = nullptr;
-		if (IPoseHistoryProvider* PoseHistoryProvider = Context.GetMessage<IPoseHistoryProvider>())
+		AddToPerDbAnimMap(PerDbAnimMap, PlayingAnimationAsset);
+		for (const FPerDbAnimPair& PerDbAnimPair : PerDbAnimMap)
 		{
-			History = &PoseHistoryProvider->GetPoseHistory();
+			const UPoseSearchDatabase* Database = PerDbAnimPair.Key;
+			check(Database);
+
+#if WITH_EDITOR
+			if (!FAsyncPoseSearchDatabasesManagement::RequestAsyncBuildIndex(Database, ERequestAsyncBuildFlag::ContinueRequest))
+			{
+				SearchContext.SetAsyncBuildIndexInProgress();
+			}
+			else
+#endif // WITH_EDITOR
+			{
+				const FSearchIndex& SearchIndex = Database->GetSearchIndex();
+				for (int32 AssetIndex = 0; AssetIndex < SearchIndex.Assets.Num(); ++AssetIndex)
+				{
+					const FSearchIndexAsset& SearchIndexAsset = SearchIndex.Assets[AssetIndex];
+					if (const FPoseSearchDatabaseAnimationAssetBase* DatabaseAnimationAssetBase = Database->GetAnimationAssetBase(SearchIndexAsset))
+					{
+						if (PlayingAnimationAsset == DatabaseAnimationAssetBase->GetAnimationAsset())
+						{
+							const float FirstSampleTime = SearchIndexAsset.GetFirstSampleTime(Database->Schema->SampleRate);
+							const float LastSampleTime = SearchIndexAsset.GetLastSampleTime(Database->Schema->SampleRate);
+							if (PlayingAnimationAssetAccumulatedTime >= FirstSampleTime && PlayingAnimationAssetAccumulatedTime <= LastSampleTime)
+							{
+								ReconstructedPreviousSearchResult.Database = Database;
+								ReconstructedPreviousSearchResult.AssetTime = PlayingAnimationAssetAccumulatedTime;
+								ReconstructedPreviousSearchResult.PoseIdx = Database->GetPoseIndexFromTime(PlayingAnimationAssetAccumulatedTime, SearchIndexAsset);
+								SearchContext.UpdateCurrentResultPoseVector();
+
+								SearchResult = Database->SearchContinuingPose(SearchContext);
+								SearchContext.UpdateCurrentBestCost(SearchResult.PoseCost);
+							}
+						}
+					}
+				}
+			}
 		}
 
-		const UAnimInstance* AnimInstance = Cast<const UAnimInstance>(Context.AnimInstanceProxy->GetAnimInstanceObject());
-		check(AnimInstance);
+		PerDbAnimMap.Reset();
+	}
 
-		FMemMark Mark(FMemStack::Get());
-		FSearchContext SearchContext(AnimInstance, History);
+	// collecting all the other databases searches
+	if (!AnimationAssets.IsEmpty())
+	{
+		for (const UAnimationAsset* AnimationAsset : AnimationAssets)
+		{
+			AddToPerDbAnimMap(PerDbAnimMap, AnimationAsset);
+		}
 
 		for (const FPerDbAnimPair& PerDbAnimPair : PerDbAnimMap)
 		{
-			check(PerDbAnimPair.Key);
-			
+			const UPoseSearchDatabase* Database = PerDbAnimPair.Key;
+			check(Database);
+
 			SearchContext.SetAnimationsToConsider(PerDbAnimPair.Value);
 
-			const FSearchResult NewSearchResult = PerDbAnimPair.Key->Search(SearchContext);
+			const FSearchResult NewSearchResult = Database->Search(SearchContext);
 			if (NewSearchResult.PoseCost.GetTotalCost() < SearchResult.PoseCost.GetTotalCost())
 			{
 				SearchResult = NewSearchResult;
 				SearchContext.UpdateCurrentBestCost(SearchResult.PoseCost);
 			}
 		}
+	}
 
 #if ENABLE_DRAW_DEBUG && ENABLE_ANIM_DEBUG
-		if (SearchResult.IsValid())
+	if (SearchResult.IsValid())
+	{
+		if (CVarAnimMotionMatchDrawMatchEnable.GetValueOnAnyThread())
 		{
-			if (CVarAnimMotionMatchDrawMatchEnable.GetValueOnAnyThread())
-			{
-				FDebugDrawParams DrawParams(Context.AnimInstanceProxy, SearchContext.GetWorldBoneTransformAtTime(0.f), SearchResult.Database.Get());
-				DrawParams.DrawFeatureVector(SearchResult.PoseIdx);
-			}
+			FDebugDrawParams DrawParams(Context.AnimInstanceProxy, SearchContext.GetWorldBoneTransformAtTime(0.f), SearchResult.Database.Get());
+			DrawParams.DrawFeatureVector(SearchResult.PoseIdx);
+		}
 
-			if (CVarAnimMotionMatchDrawQueryEnable.GetValueOnAnyThread())
-			{
-				FDebugDrawParams DrawParams(Context.AnimInstanceProxy, SearchContext.GetWorldBoneTransformAtTime(0.f), SearchResult.Database.Get(), EDebugDrawFlags::DrawQuery);
-				DrawParams.DrawFeatureVector(SearchContext.GetOrBuildQuery(SearchResult.Database->Schema));
-			}
+		if (CVarAnimMotionMatchDrawQueryEnable.GetValueOnAnyThread())
+		{
+			FDebugDrawParams DrawParams(Context.AnimInstanceProxy, SearchContext.GetWorldBoneTransformAtTime(0.f), SearchResult.Database.Get(), EDebugDrawFlags::DrawQuery);
+			DrawParams.DrawFeatureVector(SearchContext.GetOrBuildQuery(SearchResult.Database->Schema));
 		}
+	}
 #endif // ENABLE_DRAW_DEBUG && ENABLE_ANIM_DEBUG
 
 #if UE_POSE_SEARCH_TRACE_ENABLED
-		const float SearchBestCost = SearchResult.PoseCost.GetTotalCost();
-		const float SearchBruteForceCost = SearchResult.BruteForcePoseCost.GetTotalCost();
-		TraceMotionMatchingState(SearchContext, SearchResult, 0.f, FTransform::Identity, AnimInstance, Context.GetCurrentNodeId(),
-			AnimInstance->GetDeltaSeconds(), true, FObjectTrace::GetWorldElapsedTime(AnimInstance->GetWorld()));
+	const float SearchBestCost = SearchResult.PoseCost.GetTotalCost();
+	const float SearchBruteForceCost = SearchResult.BruteForcePoseCost.GetTotalCost();
+	TraceMotionMatchingState(SearchContext, SearchResult, 0.f, FTransform::Identity, AnimInstance, Context.GetCurrentNodeId(),
+		AnimInstance->GetDeltaSeconds(), true, FObjectTrace::GetWorldElapsedTime(AnimInstance->GetWorld()));
 #endif // UE_POSE_SEARCH_TRACE_ENABLED
-	}
 
 	return SearchResult;
 }
diff --git a/Engine/Plugins/Experimental/Animation/PoseSearch/Source/Runtime/Private/PoseSearchModule.cpp b/Engine/Plugins/Experimental/Animation/PoseSearch/Source/Runtime/Private/PoseSearchModule.cpp
--- a/Engine/Plugins/Experimental/Animation/PoseSearch/Source/Runtime/Private/PoseSearchModule.cpp
+++ b/Engine/Plugins/Experimental/Animation/PoseSearch/Source/Runtime/Private/PoseSearchModule.cpp
@@ -22,9 +22,10 @@ public:
 	}
 
 	// IPoseSearchProvider
-	virtual UE::Anim::IPoseSearchProvider::FSearchResult Search(const FAnimationBaseContext& GraphContext, TConstArrayView<UAnimationAsset*> AnimationAssets) override
+	virtual UE::Anim::IPoseSearchProvider::FSearchResult Search(const FAnimationBaseContext& GraphContext, TConstArrayView<UAnimationAsset*> AnimationAssets,
+		const UAnimationAsset* PlayingAnimationAsset, float PlayingAnimationAssetAccumulatedTime) override
 	{
-		const UE::PoseSearch::FSearchResult SearchResult = UPoseSearchLibrary::MotionMatch(GraphContext, AnimationAssets);
+		const UE::PoseSearch::FSearchResult SearchResult = UPoseSearchLibrary::MotionMatch(GraphContext, AnimationAssets, PlayingAnimationAsset, PlayingAnimationAssetAccumulatedTime);
 		UE::Anim::IPoseSearchProvider::FSearchResult ProviderResult;
 		if (const UE::PoseSearch::FSearchIndexAsset* SearchIndexAsset = SearchResult.GetSearchIndexAsset())
 		{
diff --git a/Engine/Plugins/Experimental/Animation/PoseSearch/Source/Runtime/Public/PoseSearch/PoseSearchContext.h b/Engine/Plugins/Experimental/Animation/PoseSearch/Source/Runtime/Public/PoseSearch/PoseSearchContext.h
--- a/Engine/Plugins/Experimental/Animation/PoseSearch/Source/Runtime/Public/PoseSearch/PoseSearchContext.h
+++ b/Engine/Plugins/Experimental/Animation/PoseSearch/Source/Runtime/Public/PoseSearch/PoseSearchContext.h
@@ -201,6 +201,7 @@ struct POSESEARCH_API FSearchContext
 
 	TConstArrayView<float> GetCurrentResultPoseVector() const { return CurrentResultPoseVector; }
 
+	void UpdateCurrentResultPoseVector();
 	const FSearchResult& GetCurrentResult() const { return CurrentResult; }
 	const FFloatInterval& GetPoseJumpThresholdTime() const { return PoseJumpThresholdTime; }
 	const FPoseIndicesHistory* GetPoseIndicesHistory() const { return PoseIndicesHistory; }
diff --git a/Engine/Plugins/Experimental/Animation/PoseSearch/Source/Runtime/Public/PoseSearch/PoseSearchLibrary.h b/Engine/Plugins/Experimental/Animation/PoseSearch/Source/Runtime/Public/PoseSearch/PoseSearchLibrary.h
--- a/Engine/Plugins/Experimental/Animation/PoseSearch/Source/Runtime/Public/PoseSearch/PoseSearchLibrary.h
+++ b/Engine/Plugins/Experimental/Animation/PoseSearch/Source/Runtime/Public/PoseSearch/PoseSearchLibrary.h
@@ -171,6 +171,7 @@ public:
 		float TimeToFutureAnimationStart = 0.f,
 		const int32 DebugSessionUniqueIdentifier = 6174);
 
-	static UE::PoseSearch::FSearchResult MotionMatch(const FAnimationBaseContext& Context, TConstArrayView<UAnimationAsset*> AnimationAssets);
+	static UE::PoseSearch::FSearchResult MotionMatch(const FAnimationBaseContext& Context, TConstArrayView<UAnimationAsset*> AnimationAssets,
+		const UAnimationAsset* PlayingAnimationAsset = nullptr, float PlayingAnimationAssetAccumulatedTime = 0.f);
 };
 
diff --git a/Engine/Plugins/Experimental/Chooser/Source/Chooser/Private/AnimNode_ChooserPlayer.cpp b/Engine/Plugins/Experimental/Chooser/Source/Chooser/Private/AnimNode_ChooserPlayer.cpp
--- a/Engine/Plugins/Experimental/Chooser/Source/Chooser/Private/AnimNode_ChooserPlayer.cpp
+++ b/Engine/Plugins/Experimental/Chooser/Source/Chooser/Private/AnimNode_ChooserPlayer.cpp
@@ -45,7 +45,7 @@ UAnimationAsset* FAnimNode_ChooserPlayer::ChooseAsset(const FAnimationUpdateCont
 						return FObjectChooserBase::EIteratorStatus::Continue;
 					}));
 
-				const UE::Anim::IPoseSearchProvider::FSearchResult SearchResult = PoseSearchProvider->Search(Context, AnimationAssets);
+				const UE::Anim::IPoseSearchProvider::FSearchResult SearchResult = PoseSearchProvider->Search(Context, AnimationAssets, GetAnimAsset(), GetAccumulatedTime());
 				if (SearchResult.AnimationAsset)
 				{
 					Settings.StartTime = SearchResult.TimeOffsetSeconds;
diff --git a/Engine/Source/Runtime/Engine/Public/Animation/AnimPoseSearchProvider.h b/Engine/Source/Runtime/Engine/Public/Animation/AnimPoseSearchProvider.h
--- a/Engine/Source/Runtime/Engine/Public/Animation/AnimPoseSearchProvider.h
+++ b/Engine/Source/Runtime/Engine/Public/Animation/AnimPoseSearchProvider.h
@@ -34,10 +34,13 @@ public:
 	* 
 	* @param	GraphContext	Graph execution context used to construct a pose search query
 	* @param	AnimationAssets	The animation assets to search for the pose query
+	* @param	PlayingAnimationAsset					The currently playing animation asset, used to bias the score of the eventually found continuing pose
+	* @param	PlayingAnimationAssetAccumulatedTime	The accumulated time of the currently playing animation asset
 	* 
 	* @return	The pose in the AnimationAssets that most closely matches the query
 	*/
-	virtual FSearchResult Search(const FAnimationBaseContext& GraphContext, TConstArrayView<UAnimationAsset*> AnimationAssets) = 0;
+	virtual FSearchResult Search(const FAnimationBaseContext& GraphContext, TConstArrayView<UAnimationAsset*> AnimationAssets, 
+		const UAnimationAsset* PlayingAnimationAsset = nullptr, float PlayingAnimationAssetAccumulatedTime = 0.f) = 0;
 };
 
 } // namespace UE::Anim
\ No newline at end of file
