motion matching - improved workflow to use UAnimNotifyState_PoseSearchBranchIn

diff --git a/Engine/Plugins/Experimental/Animation/PoseSearch/Source/Editor/PoseSearchEditor.Build.cs b/Engine/Plugins/Experimental/Animation/PoseSearch/Source/Editor/PoseSearchEditor.Build.cs
--- a/Engine/Plugins/Experimental/Animation/PoseSearch/Source/Editor/PoseSearchEditor.Build.cs
+++ b/Engine/Plugins/Experimental/Animation/PoseSearch/Source/Editor/PoseSearchEditor.Build.cs
@@ -16,6 +16,7 @@ public class PoseSearchEditor : ModuleRules
 				"BlendStack",
 				"AnimGraph",
 				"AnimGraphRuntime",
+				"AnimationBlueprintLibrary",
 				"AnimationCore",
 				"Core",
 				"CoreUObject",
diff --git a/Engine/Plugins/Experimental/Animation/PoseSearch/Source/Editor/Private/PoseSearchDatabaseAssetTree.cpp b/Engine/Plugins/Experimental/Animation/PoseSearch/Source/Editor/Private/PoseSearchDatabaseAssetTree.cpp
--- a/Engine/Plugins/Experimental/Animation/PoseSearch/Source/Editor/Private/PoseSearchDatabaseAssetTree.cpp
+++ b/Engine/Plugins/Experimental/Animation/PoseSearch/Source/Editor/Private/PoseSearchDatabaseAssetTree.cpp
@@ -5,6 +5,7 @@
 #include "Animation/AnimMontage.h"
 #include "Animation/AnimSequence.h"
 #include "Animation/BlendSpace.h"
+#include "AnimationBlueprintLibrary.h"
 #include "AssetSelection.h"
 #include "ClassIconFinder.h"
 #include "DetailColumnSizeData.h"
@@ -13,6 +14,7 @@
 #include "Framework/MultiBox/MultiBoxBuilder.h"
 #include "Misc/FeedbackContext.h"
 #include "Misc/TransactionObjectEvent.h"
+#include "PoseSearch/PoseSearchAnimNotifies.h"
 #include "PoseSearch/PoseSearchDatabase.h"
 #include "PoseSearchDatabaseViewModel.h"
 #include "SPositiveActionButton.h"
@@ -499,8 +501,8 @@ namespace UE::PoseSearch
 		if (!SelectedNodes.IsEmpty())
 		{
 			MenuBuilder.AddMenuEntry(
-				LOCTEXT("DeleteUngroup", "Delete / Remove"),
-				LOCTEXT("DeleteUngroupTooltip", "Deletes groups and ungrouped assets; removes grouped assets from group."),
+				LOCTEXT("Remove", "Remove"),
+				LOCTEXT("RemoveTooltip", "Removes assets from database"),
 				FSlateIcon(),
 				FUIAction(FExecuteAction::CreateSP(this, &SDatabaseAssetTree::OnDeleteNodes)),
 				NAME_None,
@@ -508,9 +510,7 @@ namespace UE::PoseSearch
 
 			MenuBuilder.AddMenuEntry(
 				LOCTEXT("Enable", "Enable"),
-				LOCTEXT(
-					"EnableTooltip",
-					"Sets Assets Enabled."),
+				LOCTEXT("EnableTooltip", "Sets Assets Enabled."),
 				FSlateIcon(),
 				FUIAction(FExecuteAction::CreateSP(this, &SDatabaseAssetTree::OnEnableNodes)),
 				NAME_None,
@@ -518,13 +518,19 @@ namespace UE::PoseSearch
 
 			MenuBuilder.AddMenuEntry(
 				LOCTEXT("Disable", "Disable"),
-				LOCTEXT(
-					"DisableToolTip",
-					"Sets Assets Disabled."),
+				LOCTEXT("DisableToolTip", "Sets Assets Disabled."),
 				FSlateIcon(),
 				FUIAction(FExecuteAction::CreateSP(this, &SDatabaseAssetTree::OnDisableNodes)),
 				NAME_None,
 				EUserInterfaceActionType::Button);
+
+			MenuBuilder.AddMenuEntry(
+				LOCTEXT("ConvertToBranchIn", "ConvertToBranchIn"),
+				LOCTEXT("ConvertToBranchInToolTip", "Creates PoseSearchBranchIn notify state for the asset sampling range"),
+				FSlateIcon(),
+				FUIAction(FExecuteAction::CreateSP(this, &SDatabaseAssetTree::OnConvertToBranchIn)),
+				NAME_None,
+				EUserInterfaceActionType::Button);
 		}
 
 		return MenuBuilder.MakeWidget();
@@ -727,6 +733,47 @@ namespace UE::PoseSearch
 		}
 	}
 
+	void SDatabaseAssetTree::OnConvertToBranchIn()
+	{
+		const TSharedPtr<FDatabaseViewModel> ViewModel = EditorViewModel.Pin();
+		if (UPoseSearchDatabase* PoseSearchDatabase = ViewModel->GetPoseSearchDatabase())
+		{
+			TArray<TSharedPtr<FDatabaseAssetTreeNode>> SelectedNodes = TreeView->GetSelectedItems();
+			if (!SelectedNodes.IsEmpty())
+			{
+				const FScopedTransaction Transaction(LOCTEXT("ConvertToBranchIn", "Create PoseSearchBranchIn notify state for assets in Pose Search Database"));
+
+				for (TSharedPtr<FDatabaseAssetTreeNode> SelectedNode : SelectedNodes)
+				{
+					if (FPoseSearchDatabaseAnimationAssetBase* DatabaseAnimationAssetBase = PoseSearchDatabase->GetMutableAnimationAssetBase(SelectedNode->SourceAssetIdx))
+					{
+						if (UAnimSequenceBase* AnimSequenceBase = Cast<UAnimSequenceBase>(DatabaseAnimationAssetBase->GetAnimationAsset()))
+						{
+							const FFloatInterval SamplingRange = FPoseSearchDatabaseAnimationAssetBase::GetEffectiveSamplingRange(AnimSequenceBase, DatabaseAnimationAssetBase->GetSamplingRange());
+							const float StartTime = SamplingRange.Min;
+							const float Duration = SamplingRange.Max - SamplingRange.Min;
+							FName TrackName = "PoseSearch";
+							
+							if (!UAnimationBlueprintLibrary::IsValidAnimNotifyTrackName(AnimSequenceBase, TrackName))
+							{
+								UAnimationBlueprintLibrary::AddAnimationNotifyTrack(AnimSequenceBase, TrackName, FColor::Turquoise);
+							}
+
+							UAnimNotifyState_PoseSearchBranchIn* PoseSearchBranchIn = CastChecked<UAnimNotifyState_PoseSearchBranchIn>(UAnimationBlueprintLibrary::AddAnimationNotifyStateEvent(AnimSequenceBase, TrackName, StartTime, Duration, UAnimNotifyState_PoseSearchBranchIn::StaticClass()));
+							PoseSearchBranchIn->Database = PoseSearchDatabase;
+							DatabaseAnimationAssetBase->bSynchronizeWithExternalDependency = true;
+							AnimSequenceBase->Modify();
+						}
+					}
+				}
+
+				PoseSearchDatabase->SynchronizeWithExternalDependencies();
+
+				FinalizeTreeChanges();
+			}
+		}
+	}
+
 	void SDatabaseAssetTree::FinalizeTreeChanges(bool bRecoverSelection)
 	{
 		RefreshTreeView(false, bRecoverSelection);
diff --git a/Engine/Plugins/Experimental/Animation/PoseSearch/Source/Editor/Private/PoseSearchDatabaseViewModel.cpp b/Engine/Plugins/Experimental/Animation/PoseSearch/Source/Editor/Private/PoseSearchDatabaseViewModel.cpp
--- a/Engine/Plugins/Experimental/Animation/PoseSearch/Source/Editor/Private/PoseSearchDatabaseViewModel.cpp
+++ b/Engine/Plugins/Experimental/Animation/PoseSearch/Source/Editor/Private/PoseSearchDatabaseViewModel.cpp
@@ -510,13 +510,33 @@ bool FDatabaseViewModel::DeleteFromDatabase(int32 AnimationAssetIndex)
 	{
 		if (const FPoseSearchDatabaseAnimationAssetBase* DatabaseAnimationAssetBase = Database->GetAnimationAssetBase(AnimationAssetIndex))
 		{
-			if (!DatabaseAnimationAssetBase->bSynchronizeWithExternalDependency)
+			if (UAnimSequenceBase* AnimSequenceBase = Cast<UAnimSequenceBase>(DatabaseAnimationAssetBase->GetAnimationAsset()))
 			{
-				Database->AnimationAssets.RemoveAt(AnimationAssetIndex);
-				Database->Modify();
+				bool bModified = false;
+				for (int32 NotifyIndex = AnimSequenceBase->Notifies.Num() - 1; NotifyIndex >= 0; --NotifyIndex)
+				{
+					const FAnimNotifyEvent& NotifyEvent = AnimSequenceBase->Notifies[NotifyIndex];
+					if (const UAnimNotifyState_PoseSearchBranchIn* PoseSearchBranchIn = Cast<UAnimNotifyState_PoseSearchBranchIn>(NotifyEvent.NotifyStateClass))
+					{
+						if (PoseSearchBranchIn->Database == Database)
+						{
+							AnimSequenceBase->Notifies.RemoveAt(NotifyIndex);
+							bModified = true;
+						}
+					}
+				}
 
-				return true;
+				if (bModified)
+				{
+					AnimSequenceBase->RefreshCacheData();	
+					AnimSequenceBase->Modify();
+				}
 			}
+
+			Database->AnimationAssets.RemoveAt(AnimationAssetIndex);
+			Database->Modify();
+
+			return true;
 		}
 	}
 
diff --git a/Engine/Plugins/Experimental/Animation/PoseSearch/Source/Editor/Public/PoseSearchDatabaseAssetTree.h b/Engine/Plugins/Experimental/Animation/PoseSearch/Source/Editor/Public/PoseSearchDatabaseAssetTree.h
--- a/Engine/Plugins/Experimental/Animation/PoseSearch/Source/Editor/Public/PoseSearchDatabaseAssetTree.h
+++ b/Engine/Plugins/Experimental/Animation/PoseSearch/Source/Editor/Public/PoseSearchDatabaseAssetTree.h
@@ -101,6 +101,8 @@ namespace UE::PoseSearch
 		void OnEnableNodes() { EnableSelectedNodes(true); }
 		void OnDisableNodes() { EnableSelectedNodes(false); }
 
+		void OnConvertToBranchIn();
+
 		friend SDatabaseAssetListItem;
 
 		FDetailColumnSizeData ColumnSizeData;
diff --git a/Engine/Plugins/Experimental/Animation/PoseSearch/Source/Runtime/Private/PoseSearchDatabase.cpp b/Engine/Plugins/Experimental/Animation/PoseSearch/Source/Runtime/Private/PoseSearchDatabase.cpp
--- a/Engine/Plugins/Experimental/Animation/PoseSearch/Source/Runtime/Private/PoseSearchDatabase.cpp
+++ b/Engine/Plugins/Experimental/Animation/PoseSearch/Source/Runtime/Private/PoseSearchDatabase.cpp
@@ -314,6 +314,16 @@ int64 FPoseSearchDatabaseAnimationAssetBase::GetEditorMemSize() const
 	FArchiveCountMem EditorMemCount(GetAnimationAsset());
 	return EditorMemCount.GetNum();
 }
+
+FFloatInterval FPoseSearchDatabaseAnimationAssetBase::GetEffectiveSamplingRange(const UAnimSequenceBase* SequenceBase, const FFloatInterval& RequestedSamplingRange)
+{
+	const bool bSampleAll = (RequestedSamplingRange.Min == 0.0f) && (RequestedSamplingRange.Max == 0.0f);
+	const float SequencePlayLength = SequenceBase->GetPlayLength();
+	FFloatInterval Range;
+	Range.Min = bSampleAll ? 0.0f : RequestedSamplingRange.Min;
+	Range.Max = bSampleAll ? SequencePlayLength : FMath::Min(SequencePlayLength, RequestedSamplingRange.Max);
+	return Range;
+}
 #endif // WITH_EDITORONLY_DATA
 
 //////////////////////////////////////////////////////////////////////////
diff --git a/Engine/Plugins/Experimental/Animation/PoseSearch/Source/Runtime/Private/PoseSearchDerivedData.cpp b/Engine/Plugins/Experimental/Animation/PoseSearch/Source/Runtime/Private/PoseSearchDerivedData.cpp
--- a/Engine/Plugins/Experimental/Animation/PoseSearch/Source/Runtime/Private/PoseSearchDerivedData.cpp
+++ b/Engine/Plugins/Experimental/Animation/PoseSearch/Source/Runtime/Private/PoseSearchDerivedData.cpp
@@ -266,16 +266,6 @@ public:
 	}
 };
 
-static inline FFloatInterval GetEffectiveSamplingRange(const UAnimSequenceBase* Sequence, FFloatInterval RequestedSamplingRange)
-{
-	const bool bSampleAll = (RequestedSamplingRange.Min == 0.0f) && (RequestedSamplingRange.Max == 0.0f);
-	const float SequencePlayLength = Sequence->GetPlayLength();
-	FFloatInterval Range;
-	Range.Min = bSampleAll ? 0.0f : RequestedSamplingRange.Min;
-	Range.Max = bSampleAll ? SequencePlayLength : FMath::Min(SequencePlayLength, RequestedSamplingRange.Max);
-	return Range;
-}
-
 static void FindValidSequenceIntervals(const UAnimSequenceBase* SequenceBase, FFloatInterval SamplingRange, bool bIsLooping,
 	const FFloatInterval& ExcludeFromDatabaseParameters, TArray<FFloatRange>& ValidRanges)
 {
@@ -283,7 +273,7 @@ static void FindValidSequenceIntervals(const UAnimSequenceBase* SequenceBase, FF
 
 	const float SequenceLength = SequenceBase->GetPlayLength();
 
-	const FFloatInterval EffectiveSamplingInterval = GetEffectiveSamplingRange(SequenceBase, SamplingRange);
+	const FFloatInterval EffectiveSamplingInterval = FPoseSearchDatabaseAnimationAssetBase::GetEffectiveSamplingRange(SequenceBase, SamplingRange);
 	FFloatRange EffectiveSamplingRange = FFloatRange::Inclusive(EffectiveSamplingInterval.Min, EffectiveSamplingInterval.Max);
 	if (!bIsLooping)
 	{
diff --git a/Engine/Plugins/Experimental/Animation/PoseSearch/Source/Runtime/Public/PoseSearch/PoseSearchDatabase.h b/Engine/Plugins/Experimental/Animation/PoseSearch/Source/Runtime/Public/PoseSearch/PoseSearchDatabase.h
--- a/Engine/Plugins/Experimental/Animation/PoseSearch/Source/Runtime/Public/PoseSearch/PoseSearchDatabase.h
+++ b/Engine/Plugins/Experimental/Animation/PoseSearch/Source/Runtime/Public/PoseSearch/PoseSearchDatabase.h
@@ -65,6 +65,7 @@ struct POSESEARCH_API FPoseSearchDatabaseAnimationAssetBase
 
 	// [0, 0] represents the entire frame range of the original animation.
 	virtual FFloatInterval GetSamplingRange() const { return FFloatInterval(0.f, 0.f); }
+	static FFloatInterval GetEffectiveSamplingRange(const UAnimSequenceBase* Sequence, const FFloatInterval& RequestedSamplingRange);
 
 	virtual int64 GetEditorMemSize() const;
 	virtual int64 GetApproxCookedSize() const { return GetEditorMemSize(); }
