motion matching - fix for text conflict from LOCTEXT macro for namespace "SDatabaseAssetTree" and key "ConvertToBranchIn"

diff --git a/Engine/Plugins/Experimental/Animation/PoseSearch/Source/Editor/Private/PoseSearchDatabaseAssetTree.cpp b/Engine/Plugins/Experimental/Animation/PoseSearch/Source/Editor/Private/PoseSearchDatabaseAssetTree.cpp
--- a/Engine/Plugins/Experimental/Animation/PoseSearch/Source/Editor/Private/PoseSearchDatabaseAssetTree.cpp
+++ b/Engine/Plugins/Experimental/Animation/PoseSearch/Source/Editor/Private/PoseSearchDatabaseAssetTree.cpp
@@ -741,7 +741,7 @@ namespace UE::PoseSearch
 			TArray<TSharedPtr<FDatabaseAssetTreeNode>> SelectedNodes = TreeView->GetSelectedItems();
 			if (!SelectedNodes.IsEmpty())
 			{
-				const FScopedTransaction Transaction(LOCTEXT("ConvertToBranchIn", "Create PoseSearchBranchIn notify state for assets in Pose Search Database"));
+				const FScopedTransaction Transaction(LOCTEXT("ConvertToBranchInTransaction", "Create PoseSearchBranchIn notify state for assets in Pose Search Database"));
 
 				for (TSharedPtr<FDatabaseAssetTreeNode> SelectedNode : SelectedNodes)
 				{
