motion matching - support for non homogeneous chooser tables, containing UAnimationAsset(s) or UPoseSearchDatabase(s), while doing a pose search in the chooser anim node

diff --git a/Engine/Plugins/Experimental/Animation/PoseSearch/Source/Editor/Private/PoseSearchDatabaseAssetListItem.cpp b/Engine/Plugins/Experimental/Animation/PoseSearch/Source/Editor/Private/PoseSearchDatabaseAssetListItem.cpp
--- a/Engine/Plugins/Experimental/Animation/PoseSearch/Source/Editor/Private/PoseSearchDatabaseAssetListItem.cpp
+++ b/Engine/Plugins/Experimental/Animation/PoseSearch/Source/Editor/Private/PoseSearchDatabaseAssetListItem.cpp
@@ -123,25 +123,26 @@ namespace UE::PoseSearch
 			{
 				if (UAssetEditorSubsystem* AssetEditorSS = GEditor->GetEditorSubsystem<UAssetEditorSubsystem>())
 				{
-					UAnimationAsset* AnimationAsset = DatabaseAnimationAsset->GetAnimationAsset();
-
-					AssetEditorSS->OpenEditorForAsset(AnimationAsset);
-
-					if (IAssetEditorInstance* Editor = AssetEditorSS->FindEditorForAsset(AnimationAsset, true))
+					if (UObject* AnimationAsset = DatabaseAnimationAsset->GetAnimationAsset())
 					{
-						if (Editor->GetEditorName() == "AnimationEditor")
+						AssetEditorSS->OpenEditorForAsset(AnimationAsset);
+
+						if (IAssetEditorInstance* Editor = AssetEditorSS->FindEditorForAsset(AnimationAsset, true))
 						{
-							float AnimationAssetTime = 0.f;
-							FVector AnimationAssetBlendParameters = FVector::ZeroVector;
-							ViewModel->GetAnimationTime(Node->SourceAssetIdx, AnimationAssetTime, AnimationAssetBlendParameters);
-							
-							const IAnimationEditor* AnimationEditor = static_cast<IAnimationEditor*>(Editor);
-							const UDebugSkelMeshComponent* PreviewComponent = AnimationEditor->GetPersonaToolkit()->GetPreviewMeshComponent();
-
-							// Open asset paused and at specific time as seen on the pose search debugger.
-							PreviewComponent->PreviewInstance->SetPosition(AnimationAssetTime);
-							PreviewComponent->PreviewInstance->SetPlaying(false);
-							PreviewComponent->PreviewInstance->SetBlendSpacePosition(AnimationAssetBlendParameters);
+							if (Editor->GetEditorName() == "AnimationEditor")
+							{
+								float AnimationAssetTime = 0.f;
+								FVector AnimationAssetBlendParameters = FVector::ZeroVector;
+								ViewModel->GetAnimationTime(Node->SourceAssetIdx, AnimationAssetTime, AnimationAssetBlendParameters);
+
+								const IAnimationEditor* AnimationEditor = static_cast<IAnimationEditor*>(Editor);
+								const UDebugSkelMeshComponent* PreviewComponent = AnimationEditor->GetPersonaToolkit()->GetPreviewMeshComponent();
+
+								// Open asset paused and at specific time as seen on the pose search debugger.
+								PreviewComponent->PreviewInstance->SetPosition(AnimationAssetTime);
+								PreviewComponent->PreviewInstance->SetPlaying(false);
+								PreviewComponent->PreviewInstance->SetBlendSpacePosition(AnimationAssetBlendParameters);
+							}
 						}
 					}
 				}
diff --git a/Engine/Plugins/Experimental/Animation/PoseSearch/Source/Editor/Private/PoseSearchDatabaseEditorReflection.cpp b/Engine/Plugins/Experimental/Animation/PoseSearch/Source/Editor/Private/PoseSearchDatabaseEditorReflection.cpp
--- a/Engine/Plugins/Experimental/Animation/PoseSearch/Source/Editor/Private/PoseSearchDatabaseEditorReflection.cpp
+++ b/Engine/Plugins/Experimental/Animation/PoseSearch/Source/Editor/Private/PoseSearchDatabaseEditorReflection.cpp
@@ -169,7 +169,7 @@ void UPoseSearchDatabaseStatistics::Initialize(const UPoseSearchDatabase* PoseSe
 			
 		{
 			uint32 SourceAnimAssetsSizeCookedEstimateInBytes = 0;
-			TSet<const UAnimationAsset*> Analyzed;
+			TSet<const UObject*> Analyzed;
 			Analyzed.Reserve(PoseSearchDatabase->AnimationAssets.Num());
 			for (const FInstancedStruct& AnimAsset : PoseSearchDatabase->AnimationAssets)
 			{
diff --git a/Engine/Plugins/Experimental/Animation/PoseSearch/Source/Editor/Private/PoseSearchDatabaseViewModel.cpp b/Engine/Plugins/Experimental/Animation/PoseSearch/Source/Editor/Private/PoseSearchDatabaseViewModel.cpp
--- a/Engine/Plugins/Experimental/Animation/PoseSearch/Source/Editor/Private/PoseSearchDatabaseViewModel.cpp
+++ b/Engine/Plugins/Experimental/Animation/PoseSearch/Source/Editor/Private/PoseSearchDatabaseViewModel.cpp
@@ -42,7 +42,7 @@ bool FDatabasePreviewActor::SpawnPreviewActor(UWorld* World, const UPoseSearchDa
 	const FSearchIndexAsset& IndexAsset = SearchIndex.Assets[IndexAssetIdx];
 
 	const FPoseSearchDatabaseAnimationAssetBase* DatabaseAnimationAsset = PoseSearchDatabase->GetAnimationAssetBase(IndexAsset.GetSourceAssetIdx());
-	UAnimationAsset* PreviewAsset = DatabaseAnimationAsset->GetAnimationAsset();
+	UAnimationAsset* PreviewAsset = Cast<UAnimationAsset>(DatabaseAnimationAsset->GetAnimationAsset());
 	if (!PreviewAsset)
 	{
 		return false;
diff --git a/Engine/Plugins/Experimental/Animation/PoseSearch/Source/Runtime/Private/AnimNode_MotionMatching.cpp b/Engine/Plugins/Experimental/Animation/PoseSearch/Source/Runtime/Private/AnimNode_MotionMatching.cpp
--- a/Engine/Plugins/Experimental/Animation/PoseSearch/Source/Runtime/Private/AnimNode_MotionMatching.cpp
+++ b/Engine/Plugins/Experimental/Animation/PoseSearch/Source/Runtime/Private/AnimNode_MotionMatching.cpp
@@ -197,11 +197,18 @@ void FAnimNode_MotionMatching::UpdateAssetPlayer(const FAnimationUpdateContext&
 		{
 			if (const FPoseSearchDatabaseAnimationAssetBase* DatabaseAsset = CurrentResultDatabase->GetAnimationAssetBase(*SearchIndexAsset))
 			{
-				// root bone blending needs to be immediate if MM node controls the offset between mesh component and root bone
-				const float RootBoneBlendTime = YawFromAnimationBlendRate < 0.f ? BlendTime : 0.f;
-				FAnimNode_BlendStack_Standalone::BlendTo(Context, DatabaseAsset->GetAnimationAsset(), MotionMatchingState.CurrentSearchResult.AssetTime,
-					SearchIndexAsset->IsLooping(), SearchIndexAsset->IsMirrored(), CurrentResultDatabase->Schema->MirrorDataTable.Get(), BlendTime,
-					RootBoneBlendTime, BlendProfile, BlendOption, bUseInertialBlend, SearchIndexAsset->GetBlendParameters(), MotionMatchingState.WantedPlayRate);
+				if (UAnimationAsset* AnimationAsset = Cast<UAnimationAsset>(DatabaseAsset->GetAnimationAsset()))
+				{
+					// root bone blending needs to be immediate if MM node controls the offset between mesh component and root bone
+					const float RootBoneBlendTime = YawFromAnimationBlendRate < 0.f ? BlendTime : 0.f;
+					FAnimNode_BlendStack_Standalone::BlendTo(Context, AnimationAsset, MotionMatchingState.CurrentSearchResult.AssetTime,
+						SearchIndexAsset->IsLooping(), SearchIndexAsset->IsMirrored(), CurrentResultDatabase->Schema->MirrorDataTable.Get(), BlendTime,
+						RootBoneBlendTime, BlendProfile, BlendOption, bUseInertialBlend, SearchIndexAsset->GetBlendParameters(), MotionMatchingState.WantedPlayRate);
+				}
+				else
+				{
+					unimplemented();
+				}
 			}
 		}
 	}
diff --git a/Engine/Plugins/Experimental/Animation/PoseSearch/Source/Runtime/Private/PoseSearchContext.cpp b/Engine/Plugins/Experimental/Animation/PoseSearch/Source/Runtime/Private/PoseSearchContext.cpp
--- a/Engine/Plugins/Experimental/Animation/PoseSearch/Source/Runtime/Private/PoseSearchContext.cpp
+++ b/Engine/Plugins/Experimental/Animation/PoseSearch/Source/Runtime/Private/PoseSearchContext.cpp
@@ -360,12 +360,12 @@ FCachedQuery::FCachedQuery(const UPoseSearchSchema* InSchema)
 
 //////////////////////////////////////////////////////////////////////////
 // FSearchContext
-FSearchContext::FSearchContext(const UAnimInstance* InAnimInstance, const IPoseHistory* InHistory, TConstArrayView<const UAnimationAsset*> InAnimationsToConsider,
+FSearchContext::FSearchContext(const UAnimInstance* InAnimInstance, const IPoseHistory* InHistory, TConstArrayView<const UObject*> InAssetsToConsider,
 		float InDesiredPermutationTimeOffset, const FPoseIndicesHistory* InPoseIndicesHistory,
 		const FSearchResult& InCurrentResult, const FFloatInterval& InPoseJumpThresholdTime, bool bInUseCachedChannelData)
 : AnimInstance(InAnimInstance)
 , History(InHistory)
-, AnimationsToConsider(InAnimationsToConsider)
+, AssetsToConsider(InAssetsToConsider)
 , DesiredPermutationTimeOffset(InDesiredPermutationTimeOffset)
 , PoseIndicesHistory(InPoseIndicesHistory)
 , CurrentResult(InCurrentResult)
diff --git a/Engine/Plugins/Experimental/Animation/PoseSearch/Source/Runtime/Private/PoseSearchDatabase.cpp b/Engine/Plugins/Experimental/Animation/PoseSearch/Source/Runtime/Private/PoseSearchDatabase.cpp
--- a/Engine/Plugins/Experimental/Animation/PoseSearch/Source/Runtime/Private/PoseSearchDatabase.cpp
+++ b/Engine/Plugins/Experimental/Animation/PoseSearch/Source/Runtime/Private/PoseSearchDatabase.cpp
@@ -41,11 +41,11 @@ static TAutoConsoleVariable<bool> CVarMotionMatchValidateKNNSearch(TEXT("a.Motio
 #endif
 
 typedef TArray<int32, TInlineAllocator<256>> FSelectableAssetIdx;
-static void PopulateSelectableAssetIdx(FSelectableAssetIdx& SelectableAssetIdx, TConstArrayView<const UAnimationAsset*> AnimationsToConsider, const UPoseSearchDatabase* Database)
+static void PopulateSelectableAssetIdx(FSelectableAssetIdx& SelectableAssetIdx, TConstArrayView<const UObject*> AssetsToConsider, const UPoseSearchDatabase* Database)
 {
 	check(Database);
 	SelectableAssetIdx.Reset();
-	if (!AnimationsToConsider.IsEmpty())
+	if (!AssetsToConsider.IsEmpty())
 	{
 		const FSearchIndex& SearchIndex = Database->GetSearchIndex();
 
@@ -53,7 +53,7 @@ static void PopulateSelectableAssetIdx(FSelectableAssetIdx& SelectableAssetIdx,
 		{
 			if (const FPoseSearchDatabaseAnimationAssetBase* DatabaseAnimationAssetBase = Database->GetAnimationAssetBase(SearchIndex.Assets[AssetIndex]))
 			{
-				if (AnimationsToConsider.Contains(DatabaseAnimationAssetBase->GetAnimationAsset()))
+				if (AssetsToConsider.Contains(DatabaseAnimationAssetBase->GetAnimationAsset()))
 				{
 					SelectableAssetIdx.Add(AssetIndex);
 				}
@@ -308,6 +308,18 @@ private:
 
 //////////////////////////////////////////////////////////////////////////
 // FPoseSearchDatabaseAnimationAssetBase
+
+float FPoseSearchDatabaseAnimationAssetBase::GetPlayLength() const
+{
+	if (const UAnimationAsset* AnimationAsset = Cast<UAnimationAsset>(GetAnimationAsset()))
+	{
+		return AnimationAsset->GetPlayLength();
+	}
+
+	checkNoEntry();
+	return 0;
+}
+
 #if WITH_EDITORONLY_DATA
 int64 FPoseSearchDatabaseAnimationAssetBase::GetEditorMemSize() const
 {
@@ -328,7 +340,7 @@ FFloatInterval FPoseSearchDatabaseAnimationAssetBase::GetEffectiveSamplingRange(
 
 //////////////////////////////////////////////////////////////////////////
 // FPoseSearchDatabaseSequence
-UAnimationAsset* FPoseSearchDatabaseSequence::GetAnimationAsset() const
+UObject* FPoseSearchDatabaseSequence::GetAnimationAsset() const
 {
 	return Sequence.Get();
 }
@@ -360,7 +372,7 @@ bool FPoseSearchDatabaseSequence::IsRootMotionEnabled() const
 
 //////////////////////////////////////////////////////////////////////////
 // FPoseSearchDatabaseBlendSpace
-UAnimationAsset* FPoseSearchDatabaseBlendSpace::GetAnimationAsset() const
+UObject* FPoseSearchDatabaseBlendSpace::GetAnimationAsset() const
 {
 	return BlendSpace.Get();
 }
@@ -471,7 +483,7 @@ FVector FPoseSearchDatabaseBlendSpace::BlendParameterForSampleRanges(int32 Horiz
 
 //////////////////////////////////////////////////////////////////////////
 // FPoseSearchDatabaseAnimComposite
-UAnimationAsset* FPoseSearchDatabaseAnimComposite::GetAnimationAsset() const
+UObject* FPoseSearchDatabaseAnimComposite::GetAnimationAsset() const
 {
 	return AnimComposite.Get();
 }
@@ -503,7 +515,7 @@ bool FPoseSearchDatabaseAnimComposite::IsRootMotionEnabled() const
 
 //////////////////////////////////////////////////////////////////////////
 // FPoseSearchDatabaseAnimMontage
-UAnimationAsset* FPoseSearchDatabaseAnimMontage::GetAnimationAsset() const
+UObject* FPoseSearchDatabaseAnimMontage::GetAnimationAsset() const
 {
 	return AnimMontage.Get();
 }
@@ -968,7 +980,8 @@ UE::PoseSearch::FSearchResult UPoseSearchDatabase::SearchContinuingPose(UE::Pose
 	const FSearchIndexAsset& SearchIndexAsset = SearchIndex.GetAssetForPose(PoseIdx);
 	const FPoseSearchDatabaseAnimationAssetBase* DatabaseAnimationAssetBase = GetAnimationAssetStruct(SearchIndexAsset).GetPtr<FPoseSearchDatabaseAnimationAssetBase>();
 	check(DatabaseAnimationAssetBase);
-	const FAnimationAssetSampler SequenceBaseSampler(DatabaseAnimationAssetBase->GetAnimationAsset(), SearchIndexAsset.GetBlendParameters());
+	const UAnimationAsset* AnimationAsset = CastChecked<UAnimationAsset>(DatabaseAnimationAssetBase->GetAnimationAsset());
+	const FAnimationAssetSampler SequenceBaseSampler(AnimationAsset, SearchIndexAsset.GetBlendParameters());
 	const float SampleTime = GetRealAssetTime(PoseIdx);
 
 	float UpdatedContinuingPoseCostBias = ContinuingPoseCostBias;
@@ -1043,7 +1056,7 @@ UE::PoseSearch::FSearchResult UPoseSearchDatabase::SearchPCAKDTree(UE::PoseSearc
 		TConstArrayView<float> QueryValues = SearchContext.GetOrBuildQuery(Schema);
 
 		FSelectableAssetIdx SelectableAssetIdx;
-		PopulateSelectableAssetIdx(SelectableAssetIdx, SearchContext.GetAnimationsToConsider(), this);
+		PopulateSelectableAssetIdx(SelectableAssetIdx, SearchContext.GetAssetsToConsider(), this);
 
 		FNonSelectableIdx NonSelectableIdx;
 		PopulateNonSelectableIdx(NonSelectableIdx, SearchContext, this
@@ -1217,7 +1230,7 @@ UE::PoseSearch::FSearchResult UPoseSearchDatabase::SearchVPTree(UE::PoseSearch::
 		TConstArrayView<float> QueryValues = SearchContext.GetOrBuildQuery(Schema);
 
 		FSelectableAssetIdx SelectableAssetIdx;
-		PopulateSelectableAssetIdx(SelectableAssetIdx, SearchContext.GetAnimationsToConsider(), this);
+		PopulateSelectableAssetIdx(SelectableAssetIdx, SearchContext.GetAssetsToConsider(), this);
 
 		// @todo: implement filtering within the VPTree as KDTree does
 		FNonSelectableIdx NonSelectableIdx;
@@ -1304,7 +1317,7 @@ UE::PoseSearch::FSearchResult UPoseSearchDatabase::SearchBruteForce(UE::PoseSear
 		TConstArrayView<float> QueryValues = SearchContext.GetOrBuildQuery(Schema);
 
 		FSelectableAssetIdx SelectableAssetIdx;
-		PopulateSelectableAssetIdx(SelectableAssetIdx, SearchContext.GetAnimationsToConsider(), this);
+		PopulateSelectableAssetIdx(SelectableAssetIdx, SearchContext.GetAssetsToConsider(), this);
 
 		FNonSelectableIdx NonSelectableIdx;
 		PopulateNonSelectableIdx(NonSelectableIdx, SearchContext, this
diff --git a/Engine/Plugins/Experimental/Animation/PoseSearch/Source/Runtime/Private/PoseSearchDerivedData.cpp b/Engine/Plugins/Experimental/Animation/PoseSearch/Source/Runtime/Private/PoseSearchDerivedData.cpp
--- a/Engine/Plugins/Experimental/Animation/PoseSearch/Source/Runtime/Private/PoseSearchDerivedData.cpp
+++ b/Engine/Plugins/Experimental/Animation/PoseSearch/Source/Runtime/Private/PoseSearchDerivedData.cpp
@@ -787,7 +787,7 @@ static bool IndexDatabase(FSearchIndexBase& SearchIndexBase, const UPoseSearchDa
 	TArray<FAnimationAssetSampler> Samplers;
 	Samplers.Reserve(256);
 	
-	TMap<TPair<const UAnimationAsset*, FVector>, int32> SamplerMap;
+	TMap<TPair<const UObject*, FVector>, int32> SamplerMap;
 	SamplerMap.Reserve(256);
 
 	for (const FInstancedStruct& DatabaseAssetStruct : Database.AnimationAssets)
@@ -816,8 +816,9 @@ static bool IndexDatabase(FSearchIndexBase& SearchIndexBase, const UPoseSearchDa
 		}
 		else if (const FPoseSearchDatabaseAnimationAssetBase* DatabaseAnimationAssetBase = DatabaseAssetStruct.GetPtr<FPoseSearchDatabaseAnimationAssetBase>())
 		{
-			if (const UAnimationAsset* AnimationAsset = DatabaseAnimationAssetBase->GetAnimationAsset())
+			if (const UObject* AnimationAssetObject = DatabaseAnimationAssetBase->GetAnimationAsset())
 			{
+				const UAnimationAsset* AnimationAsset = CastChecked<UAnimationAsset>(AnimationAssetObject);
 				if (!SamplerMap.Contains({ AnimationAsset, FVector::ZeroVector }))
 				{
 					SamplerMap.Add({ AnimationAsset, FVector::ZeroVector }, Samplers.Num());
diff --git a/Engine/Plugins/Experimental/Animation/PoseSearch/Source/Runtime/Private/PoseSearchLibrary.cpp b/Engine/Plugins/Experimental/Animation/PoseSearch/Source/Runtime/Private/PoseSearchLibrary.cpp
--- a/Engine/Plugins/Experimental/Animation/PoseSearch/Source/Runtime/Private/PoseSearchLibrary.cpp
+++ b/Engine/Plugins/Experimental/Animation/PoseSearch/Source/Runtime/Private/PoseSearchLibrary.cpp
@@ -387,7 +387,7 @@ void UPoseSearchLibrary::UpdateMotionMatchingState(
 		InOutMotionMatchingState.CurrentSearchResult.Reset();
 	}
 
-	FSearchContext SearchContext(AnimInstance, PoseHistory, TConstArrayView<const UAnimationAsset*>(), 0.f,
+	FSearchContext SearchContext(AnimInstance, PoseHistory, TConstArrayView<const UObject*>(), 0.f,
 		&InOutMotionMatchingState.PoseIndicesHistory, InOutMotionMatchingState.CurrentSearchResult, PoseJumpThresholdTime);
 
 	const bool bCanAdvance = InOutMotionMatchingState.CurrentSearchResult.CanAdvance(DeltaTime);
@@ -650,7 +650,7 @@ void UPoseSearchLibrary::MotionMatch(
 #endif // ENABLE_DRAW_DEBUG && ENABLE_ANIM_DEBUG
 		}
 
-		FSearchContext SearchContext(AnimInstance, ExtendedPoseHistory.IsInitialized() ? &ExtendedPoseHistory : nullptr, TConstArrayView<const UAnimationAsset*>(), TimeToFutureAnimationStart);
+		FSearchContext SearchContext(AnimInstance, ExtendedPoseHistory.IsInitialized() ? &ExtendedPoseHistory : nullptr, TConstArrayView<const UObject*>(), TimeToFutureAnimationStart);
 
 		FSearchResult SearchResult = Database->Search(SearchContext);
 		if (SearchResult.IsValid())
@@ -694,8 +694,8 @@ void UPoseSearchLibrary::MotionMatch(
 	}
 }
 
-UE::PoseSearch::FSearchResult UPoseSearchLibrary::MotionMatch(const FAnimationBaseContext& Context, TConstArrayView<UAnimationAsset*> AnimationAssets,
-	const UAnimationAsset* PlayingAnimationAsset, float PlayingAnimationAssetAccumulatedTime)
+UE::PoseSearch::FSearchResult UPoseSearchLibrary::MotionMatch(const FAnimationBaseContext& Context, 
+	TConstArrayView<UObject*> AssetsToSearch, const UObject* PlayingAsset, float PlayingAssetAccumulatedTime)
 {
 	using namespace UE::PoseSearch;
 
@@ -712,20 +712,21 @@ UE::PoseSearch::FSearchResult UPoseSearchLibrary::MotionMatch(const FAnimationBa
 
 	FMemMark Mark(FMemStack::Get());
 	FSearchResult ReconstructedPreviousSearchResult;
-	FSearchContext SearchContext(AnimInstance, History, TConstArrayView<const UAnimationAsset*>(), 0.f, nullptr, ReconstructedPreviousSearchResult);
+	FSearchContext SearchContext(AnimInstance, History, TConstArrayView<const UObject*>(), 0.f, nullptr, ReconstructedPreviousSearchResult);
 
 	// budgeting some stack allocations for simple use cases. bigger requests of AnimationAssets contining 
 	// UAnimNotifyState_PoseSearchBranchIn referencing multiple datbases will default to slower heap allocations
 	enum { MAX_STACK_ALLOCATED_ANIMATIONS = 16 };
 	enum { MAX_STACK_ALLOCATED_SETS = 2 };
-	typedef	TArray<const UAnimationAsset*, TInlineAllocator<MAX_STACK_ALLOCATED_ANIMATIONS>> TDbAnims;
-	typedef TMap<const UPoseSearchDatabase*, TDbAnims, TInlineSetAllocator<MAX_STACK_ALLOCATED_SETS>> FPerDbAnimMap;
-	typedef TPair<const UPoseSearchDatabase*, TDbAnims> FPerDbAnimPair;
-	FPerDbAnimMap PerDbAnimMap;
+	typedef	TArray<const UObject*, TInlineAllocator<MAX_STACK_ALLOCATED_ANIMATIONS>> TAssetsToSearch;
+	// an empty TAssetsToSearch associated to Database means we need to search ALL the assets
+	typedef TMap<const UPoseSearchDatabase*, TAssetsToSearch, TInlineSetAllocator<MAX_STACK_ALLOCATED_SETS>> TAssetsToSearchPerDatabaseMap;
+	typedef TPair<const UPoseSearchDatabase*, TAssetsToSearch> TAssetsToSearchPerDatabasePair;
+	TAssetsToSearchPerDatabaseMap AssetsToSearchPerDatabaseMap;
 	
-	auto AddToPerDbAnimMap = [](FPerDbAnimMap& PerDbAnimMap, const UAnimationAsset* AnimationAsset)
+	auto AddToSearch = [](TAssetsToSearchPerDatabaseMap& AssetsToSearchPerDatabaseMap, const UObject* AssetToSearch)
 	{
-		if (const UAnimSequenceBase* SequenceBase = Cast<const UAnimSequenceBase>(AnimationAsset))
+		if (const UAnimSequenceBase* SequenceBase = Cast<const UAnimSequenceBase>(AssetToSearch))
 		{
 			for (const FAnimNotifyEvent& NotifyEvent : SequenceBase->Notifies)
 			{
@@ -733,7 +734,18 @@ UE::PoseSearch::FSearchResult UPoseSearchLibrary::MotionMatch(const FAnimationBa
 				{
 					if (PoseSearchBranchIn->Database)
 					{
-						PerDbAnimMap.FindOrAdd(PoseSearchBranchIn->Database).AddUnique(SequenceBase);
+						if (TAssetsToSearch* AssetsToSearch = AssetsToSearchPerDatabaseMap.Find(PoseSearchBranchIn->Database))
+						{
+							// an empty TAssetsToSearch associated to Database means we need to search ALL the assets, so we don't need to add this SequenceBase
+							if (!AssetsToSearch->IsEmpty())
+							{
+								AssetsToSearch->AddUnique(SequenceBase);
+							}
+						}
+						else
+						{
+							AssetsToSearchPerDatabaseMap.Add(PoseSearchBranchIn->Database).AddUnique(SequenceBase);
+						}
 					}
 					else
 					{
@@ -742,16 +754,20 @@ UE::PoseSearch::FSearchResult UPoseSearchLibrary::MotionMatch(const FAnimationBa
 				}
 			}
 		}
+		else if (const UPoseSearchDatabase* Database = Cast<UPoseSearchDatabase>(AssetToSearch))
+		{
+			// an empty TAssetsToSearch associated to Database means we need to search ALL the assets
+			AssetsToSearchPerDatabaseMap.FindOrAdd(Database).Reset();
+		}
 	};
 
 	// collecting all the possible continuing pose search (it could be multiple searches, but most likely only one)
-	if (PlayingAnimationAsset)
+	if (const UAnimationAsset* PlayingAnimationAsset = Cast<UAnimationAsset>(PlayingAsset))
 	{
-		AddToPerDbAnimMap(PerDbAnimMap, PlayingAnimationAsset);
-		for (const FPerDbAnimPair& PerDbAnimPair : PerDbAnimMap)
+		AddToSearch(AssetsToSearchPerDatabaseMap, PlayingAnimationAsset);
+		for (const TAssetsToSearchPerDatabasePair& AssetsToSearchPerDatabasePair : AssetsToSearchPerDatabaseMap)
 		{
-			const UPoseSearchDatabase* Database = PerDbAnimPair.Key;
-			check(Database);
+			const UPoseSearchDatabase* Database = AssetsToSearchPerDatabasePair.Key;
 
 #if WITH_EDITOR
 			if (!FAsyncPoseSearchDatabasesManagement::RequestAsyncBuildIndex(Database, ERequestAsyncBuildFlag::ContinueRequest))
@@ -761,6 +777,8 @@ UE::PoseSearch::FSearchResult UPoseSearchLibrary::MotionMatch(const FAnimationBa
 			else
 #endif // WITH_EDITOR
 			{
+				check(Database);
+
 				const FSearchIndex& SearchIndex = Database->GetSearchIndex();
 				for (int32 AssetIndex = 0; AssetIndex < SearchIndex.Assets.Num(); ++AssetIndex)
 				{
@@ -771,11 +789,11 @@ UE::PoseSearch::FSearchResult UPoseSearchLibrary::MotionMatch(const FAnimationBa
 						{
 							const float FirstSampleTime = SearchIndexAsset.GetFirstSampleTime(Database->Schema->SampleRate);
 							const float LastSampleTime = SearchIndexAsset.GetLastSampleTime(Database->Schema->SampleRate);
-							if (PlayingAnimationAssetAccumulatedTime >= FirstSampleTime && PlayingAnimationAssetAccumulatedTime <= LastSampleTime)
+							if (PlayingAssetAccumulatedTime >= FirstSampleTime && PlayingAssetAccumulatedTime <= LastSampleTime)
 							{
 								ReconstructedPreviousSearchResult.Database = Database;
-								ReconstructedPreviousSearchResult.AssetTime = PlayingAnimationAssetAccumulatedTime;
-								ReconstructedPreviousSearchResult.PoseIdx = Database->GetPoseIndexFromTime(PlayingAnimationAssetAccumulatedTime, SearchIndexAsset);
+								ReconstructedPreviousSearchResult.AssetTime = PlayingAssetAccumulatedTime;
+								ReconstructedPreviousSearchResult.PoseIdx = Database->GetPoseIndexFromTime(PlayingAssetAccumulatedTime, SearchIndexAsset);
 								SearchContext.UpdateCurrentResultPoseVector();
 
 								SearchResult = Database->SearchContinuingPose(SearchContext);
@@ -787,23 +805,23 @@ UE::PoseSearch::FSearchResult UPoseSearchLibrary::MotionMatch(const FAnimationBa
 			}
 		}
 
-		PerDbAnimMap.Reset();
+		AssetsToSearchPerDatabaseMap.Reset();
 	}
 
 	// collecting all the other databases searches
-	if (!AnimationAssets.IsEmpty())
+	if (!AssetsToSearch.IsEmpty())
 	{
-		for (const UAnimationAsset* AnimationAsset : AnimationAssets)
+		for (const UObject* AssetToSearch : AssetsToSearch)
 		{
-			AddToPerDbAnimMap(PerDbAnimMap, AnimationAsset);
+			AddToSearch(AssetsToSearchPerDatabaseMap, AssetToSearch);
 		}
 
-		for (const FPerDbAnimPair& PerDbAnimPair : PerDbAnimMap)
+		for (const TAssetsToSearchPerDatabasePair& AssetsToSearchPerDatabasePair : AssetsToSearchPerDatabaseMap)
 		{
-			const UPoseSearchDatabase* Database = PerDbAnimPair.Key;
+			const UPoseSearchDatabase* Database = AssetsToSearchPerDatabasePair.Key;
 			check(Database);
 
-			SearchContext.SetAnimationsToConsider(PerDbAnimPair.Value);
+			SearchContext.SetAssetsToConsider(AssetsToSearchPerDatabasePair.Value);
 
 			const FSearchResult NewSearchResult = Database->Search(SearchContext);
 			if (NewSearchResult.PoseCost.GetTotalCost() < SearchResult.PoseCost.GetTotalCost())
diff --git a/Engine/Plugins/Experimental/Animation/PoseSearch/Source/Runtime/Private/PoseSearchModule.cpp b/Engine/Plugins/Experimental/Animation/PoseSearch/Source/Runtime/Private/PoseSearchModule.cpp
--- a/Engine/Plugins/Experimental/Animation/PoseSearch/Source/Runtime/Private/PoseSearchModule.cpp
+++ b/Engine/Plugins/Experimental/Animation/PoseSearch/Source/Runtime/Private/PoseSearchModule.cpp
@@ -22,16 +22,15 @@ public:
 	}
 
 	// IPoseSearchProvider
-	virtual UE::Anim::IPoseSearchProvider::FSearchResult Search(const FAnimationBaseContext& GraphContext, TConstArrayView<UAnimationAsset*> AnimationAssets,
-		const UAnimationAsset* PlayingAnimationAsset, float PlayingAnimationAssetAccumulatedTime) override
+	virtual UE::Anim::IPoseSearchProvider::FSearchResult Search(const FAnimationBaseContext& GraphContext, TConstArrayView<UObject*> AssetsToSearch, const UObject* PlayingAsset, float PlayingAssetAccumulatedTime) const override
 	{
-		const UE::PoseSearch::FSearchResult SearchResult = UPoseSearchLibrary::MotionMatch(GraphContext, AnimationAssets, PlayingAnimationAsset, PlayingAnimationAssetAccumulatedTime);
+		const UE::PoseSearch::FSearchResult SearchResult = UPoseSearchLibrary::MotionMatch(GraphContext, AssetsToSearch, PlayingAsset, PlayingAssetAccumulatedTime);
 		UE::Anim::IPoseSearchProvider::FSearchResult ProviderResult;
 		if (const UE::PoseSearch::FSearchIndexAsset* SearchIndexAsset = SearchResult.GetSearchIndexAsset())
 		{
 			if (const FPoseSearchDatabaseAnimationAssetBase* DatabaseAnimationAssetBase = SearchResult.Database->GetAnimationAssetBase(*SearchIndexAsset))
 			{
-				ProviderResult.AnimationAsset = DatabaseAnimationAssetBase->GetAnimationAsset();
+				ProviderResult.SelectedAsset = DatabaseAnimationAssetBase->GetAnimationAsset();
 				ProviderResult.Dissimilarity = SearchResult.PoseCost.GetTotalCost();
 				ProviderResult.TimeOffsetSeconds = SearchResult.AssetTime;
 			}
diff --git a/Engine/Plugins/Experimental/Animation/PoseSearch/Source/Runtime/Private/PoseSearchResult.cpp b/Engine/Plugins/Experimental/Animation/PoseSearch/Source/Runtime/Private/PoseSearchResult.cpp
--- a/Engine/Plugins/Experimental/Animation/PoseSearch/Source/Runtime/Private/PoseSearchResult.cpp
+++ b/Engine/Plugins/Experimental/Animation/PoseSearch/Source/Runtime/Private/PoseSearchResult.cpp
@@ -95,7 +95,7 @@ bool FSearchResult::CanAdvance(float DeltaTime) const
 		}
 		else if (const FPoseSearchDatabaseAnimationAssetBase* DatabaseAnimationAssetBase = DatabaseAsset.GetPtr<FPoseSearchDatabaseAnimationAssetBase>())
 		{
-			const float AssetLength = DatabaseAnimationAssetBase->GetAnimationAsset()->GetPlayLength();
+			const float AssetLength = DatabaseAnimationAssetBase->GetPlayLength();
 			bCanAdvance = ETAA_Finished != FAnimationRuntime::AdvanceTime(SearchIndexAsset->IsLooping(), DeltaTime, SteppedTime, AssetLength);
 		}
 	}
diff --git a/Engine/Plugins/Experimental/Animation/PoseSearch/Source/Runtime/Public/PoseSearch/PoseSearchContext.h b/Engine/Plugins/Experimental/Animation/PoseSearch/Source/Runtime/Public/PoseSearch/PoseSearchContext.h
--- a/Engine/Plugins/Experimental/Animation/PoseSearch/Source/Runtime/Public/PoseSearch/PoseSearchContext.h
+++ b/Engine/Plugins/Experimental/Animation/PoseSearch/Source/Runtime/Public/PoseSearch/PoseSearchContext.h
@@ -169,7 +169,7 @@ struct FCachedChannel
 
 struct POSESEARCH_API FSearchContext
 {
-	FSearchContext(const UAnimInstance* InAnimInstance, const IPoseHistory* InHistory, TConstArrayView<const UAnimationAsset*> InAnimationsToConsider = TConstArrayView<const UAnimationAsset*>(),
+	FSearchContext(const UAnimInstance* InAnimInstance, const IPoseHistory* InHistory, TConstArrayView<const UObject*> InAssetsToConsider = TConstArrayView<const UObject*>(),
 		float InDesiredPermutationTimeOffset = 0.f, const FPoseIndicesHistory* InPoseIndicesHistory = nullptr,
 		const FSearchResult& InCurrentResult = FSearchResult(), const FFloatInterval& InPoseJumpThresholdTime = FFloatInterval(0.f, 0.f), bool bInUseCachedChannelData = false);
 
@@ -209,8 +209,8 @@ struct POSESEARCH_API FSearchContext
 	float GetDesiredPermutationTimeOffset() const { return DesiredPermutationTimeOffset; }
 	const UAnimInstance* GetAnimInstance() const { return AnimInstance; }
 
-	void SetAnimationsToConsider(TConstArrayView<const UAnimationAsset*> InAnimationsToConsider) { AnimationsToConsider = InAnimationsToConsider; }
-	TConstArrayView<const UAnimationAsset*> GetAnimationsToConsider() const { return AnimationsToConsider; }
+	void SetAssetsToConsider(TConstArrayView<const UObject*> InAssetsToConsider) { AssetsToConsider = InAssetsToConsider; }
+	TConstArrayView<const UObject*> GetAssetsToConsider() const { return AssetsToConsider; }
 	
 	// returns the world space transform of the bone SchemaBoneIdx at time SampleTime
 	FTransform GetWorldBoneTransformAtTime(float SampleTime, int8 SchemaBoneIdx = RootSchemaBoneIdx);
@@ -237,8 +237,8 @@ private:
 	const UAnimInstance* AnimInstance = nullptr;
 	const IPoseHistory* History = nullptr;
 
-	// if AnimationsToConsider is not empty, we'll search only for poses from UAnimationAsset(s) that are in the AnimationsToConsider
-	TConstArrayView<const UAnimationAsset*> AnimationsToConsider;
+	// if AssetsToConsider is not empty, we'll search only for poses from UObject(s) that are in the AssetsToConsider
+	TConstArrayView<const UObject*> AssetsToConsider;
 
 	const float DesiredPermutationTimeOffset = 0.f;
 	const FPoseIndicesHistory* PoseIndicesHistory = nullptr;
diff --git a/Engine/Plugins/Experimental/Animation/PoseSearch/Source/Runtime/Public/PoseSearch/PoseSearchDatabase.h b/Engine/Plugins/Experimental/Animation/PoseSearch/Source/Runtime/Public/PoseSearch/PoseSearchDatabase.h
--- a/Engine/Plugins/Experimental/Animation/PoseSearch/Source/Runtime/Public/PoseSearch/PoseSearchDatabase.h
+++ b/Engine/Plugins/Experimental/Animation/PoseSearch/Source/Runtime/Public/PoseSearch/PoseSearchDatabase.h
@@ -50,7 +50,8 @@ struct POSESEARCH_API FPoseSearchDatabaseAnimationAssetBase
 {
 	GENERATED_BODY()
 	virtual ~FPoseSearchDatabaseAnimationAssetBase() = default;
-	virtual UAnimationAsset* GetAnimationAsset() const { return nullptr; }
+	virtual UObject* GetAnimationAsset() const { return nullptr; }
+	virtual float GetPlayLength() const;
 
 #if WITH_EDITORONLY_DATA
 	virtual bool IsDisableReselection() const { return bDisableReselection; }
@@ -114,7 +115,7 @@ struct POSESEARCH_API FPoseSearchDatabaseSequence : public FPoseSearchDatabaseAn
 	FFloatInterval GetSamplingRange() const override { return SamplingRange; }
 #endif // WITH_EDITORONLY_DATA
 	
-	UAnimationAsset* GetAnimationAsset() const override;
+	UObject* GetAnimationAsset() const override;
 };
 
 /** An blend space entry in a UPoseSearchDatabase. */
@@ -162,7 +163,7 @@ struct POSESEARCH_API FPoseSearchDatabaseBlendSpace : public FPoseSearchDatabase
 	FVector BlendParameterForSampleRanges(int32 HorizontalBlendIndex, int32 VerticalBlendIndex) const;
 #endif // WITH_EDITORONLY_DATA
 
-	UAnimationAsset* GetAnimationAsset() const override;
+	UObject* GetAnimationAsset() const override;
 };
 
 /** An entry in a UPoseSearchDatabase. */
@@ -189,7 +190,7 @@ struct POSESEARCH_API FPoseSearchDatabaseAnimComposite : public FPoseSearchDatab
 	FFloatInterval GetSamplingRange() const override { return SamplingRange; }
 #endif // WITH_EDITORONLY_DATA
 
-	UAnimationAsset* GetAnimationAsset() const override;
+	UObject* GetAnimationAsset() const override;
 };
 
 /** An anim montage entry in a UPoseSearchDatabase. */
@@ -216,7 +217,7 @@ struct POSESEARCH_API FPoseSearchDatabaseAnimMontage : public FPoseSearchDatabas
 	FFloatInterval GetSamplingRange() const override { return SamplingRange; }
 #endif // WITH_EDITORONLY_DATA
 
-	UAnimationAsset* GetAnimationAsset() const override;
+	UObject* GetAnimationAsset() const override;
 };
 
 /** A data asset for indexing a collection of animation sequences. */
diff --git a/Engine/Plugins/Experimental/Animation/PoseSearch/Source/Runtime/Public/PoseSearch/PoseSearchLibrary.h b/Engine/Plugins/Experimental/Animation/PoseSearch/Source/Runtime/Public/PoseSearch/PoseSearchLibrary.h
--- a/Engine/Plugins/Experimental/Animation/PoseSearch/Source/Runtime/Public/PoseSearch/PoseSearchLibrary.h
+++ b/Engine/Plugins/Experimental/Animation/PoseSearch/Source/Runtime/Public/PoseSearch/PoseSearchLibrary.h
@@ -171,7 +171,7 @@ public:
 		float TimeToFutureAnimationStart = 0.f,
 		const int32 DebugSessionUniqueIdentifier = 6174);
 
-	static UE::PoseSearch::FSearchResult MotionMatch(const FAnimationBaseContext& Context, TConstArrayView<UAnimationAsset*> AnimationAssets,
-		const UAnimationAsset* PlayingAnimationAsset = nullptr, float PlayingAnimationAssetAccumulatedTime = 0.f);
+	static UE::PoseSearch::FSearchResult MotionMatch(const FAnimationBaseContext& Context, TConstArrayView<UObject*> AssetsToSearch,
+		const UObject* PlayingAsset = nullptr, float PlayingAssetAccumulatedTime = 0.f);
 };
 
diff --git a/Engine/Plugins/Experimental/Animation/PoseSearch/Source/Runtime/Public/PoseSearch/PoseSearchResult.h b/Engine/Plugins/Experimental/Animation/PoseSearch/Source/Runtime/Public/PoseSearch/PoseSearchResult.h
--- a/Engine/Plugins/Experimental/Animation/PoseSearch/Source/Runtime/Public/PoseSearch/PoseSearchResult.h
+++ b/Engine/Plugins/Experimental/Animation/PoseSearch/Source/Runtime/Public/PoseSearch/PoseSearchResult.h
@@ -51,7 +51,7 @@ struct POSESEARCH_API FPoseSearchBlueprintResult
 public:
 
 	UPROPERTY(Transient, VisibleAnywhere, BlueprintReadOnly, Category=State)
-	TObjectPtr<const UAnimationAsset> SelectedAnimation = nullptr;
+	TObjectPtr<const UObject> SelectedAnimation = nullptr;
 	
 	UPROPERTY(Transient, VisibleAnywhere, BlueprintReadOnly, Category=State)
 	float SelectedTime = 0.f;
diff --git a/Engine/Plugins/Experimental/Chooser/Source/Chooser/Private/AnimNode_ChooserPlayer.cpp b/Engine/Plugins/Experimental/Chooser/Source/Chooser/Private/AnimNode_ChooserPlayer.cpp
--- a/Engine/Plugins/Experimental/Chooser/Source/Chooser/Private/AnimNode_ChooserPlayer.cpp
+++ b/Engine/Plugins/Experimental/Chooser/Source/Chooser/Private/AnimNode_ChooserPlayer.cpp
@@ -35,21 +35,18 @@ UAnimationAsset* FAnimNode_ChooserPlayer::ChooseAsset(const FAnimationUpdateCont
 		{
 			if (UE::Anim::IPoseSearchProvider* PoseSearchProvider = UE::Anim::IPoseSearchProvider::Get())
 			{
-				TArray<UAnimationAsset*, TInlineAllocator<128>> AnimationAssets;
-				ChooserBase.ChooseMulti(ChooserContext, FObjectChooserBase::FObjectChooserIteratorCallback::CreateLambda([&AnimationAssets](UObject* InResult)
+				TArray<UObject*, TInlineAllocator<128>> ChosenAssets;
+				ChooserBase.ChooseMulti(ChooserContext, FObjectChooserBase::FObjectChooserIteratorCallback::CreateLambda([&ChosenAssets](UObject* InResult)
 					{
-						if (UAnimationAsset* AnimationAsset = Cast<UAnimationAsset>(InResult))
-						{
-							AnimationAssets.Add(AnimationAsset);
-						}
+						ChosenAssets.Add(InResult);
 						return FObjectChooserBase::EIteratorStatus::Continue;
 					}));
 
-				const UE::Anim::IPoseSearchProvider::FSearchResult SearchResult = PoseSearchProvider->Search(Context, AnimationAssets, GetAnimAsset(), GetAccumulatedTime());
-				if (SearchResult.AnimationAsset)
+				const UE::Anim::IPoseSearchProvider::FSearchResult SearchResult = PoseSearchProvider->Search(Context, ChosenAssets, GetAnimAsset(), GetAccumulatedTime());
+				if (UAnimationAsset* SelectedAnimationAsset = Cast<UAnimationAsset>(SearchResult.SelectedAsset))
 				{
 					Settings.StartTime = SearchResult.TimeOffsetSeconds;
-					return SearchResult.AnimationAsset;
+					return SelectedAnimationAsset;
 				}
 			}
 		}
diff --git a/Engine/Source/Runtime/Engine/Private/Animation/AnimNode_SequencePlayer.cpp b/Engine/Source/Runtime/Engine/Private/Animation/AnimNode_SequencePlayer.cpp
--- a/Engine/Source/Runtime/Engine/Private/Animation/AnimNode_SequencePlayer.cpp
+++ b/Engine/Source/Runtime/Engine/Private/Animation/AnimNode_SequencePlayer.cpp
@@ -159,14 +159,13 @@ float FAnimNode_SequencePlayerBase::GetTimeFromEnd(float CurrentNodeTime) const
 float FAnimNode_SequencePlayerBase::GetEffectiveStartPosition(const FAnimationBaseContext& Context) const
 {
 	// Override the start position if pose matching is enabled
-	UAnimationAsset* CurrentSequence = GetSequence();
+	UObject* CurrentSequence = GetSequence();
 	if (CurrentSequence != nullptr && GetStartFromMatchingPose())
 	{
-		UE::Anim::IPoseSearchProvider* PoseSearchProvider = UE::Anim::IPoseSearchProvider::Get();
-		if (PoseSearchProvider)
+		if (const UE::Anim::IPoseSearchProvider* PoseSearchProvider = UE::Anim::IPoseSearchProvider::Get())
 		{
 			UE::Anim::IPoseSearchProvider::FSearchResult Result = PoseSearchProvider->Search(Context, MakeArrayView(&CurrentSequence, 1));
-			if (Result.AnimationAsset != nullptr)
+			if (Result.SelectedAsset != nullptr)
 			{
 				return Result.TimeOffsetSeconds;
 			}
diff --git a/Engine/Source/Runtime/Engine/Public/Animation/AnimPoseSearchProvider.h b/Engine/Source/Runtime/Engine/Public/Animation/AnimPoseSearchProvider.h
--- a/Engine/Source/Runtime/Engine/Public/Animation/AnimPoseSearchProvider.h
+++ b/Engine/Source/Runtime/Engine/Public/Animation/AnimPoseSearchProvider.h
@@ -24,7 +24,7 @@ public:
 	
 	struct FSearchResult
 	{
-		UAnimationAsset* AnimationAsset = nullptr;
+		UObject* SelectedAsset = nullptr;
 		float TimeOffsetSeconds = 0.f;
 		float Dissimilarity = MAX_flt;
 	};
@@ -32,15 +32,15 @@ public:
 	/**
 	* Finds a matching pose in the input Object given the current graph context
 	* 
-	* @param	GraphContext	Graph execution context used to construct a pose search query
-	* @param	AnimationAssets	The animation assets to search for the pose query
-	* @param	PlayingAnimationAsset					The currently playing animation asset, used to bias the score of the eventually found continuing pose
-	* @param	PlayingAnimationAssetAccumulatedTime	The accumulated time of the currently playing animation asset
+	* @param	GraphContext				Graph execution context used to construct a pose search query
+	* @param	AssetsToSearch				The assets to search for the pose query
+	* @param	PlayingAsset				The currently playing asset, used to bias the score of the eventually found continuing pose
+	* @param	PlayingAssetAccumulatedTime	The accumulated time of the currently playing asset
 	* 
-	* @return	The pose in the AnimationAssets that most closely matches the query
+	* @return	FSearchResult				The search result identifying the asset from AssetsToSearch or PlayingAsset that most closely matches the query
 	*/
-	virtual FSearchResult Search(const FAnimationBaseContext& GraphContext, TConstArrayView<UAnimationAsset*> AnimationAssets, 
-		const UAnimationAsset* PlayingAnimationAsset = nullptr, float PlayingAnimationAssetAccumulatedTime = 0.f) = 0;
+	virtual FSearchResult Search(const FAnimationBaseContext& GraphContext, TConstArrayView<UObject*> AssetsToSearch,
+		const UObject* PlayingAsset = nullptr, float PlayingAssetAccumulatedTime = 0.f) const = 0;
 };
 
 } // namespace UE::Anim
\ No newline at end of file
