motion matching - fix for pose search database editor crash while previewing mirrored assets

diff --git a/Engine/Plugins/Experimental/Animation/PoseSearch/Source/Editor/Private/PoseSearchDatabaseViewModel.cpp b/Engine/Plugins/Experimental/Animation/PoseSearch/Source/Editor/Private/PoseSearchDatabaseViewModel.cpp
--- a/Engine/Plugins/Experimental/Animation/PoseSearch/Source/Editor/Private/PoseSearchDatabaseViewModel.cpp
+++ b/Engine/Plugins/Experimental/Animation/PoseSearch/Source/Editor/Private/PoseSearchDatabaseViewModel.cpp
@@ -66,12 +66,6 @@ bool FDatabasePreviewActor::SpawnPreviewActor(UWorld* World, const UPoseSearchDa
 		PlayTimeOffset = PoseSearchDatabase->GetRealAssetTime(PoseIdxForTimeOffset) - IndexAsset.GetFirstSampleTime(PoseSearchDatabase->Schema->SampleRate);
 	}
 
-	// @todo: should we always use the PlayTimeOffset to extract the root transform?
-	if (PlayTimeOffset != 0.f)
-	{
-		RootTransformOrigin = ExtractRootTransform(PlayTimeOffset);
-	}
-
 	FActorSpawnParameters Params;
 	Params.SpawnCollisionHandlingOverride = ESpawnActorCollisionHandlingMethod::AdjustIfPossibleButAlwaysSpawn;
 	ActorPtr = World->SpawnActor<AActor>(AActor::StaticClass(), FTransform::Identity, Params);
@@ -81,7 +75,6 @@ bool FDatabasePreviewActor::SpawnPreviewActor(UWorld* World, const UPoseSearchDa
 	Mesh->RegisterComponentWithWorld(World);
 
 	UAnimPreviewInstance* AnimInstance = NewObject<UAnimPreviewInstance>(Mesh);
-
 	Mesh->PreviewInstance = AnimInstance;
 	AnimInstance->InitializeAnimation();
 
@@ -95,7 +88,14 @@ bool FDatabasePreviewActor::SpawnPreviewActor(UWorld* World, const UPoseSearchDa
 	if (IndexAsset.IsMirrored() && PoseSearchDatabase->Schema)
 	{
 		AnimInstance->SetMirrorDataTable(PoseSearchDatabase->Schema->MirrorDataTable);
-		MirrorDataCache.Init(PoseSearchDatabase->Schema->MirrorDataTable, AnimInstance->GetRequiredBonesOnAnyThread());
+	}
+
+	const FMirrorDataCache MirrorDataCache(AnimInstance->GetMirrorDataTable(), AnimInstance->GetRequiredBonesOnAnyThread());
+	
+	// @todo: should we always use the PlayTimeOffset to extract the root transform?
+	if (PlayTimeOffset != 0.f)
+	{
+		RootTransformOrigin = MirrorDataCache.MirrorTransform(Sampler.ExtractRootTransform(PlayTimeOffset));
 	}
 
 	AnimInstance->PlayAnim(false, 0.0f);
@@ -148,8 +148,10 @@ void FDatabasePreviewActor::UpdatePreviewActor(const UPoseSearchDatabase* PoseSe
 	AnimInstance->SetPlayRate(0.f);
 	AnimInstance->SetBlendSpacePosition(IndexAsset.GetBlendParameters());
 
+	const FMirrorDataCache MirrorDataCache(AnimInstance->GetMirrorDataTable(), AnimInstance->GetRequiredBonesOnAnyThread());
+
 	// updating root transforms
-	RootTransformCurrentQuantizedTime = ExtractRootTransform(QuantizedTime);
+	RootTransformCurrentQuantizedTime = MirrorDataCache.MirrorTransform(Sampler.ExtractRootTransform(QuantizedTime));
 	RootTransformCurrentQuantizedTime.SetToRelativeTransform(RootTransformOrigin);
 
 	if (CurrentTime == QuantizedTime)
@@ -158,7 +160,7 @@ void FDatabasePreviewActor::UpdatePreviewActor(const UPoseSearchDatabase* PoseSe
 	}
 	else
 	{
-		RootTransformCurrent = ExtractRootTransform(CurrentTime);
+		RootTransformCurrent = MirrorDataCache.MirrorTransform(Sampler.ExtractRootTransform(CurrentTime));
 		RootTransformCurrent.SetToRelativeTransform(RootTransformOrigin);
 	}
 
@@ -172,7 +174,10 @@ void FDatabasePreviewActor::UpdatePreviewActor(const UPoseSearchDatabase* PoseSe
 		FMemMark Mark(FMemStack::Get());
 		FCompactPose Pose;
 		Pose.SetBoneContainer(&BoneContainer);
-		ExtractPose(QuantizedTime, Pose);
+
+		Sampler.ExtractPose(QuantizedTime, Pose);
+		MirrorDataCache.MirrorPose(Pose);
+
 		RootBoneTransformCurrentQuantizedTime = Pose[FCompactPoseBoneIndex(RootBoneIndexType)];
 	}
 }
@@ -207,14 +212,15 @@ bool FDatabasePreviewActor::DrawPreviewActor(const UPoseSearchDatabase* PoseSear
 		DrawParams.DrawFeatureVector(QueryVector);
 	}
 
+	const FSearchIndex& SearchIndex = PoseSearchDatabase->GetSearchIndex();
+	const FSearchIndexAsset& IndexAsset = SearchIndex.Assets[IndexAssetIndex];
+
+	const FMirrorDataCache MirrorDataCache(Mesh->PreviewInstance->GetMirrorDataTable(), Mesh->PreviewInstance->GetRequiredBonesOnAnyThread());
 	if (bDisplayRootMotionSpeed || bDisplayBlockTransition)
 	{
 		// initializing SampledRootMotion if required
 		if (SampledRootMotion.IsEmpty())
 		{
-			const FSearchIndex& SearchIndex = PoseSearchDatabase->GetSearchIndex();
-			const FSearchIndexAsset& IndexAsset = SearchIndex.Assets[IndexAssetIndex];
-
 			const int NumPoses = IndexAsset.GetNumPoses();
 			if (NumPoses > 1)
 			{
@@ -225,7 +231,7 @@ bool FDatabasePreviewActor::DrawPreviewActor(const UPoseSearchDatabase* PoseSear
 				{
 					const int32 IndexAssetPoseIdx = Index + IndexAsset.GetFirstPoseIdx();
 					const float IndexAssetPoseTime = IndexAsset.GetTimeFromPoseIndex(IndexAssetPoseIdx, PoseSearchDatabase->Schema->SampleRate);
-					FTransform IndexAssetPoseTransform = ExtractRootTransform(IndexAssetPoseTime);
+					FTransform IndexAssetPoseTransform = MirrorDataCache.MirrorTransform(Sampler.ExtractRootTransform(IndexAssetPoseTime));
 					IndexAssetPoseTransform.SetToRelativeTransform(RootTransformOrigin);
 
 					SampledRootMotion[Index] = IndexAssetPoseTransform.GetTranslation();
@@ -268,9 +274,6 @@ bool FDatabasePreviewActor::DrawPreviewActor(const UPoseSearchDatabase* PoseSear
 
 	if (bDisplayBlockTransition)
 	{
-		const FSearchIndex& SearchIndex = PoseSearchDatabase->GetSearchIndex();
-		const FSearchIndexAsset& IndexAsset = SearchIndex.Assets[IndexAssetIndex];
-
 		const int NumPoses = IndexAsset.GetNumPoses();
 		if (NumPoses == SampledRootMotion.Num())
 		{
@@ -298,9 +301,11 @@ bool FDatabasePreviewActor::DrawPreviewActor(const UPoseSearchDatabase* PoseSear
 		FMemMark Mark(FMemStack::Get());
 		FCompactPose Pose;
 		Pose.SetBoneContainer(&GetAnimPreviewInstance()->GetRequiredBonesOnAnyThread());
-		ExtractPose(CurrentTime, Pose);
+		
+		Sampler.ExtractPose(CurrentTime, Pose);
+		MirrorDataCache.MirrorPose(Pose);
 
-		const FTransform RootTransform = ExtractRootTransform(CurrentTime);
+		const FTransform RootTransform = MirrorDataCache.MirrorTransform(Sampler.ExtractRootTransform(CurrentTime));
 
 		FCSPose<FCompactPose> ComponentSpacePose;
 		ComponentSpacePose.InitPose(MoveTemp(Pose));
@@ -316,17 +321,6 @@ bool FDatabasePreviewActor::DrawPreviewActor(const UPoseSearchDatabase* PoseSear
 	return true;
 }
 
-FTransform FDatabasePreviewActor::ExtractRootTransform(float Time) const
-{
-	return MirrorDataCache.MirrorTransform(Sampler.ExtractRootTransform(Time));
-}
-
-void FDatabasePreviewActor::ExtractPose(float Time, FCompactPose& OutPose) const
-{
-	Sampler.ExtractPose(Time, OutPose);
-	MirrorDataCache.MirrorPose(OutPose);
-}
-
 const UDebugSkelMeshComponent* FDatabasePreviewActor::GetDebugSkelMeshComponent() const
 {
 	if (ActorPtr != nullptr)
diff --git a/Engine/Plugins/Experimental/Animation/PoseSearch/Source/Editor/Public/PoseSearchDatabaseViewModel.h b/Engine/Plugins/Experimental/Animation/PoseSearch/Source/Editor/Public/PoseSearchDatabaseViewModel.h
--- a/Engine/Plugins/Experimental/Animation/PoseSearch/Source/Editor/Public/PoseSearchDatabaseViewModel.h
+++ b/Engine/Plugins/Experimental/Animation/PoseSearch/Source/Editor/Public/PoseSearchDatabaseViewModel.h
@@ -36,8 +36,6 @@ namespace UE::PoseSearch
 		const UDebugSkelMeshComponent* GetDebugSkelMeshComponent() const;
 		const UAnimPreviewInstance* GetAnimPreviewInstance() const;
 		const FAnimationAssetSampler& GetSampler() const { return Sampler; }
-		FTransform ExtractRootTransform(float Time) const;
-		void ExtractPose(float Time, FCompactPose& OutPose) const;
 
 		const AActor* GetActor() const { return ActorPtr.Get(); }
 		int32 GetIndexAssetIndex() const { return IndexAssetIndex; }
@@ -62,7 +60,6 @@ namespace UE::PoseSearch
 		FTransform RootBoneTransformCurrentQuantizedTime = FTransform::Identity;
 
 		FAnimationAssetSampler Sampler;
-		FMirrorDataCache MirrorDataCache;
 
 		TArray<FVector> SampledRootMotion;
 		TArray<float> SampledRootMotionSpeed;
