motion matching - multicharacter WIP

diff --git a/Engine/Plugins/Experimental/Animation/PoseSearch/Source/Editor/Private/PoseSearchDatabaseAssetTree.cpp b/Engine/Plugins/Experimental/Animation/PoseSearch/Source/Editor/Private/PoseSearchDatabaseAssetTree.cpp
--- a/Engine/Plugins/Experimental/Animation/PoseSearch/Source/Editor/Private/PoseSearchDatabaseAssetTree.cpp
+++ b/Engine/Plugins/Experimental/Animation/PoseSearch/Source/Editor/Private/PoseSearchDatabaseAssetTree.cpp
@@ -246,8 +246,8 @@ namespace UE::PoseSearch
 						
 			// Build an index based off of alphabetical order than iterate the index instead
 			TArray<uint32> IndexArray;
-			IndexArray.SetNumUninitialized(Database->AnimationAssets.Num());
-			for (int32 AnimationAssetIdx = 0; AnimationAssetIdx < Database->AnimationAssets.Num(); ++AnimationAssetIdx)
+			IndexArray.SetNumUninitialized(Database->GetAnimationAssets().Num());
+			for (int32 AnimationAssetIdx = 0; AnimationAssetIdx < Database->GetAnimationAssets().Num(); ++AnimationAssetIdx)
 			{
 				IndexArray[AnimationAssetIdx] = AnimationAssetIdx;
 			}
@@ -273,7 +273,7 @@ namespace UE::PoseSearch
 			});
 
 			// create all nodes
-			for (int32 AnimationAssetIdx = 0; AnimationAssetIdx < Database->AnimationAssets.Num(); ++AnimationAssetIdx)
+			for (int32 AnimationAssetIdx = 0; AnimationAssetIdx < Database->GetAnimationAssets().Num(); ++AnimationAssetIdx)
 			{
 				const int32 MappedId = IndexArray[AnimationAssetIdx];
 
@@ -486,6 +486,14 @@ namespace UE::PoseSearch
 				FUIAction(FExecuteAction::CreateSP(this, &SDatabaseAssetTree::OnAddAnimMontage, true)),
 				NAME_None,
 				EUserInterfaceActionType::Button);
+
+			AddOptions.AddMenuEntry(
+				LOCTEXT("AddMultiSequenceOption", "Multi Sequence"),
+				LOCTEXT("AddMultiSequenceToDatabaseTooltip", "Add new multi sequence to the database"),
+				FSlateIcon(),
+				FUIAction(FExecuteAction::CreateSP(this, &SDatabaseAssetTree::OnAddMultiSequence, true)),
+				NAME_None,
+				EUserInterfaceActionType::Button);
 		}
 		AddOptions.EndSection();
 
@@ -622,6 +630,18 @@ namespace UE::PoseSearch
 		}
 	}
 
+	void SDatabaseAssetTree::OnAddMultiSequence(bool bFinalizeChanges)
+	{
+		FScopedTransaction Transaction(LOCTEXT("AddMultiSequenceTransaction", "Add Multi Sequence"));
+
+		EditorViewModel.Pin()->AddMultiSequenceToDatabase();
+
+		if (bFinalizeChanges)
+		{
+			FinalizeTreeChanges();
+		}
+	}
+
 	void SDatabaseAssetTree::OnDeleteAsset(TSharedPtr<FDatabaseAssetTreeNode> Node, bool bFinalizeChanges)
 	{
 		FScopedTransaction Transaction(LOCTEXT("DeleteAsset", "Delete Asset"));
diff --git a/Engine/Plugins/Experimental/Animation/PoseSearch/Source/Editor/Private/PoseSearchDatabaseDataDetails.cpp b/Engine/Plugins/Experimental/Animation/PoseSearch/Source/Editor/Private/PoseSearchDatabaseDataDetails.cpp
--- a/Engine/Plugins/Experimental/Animation/PoseSearch/Source/Editor/Private/PoseSearchDatabaseDataDetails.cpp
+++ b/Engine/Plugins/Experimental/Animation/PoseSearch/Source/Editor/Private/PoseSearchDatabaseDataDetails.cpp
@@ -239,7 +239,6 @@ public:
 					.Margin(FMargin(1.f, 1.f, 1.f, 1.f))
 					.Text_Lambda([this, ColumnName]() -> FText
 						{
-							FString PoseIdxString;
 							if (!ChannelItem->IsExpanded())
 							{
 								if (const TSharedPtr<FDatabaseViewModel> ViewModel = EditorViewModel.Pin())
@@ -247,17 +246,90 @@ public:
 									const UPoseSearchDatabase* PoseSearchDatabase = ViewModel->GetPoseSearchDatabase();
 									if (FAsyncPoseSearchDatabasesManagement::RequestAsyncBuildIndex(PoseSearchDatabase, ERequestAsyncBuildFlag::ContinueRequest))
 									{
-										if (const FDatabasePreviewActor* FoundPreviewActor = ViewModel->GetPreviewActors().FindByPredicate(
-											[ColumnName](const FDatabasePreviewActor& PreviewActor)
-											{ return *GetNameSafe(PreviewActor.GetActor()) == ColumnName; }))
+										for (const TArray<FDatabasePreviewActor>& PreviewActorGroup : ViewModel->GetPreviewActors())
 										{
-											const int32 PoseIdx = FoundPreviewActor->GetCurrentPoseIndex();
-											PoseIdxString.AppendInt(PoseIdx);
+											for (const FDatabasePreviewActor& PreviewActor : PreviewActorGroup)
+											{
+												if (*GetNameSafe(PreviewActor.GetActor()) == ColumnName)
+												{
+													const int32 PoseIdx = PreviewActor.GetCurrentPoseIndex();
+													FString PoseIdxString;
+													PoseIdxString.AppendInt(PoseIdx);
+													return FText::FromString(PoseIdxString);
+												}
+											}
+										}
+									}
+								}
+							}
+							return FText::GetEmpty();
+						});
+			}
+
+			if (ChannelItem->GetCompactLabel() == "SearchAssetIndex")
+			{
+				return SNew(STextBlock)
+					.Margin(FMargin(1.f, 1.f, 1.f, 1.f))
+					.Text_Lambda([this, ColumnName]() -> FText
+						{
+							if (!ChannelItem->IsExpanded())
+							{
+								if (const TSharedPtr<FDatabaseViewModel> ViewModel = EditorViewModel.Pin())
+								{
+									const UPoseSearchDatabase* PoseSearchDatabase = ViewModel->GetPoseSearchDatabase();
+									if (FAsyncPoseSearchDatabasesManagement::RequestAsyncBuildIndex(PoseSearchDatabase, ERequestAsyncBuildFlag::ContinueRequest))
+									{
+										for (const TArray<FDatabasePreviewActor>& PreviewActorGroup : ViewModel->GetPreviewActors())
+										{
+											for (const FDatabasePreviewActor& PreviewActor : PreviewActorGroup)
+											{
+												if (*GetNameSafe(PreviewActor.GetActor()) == ColumnName)
+												{
+													const int32 AssetIdx = PreviewActor.GetIndexAssetIndex();
+													FString AssetIdxString;
+													AssetIdxString.AppendInt(AssetIdx);
+													return FText::FromString(AssetIdxString);
+												}
+											}
+										}
+									}
+								}
+							}
+							return FText::GetEmpty();
+						});
+			}
+
+			if (ChannelItem->GetCompactLabel() == "SourceAssetIndex")
+			{
+				return SNew(STextBlock)
+					.Margin(FMargin(1.f, 1.f, 1.f, 1.f))
+					.Text_Lambda([this, ColumnName]() -> FText
+						{
+							if (!ChannelItem->IsExpanded())
+							{
+								if (const TSharedPtr<FDatabaseViewModel> ViewModel = EditorViewModel.Pin())
+								{
+									const UPoseSearchDatabase* PoseSearchDatabase = ViewModel->GetPoseSearchDatabase();
+									if (FAsyncPoseSearchDatabasesManagement::RequestAsyncBuildIndex(PoseSearchDatabase, ERequestAsyncBuildFlag::ContinueRequest))
+									{
+										for (const TArray<FDatabasePreviewActor>& PreviewActorGroup : ViewModel->GetPreviewActors())
+										{
+											for (const FDatabasePreviewActor& PreviewActor : PreviewActorGroup)
+											{
+												if (*GetNameSafe(PreviewActor.GetActor()) == ColumnName)
+												{
+													const FSearchIndexAsset& IndexAsset = PoseSearchDatabase->GetSearchIndex().Assets[PreviewActor.GetIndexAssetIndex()];
+													const int32 AssetIdx = PreviewActor.GetIndexAssetIndex();
+													FString AssetIdxString;
+													AssetIdxString.AppendInt(IndexAsset.GetSourceAssetIdx());
+													return FText::FromString(AssetIdxString);
+												}
+											}
 										}
 									}
 								}
 							}
-							return FText::FromString(PoseIdxString);
+							return FText::GetEmpty();
 						});
 			}
 		}
@@ -278,53 +350,58 @@ public:
 							const UPoseSearchDatabase* PoseSearchDatabase = ViewModel->GetPoseSearchDatabase();
 							if (FAsyncPoseSearchDatabasesManagement::RequestAsyncBuildIndex(PoseSearchDatabase, ERequestAsyncBuildFlag::ContinueRequest))
 							{
-								if (const FDatabasePreviewActor* FoundPreviewActor = ViewModel->GetPreviewActors().FindByPredicate(
-									[ColumnName](const FDatabasePreviewActor& PreviewActor)
-									{ return *GetNameSafe(PreviewActor.GetActor()) == ColumnName; }))
+								for (const TArray<FDatabasePreviewActor>& PreviewActorGroup : ViewModel->GetPreviewActors())
 								{
-									const int32 PoseIdx = FoundPreviewActor->GetCurrentPoseIndex();
-									const TArray<float> PoseValues = PoseSearchDatabase->GetSearchIndex().GetPoseValuesSafe(PoseIdx);
-
-									if (Cardinality > 1)
+									for (const FDatabasePreviewActor& PreviewActor : PreviewActorGroup)
 									{
-										// using only one decimal to keep the string compact
-										for (int32 i = 0; i < Cardinality; ++i)
+										if (*GetNameSafe(PreviewActor.GetActor()) == ColumnName)
 										{
-											if (i != 0)
-											{
-												StringBuilder.Append(TEXT(", "));
-											}
+											const int32 PoseIdx = PreviewActor.GetCurrentPoseIndex();
+											const TArray<float> PoseValues = PoseSearchDatabase->GetSearchIndex().GetPoseValuesSafe(PoseIdx);
 
-											const int32 PoseValueIndex = i + DataOffset;
-											if (PoseValueIndex < PoseValues.Num())
+											if (Cardinality > 1)
 											{
-												const float Value = PoseValues[PoseValueIndex];
-												StringBuilder.Appendf(TEXT("%.1f"), Value);
+												// using only one decimal to keep the string compact
+												for (int32 i = 0; i < Cardinality; ++i)
+												{
+													if (i != 0)
+													{
+														StringBuilder.Append(TEXT(", "));
+													}
+
+													const int32 PoseValueIndex = i + DataOffset;
+													if (PoseValueIndex < PoseValues.Num())
+													{
+														const float Value = PoseValues[PoseValueIndex];
+														StringBuilder.Appendf(TEXT("%.1f"), Value);
+													}
+													else
+													{
+														StringBuilder.Append(TEXT("---"));
+													}
+												}
 											}
 											else
 											{
-												StringBuilder.Append(TEXT("---"));
+												// using all the float digits 
+												if (DataOffset < PoseValues.Num())
+												{
+													const float Value = PoseValues[DataOffset];
+													StringBuilder.Appendf(TEXT("%f"), Value);
+												}
+												else
+												{
+													StringBuilder.Append(TEXT("---"));
+												}
 											}
-										}
-									}
-									else
-									{
-										// using all the float digits 
-										if (DataOffset < PoseValues.Num())
-										{
-											const float Value = PoseValues[DataOffset];
-											StringBuilder.Appendf(TEXT("%f"), Value);
-										}
-										else
-										{
-											StringBuilder.Append(TEXT("---"));
+											return FText::FromString(StringBuilder.ToString());
 										}
 									}
 								}
 							}
 						}
 					}
-					return FText::FromString(StringBuilder.ToString());
+					return FText::GetEmpty();
 				});
 	}
 };
@@ -406,13 +483,25 @@ void SDatabaseDataDetails::Reconstruct(int32 MaxPreviewActors)
 			.ToolTipText(LOCTEXT("Query_ToolTip", "Query Values")));
 	}
 
-	const TArray<FDatabasePreviewActor>& PreviewActors = ViewModel->GetPreviewActors();
-	const int32 PreviewActorNum = FMath::Min(MaxPreviewActors, PreviewActors.Num());
-	for (int32 PreviewActorIdx = 0; PreviewActorIdx < PreviewActorNum; ++PreviewActorIdx)
+	int32 AddedColumns = 0;
+	for (const TArray<FDatabasePreviewActor>& PreviewActorGroup : ViewModel->GetPreviewActors())
 	{
-		const FDatabasePreviewActor& PreviewActor = PreviewActors[PreviewActorIdx];
-		HeaderRow->AddColumn(SHeaderRow::Column(*GetNameSafe(PreviewActor.GetActor()))
-			.DefaultLabel(FText::FromString(PreviewActor.GetSampler().GetAsset()->GetName())));
+		for (const FDatabasePreviewActor& PreviewActor : PreviewActorGroup)
+		{
+			HeaderRow->AddColumn(SHeaderRow::Column(*GetNameSafe(PreviewActor.GetActor())) 
+				.DefaultLabel(FText::FromString(PreviewActor.GetSampler().GetAsset()->GetName())));
+
+			++AddedColumns;
+			if (AddedColumns >= MaxPreviewActors)
+			{
+				break;
+			}
+		}
+
+		if (AddedColumns >= MaxPreviewActors)
+		{
+			break;
+		}
 	}
 
 	ChannelItemsTreeView = SNew(SChannelItemsTreeView)
@@ -476,6 +565,8 @@ void SDatabaseDataDetails::RebuildChannelItemsStats(TArray<FChannelItemPtr>& Cha
 	ChannelItems.Add(ChannelItem);
 
 	ChannelItem->GetChannelItems().Add(MakeShareable(new FChannelItem("SDatabaseDataDetailsPoseIndex", "PoseIndex")));
+	ChannelItem->GetChannelItems().Add(MakeShareable(new FChannelItem("SDatabaseDataDetailsPoseIndex", "SearchAssetIndex")));
+	ChannelItem->GetChannelItems().Add(MakeShareable(new FChannelItem("SDatabaseDataDetailsPoseIndex", "SourceAssetIndex")));
 	// @todo: add support for additional stats
 	//ChannelItem->GetChannelItems().Add(MakeShareable(new FChannelItem("SDatabaseDataDetailsValueVectorIndex", "ValueVectorIndex")));
 	//ChannelItem->GetChannelItems().Add(MakeShareable(new FChannelItem("SDatabaseDataDetailsPCAValueVectorIndex", "PCAValueVectorIndex")));
diff --git a/Engine/Plugins/Experimental/Animation/PoseSearch/Source/Editor/Private/PoseSearchDatabaseEdMode.cpp b/Engine/Plugins/Experimental/Animation/PoseSearch/Source/Editor/Private/PoseSearchDatabaseEdMode.cpp
--- a/Engine/Plugins/Experimental/Animation/PoseSearch/Source/Editor/Private/PoseSearchDatabaseEdMode.cpp
+++ b/Engine/Plugins/Experimental/Animation/PoseSearch/Source/Editor/Private/PoseSearchDatabaseEdMode.cpp
@@ -52,44 +52,47 @@ namespace UE::PoseSearch
 
 			TArray<FTransform> WorldTransforms;
 			TArray<FLinearColor> BoneColors;
-			for (const FDatabasePreviewActor& PreviewActor : ViewModel->GetPreviewActors())
+			for (const TArray<FDatabasePreviewActor>& PreviewActorGroup : ViewModel->GetPreviewActors())
 			{
-				const UDebugSkelMeshComponent* MeshComponent = PreviewActor.GetDebugSkelMeshComponent();
-				if (MeshComponent && MeshComponent->GetSkeletalMeshAsset() && MeshComponent->GetNumDrawTransform() > 0 && MeshComponent->SkeletonDrawMode != ESkeletonDrawMode::Hidden)
+				for (const FDatabasePreviewActor& PreviewActor : PreviewActorGroup)
 				{
-					WorldTransforms.SetNumUninitialized(MeshComponent->GetNumDrawTransform());
-					BoneColors.SetNumUninitialized(MeshComponent->GetNumDrawTransform());
-
-					// factor skeleton draw mode into color selection
-					const FLinearColor BoneColor = MeshComponent->SkeletonDrawMode == ESkeletonDrawMode::GreyedOut ? PersonaOptions->DisabledBoneColor : PersonaOptions->DefaultBoneColor;
-					const FLinearColor VirtualBoneColor = MeshComponent->SkeletonDrawMode == ESkeletonDrawMode::GreyedOut ? PersonaOptions->DisabledBoneColor : PersonaOptions->VirtualBoneColor;
-
-					const TArray<FBoneIndexType>& DrawBoneIndices = MeshComponent->GetDrawBoneIndices();
-					for (int32 Index = 0; Index < DrawBoneIndices.Num(); ++Index)
-					{
-						const int32 BoneIndex = DrawBoneIndices[Index];
-						WorldTransforms[BoneIndex] = MeshComponent->GetDrawTransform(BoneIndex) * MeshComponent->GetComponentTransform();
-						BoneColors[BoneIndex] = BoneColor;
-					}
-
-					// color virtual bones
-					for (int16 VirtualBoneIndex : MeshComponent->GetReferenceSkeleton().GetRequiredVirtualBones())
+					const UDebugSkelMeshComponent* MeshComponent = PreviewActor.GetDebugSkelMeshComponent();
+					if (MeshComponent && MeshComponent->GetSkeletalMeshAsset() && MeshComponent->GetNumDrawTransform() > 0 && MeshComponent->SkeletonDrawMode != ESkeletonDrawMode::Hidden)
 					{
-						BoneColors[VirtualBoneIndex] = VirtualBoneColor;
+						WorldTransforms.SetNumUninitialized(MeshComponent->GetNumDrawTransform());
+						BoneColors.SetNumUninitialized(MeshComponent->GetNumDrawTransform());
+
+						// factor skeleton draw mode into color selection
+						const FLinearColor BoneColor = MeshComponent->SkeletonDrawMode == ESkeletonDrawMode::GreyedOut ? PersonaOptions->DisabledBoneColor : PersonaOptions->DefaultBoneColor;
+						const FLinearColor VirtualBoneColor = MeshComponent->SkeletonDrawMode == ESkeletonDrawMode::GreyedOut ? PersonaOptions->DisabledBoneColor : PersonaOptions->VirtualBoneColor;
+
+						const TArray<FBoneIndexType>& DrawBoneIndices = MeshComponent->GetDrawBoneIndices();
+						for (int32 Index = 0; Index < DrawBoneIndices.Num(); ++Index)
+						{
+							const int32 BoneIndex = DrawBoneIndices[Index];
+							WorldTransforms[BoneIndex] = MeshComponent->GetDrawTransform(BoneIndex) * MeshComponent->GetComponentTransform();
+							BoneColors[BoneIndex] = BoneColor;
+						}
+
+						// color virtual bones
+						for (int16 VirtualBoneIndex : MeshComponent->GetReferenceSkeleton().GetRequiredVirtualBones())
+						{
+							BoneColors[VirtualBoneIndex] = VirtualBoneColor;
+						}
+
+						FSkelDebugDrawConfig DrawConfig;
+						DrawConfig.BoneDrawMode = (EBoneDrawMode::Type)PersonaOptions->DefaultBoneDrawSelection;
+						DrawConfig.BoneDrawSize = 0.2f;
+						DrawConfig.bAddHitProxy = false;
+						DrawConfig.bForceDraw = false;
+						DrawConfig.DefaultBoneColor = PersonaOptions->DefaultBoneColor;
+						DrawConfig.AffectedBoneColor = PersonaOptions->AffectedBoneColor;
+						DrawConfig.SelectedBoneColor = PersonaOptions->SelectedBoneColor;
+						DrawConfig.ParentOfSelectedBoneColor = PersonaOptions->ParentOfSelectedBoneColor;
+
+						SkeletalDebugRendering::DrawBones(PDI, MeshComponent->GetComponentLocation(), DrawBoneIndices, MeshComponent->GetReferenceSkeleton(),
+							WorldTransforms, MeshComponent->BonesOfInterest, BoneColors, TArray<TRefCountPtr<HHitProxy>>(), DrawConfig);
 					}
-
-					FSkelDebugDrawConfig DrawConfig;
-					DrawConfig.BoneDrawMode = (EBoneDrawMode::Type)PersonaOptions->DefaultBoneDrawSelection;
-					DrawConfig.BoneDrawSize = 0.2f;
-					DrawConfig.bAddHitProxy = false;
-					DrawConfig.bForceDraw = false;
-					DrawConfig.DefaultBoneColor = PersonaOptions->DefaultBoneColor;
-					DrawConfig.AffectedBoneColor = PersonaOptions->AffectedBoneColor;
-					DrawConfig.SelectedBoneColor = PersonaOptions->SelectedBoneColor;
-					DrawConfig.ParentOfSelectedBoneColor = PersonaOptions->ParentOfSelectedBoneColor;
-
-					SkeletalDebugRendering::DrawBones(PDI, MeshComponent->GetComponentLocation(), DrawBoneIndices, MeshComponent->GetReferenceSkeleton(),
-						WorldTransforms, MeshComponent->BonesOfInterest, BoneColors, TArray<TRefCountPtr<HHitProxy>>(), DrawConfig);
 				}
 			}
 		}
diff --git a/Engine/Plugins/Experimental/Animation/PoseSearch/Source/Editor/Private/PoseSearchDatabaseEditor.cpp b/Engine/Plugins/Experimental/Animation/PoseSearch/Source/Editor/Private/PoseSearchDatabaseEditor.cpp
--- a/Engine/Plugins/Experimental/Animation/PoseSearch/Source/Editor/Private/PoseSearchDatabaseEditor.cpp
+++ b/Engine/Plugins/Experimental/Animation/PoseSearch/Source/Editor/Private/PoseSearchDatabaseEditor.cpp
@@ -517,12 +517,12 @@ namespace UE::PoseSearch
 		{
 			for (TSharedPtr<FDatabaseAssetTreeNode>& SelectedItem : SelectedItems)
 			{
-				if (!SelectedItem.IsValid() || !PoseSearchDatabase->AnimationAssets.IsValidIndex(SelectedItem->SourceAssetIdx))
+				if (!SelectedItem.IsValid() || !PoseSearchDatabase->GetAnimationAssets().IsValidIndex(SelectedItem->SourceAssetIdx))
 				{
 					continue;
 				}
 
-				const FInstancedStruct& DatabaseAsset = PoseSearchDatabase->AnimationAssets[SelectedItem->SourceAssetIdx];
+				const FInstancedStruct& DatabaseAsset = PoseSearchDatabase->GetAnimationAssetStruct(SelectedItem->SourceAssetIdx);
 				const UScriptStruct* ScriptStruct = DatabaseAsset.GetScriptStruct();
 				FSelectionWidget& SelectionWidget = FindOrAddSelectionWidget(ScriptStruct);
 
@@ -574,6 +574,18 @@ namespace UE::PoseSearch
 
 					SelectionWidget.SelectedReflections.Add(NewSelectionReflection);
 				}
+				else if (const FPoseSearchDatabaseMultiSequence* DatabaseMultiSequence = DatabaseAsset.GetPtr<FPoseSearchDatabaseMultiSequence>())
+				{
+					UPoseSearchDatabaseMultiSequenceReflection* NewSelectionReflection = NewObject<UPoseSearchDatabaseMultiSequenceReflection>();
+					NewSelectionReflection->AddToRoot();
+					static_cast<FPoseSearchDatabaseMultiSequence&>(NewSelectionReflection->MultiSequence) = *DatabaseMultiSequence;
+					NewSelectionReflection->MultiSequence.bLooping = DatabaseMultiSequence->IsLooping();
+					NewSelectionReflection->MultiSequence.bHasRootMotion = DatabaseMultiSequence->IsRootMotionEnabled();
+					NewSelectionReflection->SetSourceLink(SelectedItem, AssetTreeWidget);
+					NewSelectionReflection->SetFlags(RF_Transactional);
+
+					SelectionWidget.SelectedReflections.Add(NewSelectionReflection);
+				}
 				else
 				{
 					checkNoEntry();
diff --git a/Engine/Plugins/Experimental/Animation/PoseSearch/Source/Editor/Private/PoseSearchDatabaseEditorReflection.cpp b/Engine/Plugins/Experimental/Animation/PoseSearch/Source/Editor/Private/PoseSearchDatabaseEditorReflection.cpp
--- a/Engine/Plugins/Experimental/Animation/PoseSearch/Source/Editor/Private/PoseSearchDatabaseEditorReflection.cpp
+++ b/Engine/Plugins/Experimental/Animation/PoseSearch/Source/Editor/Private/PoseSearchDatabaseEditorReflection.cpp
@@ -103,6 +103,27 @@ void UPoseSearchDatabaseAnimMontageReflection::PostEditChangeProperty(struct FPr
 	}
 }
 
+void UPoseSearchDatabaseMultiSequenceReflection::PostEditChangeProperty(struct FPropertyChangedEvent& PropertyChangedEvent)
+{
+	Super::PostEditChangeProperty(PropertyChangedEvent);
+
+	if (const TSharedPtr<UE::PoseSearch::FDatabaseViewModel> ViewModel = WeakAssetTreeNode.Pin()->EditorViewModel.Pin())
+	{
+		UPoseSearchDatabase* Database = ViewModel->GetPoseSearchDatabase();
+		if (IsValid(Database))
+		{
+			FInstancedStruct& DatabaseAsset = Database->GetMutableAnimationAssetStruct(WeakAssetTreeNode.Pin()->SourceAssetIdx);
+			if (FPoseSearchDatabaseMultiSequence* DatabaseMultiSequence = DatabaseAsset.GetMutablePtr<FPoseSearchDatabaseMultiSequence>())
+			{
+				*DatabaseMultiSequence = MultiSequence;
+				Database->MarkPackageDirty();
+
+				AssetTreeWidget->FinalizeTreeChanges(true);
+			}
+		}
+	}
+}
+
 #endif // WITH_EDITOR
 
 void UPoseSearchDatabaseStatistics::Initialize(const UPoseSearchDatabase* PoseSearchDatabase)
@@ -114,7 +135,7 @@ void UPoseSearchDatabaseStatistics::Initialize(const UPoseSearchDatabase* PoseSe
 		const UE::PoseSearch::FSearchIndex& SearchIndex = PoseSearchDatabase->GetSearchIndex();
 		// General Information
 	
-		AnimationSequences = PoseSearchDatabase->AnimationAssets.Num();
+		AnimationSequences = PoseSearchDatabase->GetAnimationAssets().Num();
 			
 		const int32 SampleRate = FMath::Max(1, PoseSearchDatabase->Schema->SampleRate);
 		TotalAnimationPosesInFrames = SearchIndex.GetNumPoses();
@@ -170,8 +191,8 @@ void UPoseSearchDatabaseStatistics::Initialize(const UPoseSearchDatabase* PoseSe
 		{
 			uint32 SourceAnimAssetsSizeCookedEstimateInBytes = 0;
 			TSet<const UObject*> Analyzed;
-			Analyzed.Reserve(PoseSearchDatabase->AnimationAssets.Num());
-			for (const FInstancedStruct& AnimAsset : PoseSearchDatabase->AnimationAssets)
+			Analyzed.Reserve(PoseSearchDatabase->GetAnimationAssets().Num());
+			for (const FInstancedStruct& AnimAsset : PoseSearchDatabase->GetAnimationAssets())
 			{
 				if (const FPoseSearchDatabaseAnimationAssetBase* DatabaseAnimationAssetBase = AnimAsset.GetPtr<FPoseSearchDatabaseAnimationAssetBase>())
 				{
diff --git a/Engine/Plugins/Experimental/Animation/PoseSearch/Source/Editor/Private/PoseSearchDatabasePreviewScene.cpp b/Engine/Plugins/Experimental/Animation/PoseSearch/Source/Editor/Private/PoseSearchDatabasePreviewScene.cpp
--- a/Engine/Plugins/Experimental/Animation/PoseSearch/Source/Editor/Private/PoseSearchDatabasePreviewScene.cpp
+++ b/Engine/Plugins/Experimental/Animation/PoseSearch/Source/Editor/Private/PoseSearchDatabasePreviewScene.cpp
@@ -66,9 +66,9 @@ namespace UE::PoseSearch
 			const bool bDisplayBlockTransition = ViewModel->IsDisplayBlockTransitionChecked();
 			bool bDrawQueryVector = ViewModel->ShouldDrawQueryVector();
 
-			for (FDatabasePreviewActor& PreviewActor : ViewModel->GetPreviewActors())
+			for (TArray<FDatabasePreviewActor>& PreviewActorGroup : ViewModel->GetPreviewActors())
 			{
-				bDrawQueryVector &= !PreviewActor.DrawPreviewActor(Database, bDisplayRootMotionSpeed, bDisplayBlockTransition, bDrawQueryVector ? ViewModel->GetQueryVector() : TConstArrayView<float>());
+				bDrawQueryVector &= !FDatabasePreviewActor::DrawPreviewActors(PreviewActorGroup, Database, bDisplayRootMotionSpeed, bDisplayBlockTransition, bDrawQueryVector ? ViewModel->GetQueryVector() : TConstArrayView<float>());
 			}
 		}
 	}
diff --git a/Engine/Plugins/Experimental/Animation/PoseSearch/Source/Editor/Private/PoseSearchDatabaseViewModel.cpp b/Engine/Plugins/Experimental/Animation/PoseSearch/Source/Editor/Private/PoseSearchDatabaseViewModel.cpp
--- a/Engine/Plugins/Experimental/Animation/PoseSearch/Source/Editor/Private/PoseSearchDatabaseViewModel.cpp
+++ b/Engine/Plugins/Experimental/Animation/PoseSearch/Source/Editor/Private/PoseSearchDatabaseViewModel.cpp
@@ -19,7 +19,9 @@
 #include "PoseSearch/PoseSearchDatabase.h"
 #include "PoseSearch/PoseSearchDefines.h"
 #include "PoseSearch/PoseSearchDerivedData.h"
+#include "PoseSearch/PoseSearchHistory.h"
 #include "PoseSearch/PoseSearchSchema.h"
+#include "PoseSearch/PoseSearchTrajectoryTypes.h"
 #include "PoseSearchDatabaseAssetTreeNode.h"
 #include "PoseSearchDatabaseDataDetails.h"
 #include "PoseSearchDatabasePreviewScene.h"
@@ -35,33 +37,36 @@ static TAutoConsoleVariable<float> CVarDatabasePreviewDebugDrawSamplerSize(TEXT(
 constexpr float StepDeltaTime = 1.0f / 30.0f;
 
 // FDatabasePreviewActor
-bool FDatabasePreviewActor::SpawnPreviewActor(UWorld* World, const UPoseSearchDatabase* PoseSearchDatabase, int32 IndexAssetIdx, int32 PoseIdxForTimeOffset)
+bool FDatabasePreviewActor::SpawnPreviewActor(UWorld* World, const UPoseSearchDatabase* PoseSearchDatabase, int32 IndexAssetIdx, const FRole& Role, const FTransform& SamplerRootTransformOrigin, const FTransform* PrecalculatedRootTransformOrigin, int32 PoseIdxForTimeOffset)
 {
-	check(PoseSearchDatabase);
+	check(PoseSearchDatabase && PoseSearchDatabase->Schema);
 	const FSearchIndex& SearchIndex = PoseSearchDatabase->GetSearchIndex();
+
+	USkeleton* Skeleton = PoseSearchDatabase->Schema->GetSkeleton(Role);
+	if (!Skeleton)
+	{
+		UE_LOG(LogPoseSearchEditor, Error, TEXT("Couldn't spawn preview Actor because of missing Role '%s' in Schema '%s'"), *Role.ToString(), *PoseSearchDatabase->Schema->GetName());
+		return false;
+	}
+
 	const FSearchIndexAsset& IndexAsset = SearchIndex.Assets[IndexAssetIdx];
 
 	const FPoseSearchDatabaseAnimationAssetBase* DatabaseAnimationAsset = PoseSearchDatabase->GetAnimationAssetBase(IndexAsset.GetSourceAssetIdx());
-	UAnimationAsset* PreviewAsset = Cast<UAnimationAsset>(DatabaseAnimationAsset->GetAnimationAsset());
+	UAnimationAsset* PreviewAsset = Cast<UAnimationAsset>(DatabaseAnimationAsset->GetAnimationAssetForRole(Role));
 	if (!PreviewAsset)
 	{
 		return false;
 	}
 
-	Sampler.Init(PreviewAsset, IndexAsset.GetBlendParameters());
-
-	FBoneContainer BoneContainer;
-	BoneContainer.InitializeTo(PoseSearchDatabase->Schema->BoneIndicesWithParents, UE::Anim::FCurveFilterSettings(UE::Anim::ECurveFilterMode::DisallowAll), *PoseSearchDatabase->Schema->Skeleton);
-	Sampler.Process(BoneContainer);
+	ActorRole = Role;
+	Sampler.Init(PreviewAsset, SamplerRootTransformOrigin, IndexAsset.GetBlendParameters());
+	Sampler.Process();
 
 	IndexAssetIndex = IndexAssetIdx;
 	CurrentPoseIndex = INDEX_NONE;
 
-	if (PoseIdxForTimeOffset < 0)
-	{
-		PlayTimeOffset = 0.f;
-	}
-	else
+	PlayTimeOffset = 0.f;
+	if (PoseIdxForTimeOffset >= 0)
 	{
 		PlayTimeOffset = PoseSearchDatabase->GetRealAssetTime(PoseIdxForTimeOffset) - IndexAsset.GetFirstSampleTime(PoseSearchDatabase->Schema->SampleRate);
 	}
@@ -79,7 +84,7 @@ bool FDatabasePreviewActor::SpawnPreviewActor(UWorld* World, const UPoseSearchDa
 	AnimInstance->InitializeAnimation();
 
 	USkeletalMesh* DatabasePreviewMesh = PoseSearchDatabase->PreviewMesh;
-	Mesh->SetSkeletalMesh(DatabasePreviewMesh ? DatabasePreviewMesh : PoseSearchDatabase->Schema->Skeleton->GetPreviewMesh(true));
+	Mesh->SetSkeletalMesh(DatabasePreviewMesh ? DatabasePreviewMesh : PoseSearchDatabase->Schema->GetSkeleton(Role)->GetPreviewMesh(true));
 	Mesh->EnablePreview(true, PreviewAsset);
 		
 	AnimInstance->SetAnimationAsset(PreviewAsset, false, 0.0f);
@@ -87,13 +92,18 @@ bool FDatabasePreviewActor::SpawnPreviewActor(UWorld* World, const UPoseSearchDa
 		
 	if (IndexAsset.IsMirrored() && PoseSearchDatabase->Schema)
 	{
-		AnimInstance->SetMirrorDataTable(PoseSearchDatabase->Schema->MirrorDataTable);
+		const UMirrorDataTable* MirrorDataTable = PoseSearchDatabase->Schema->GetMirrorDataTable(Role);
+		AnimInstance->SetMirrorDataTable(MirrorDataTable);
 	}
 
 	const FMirrorDataCache MirrorDataCache(AnimInstance->GetMirrorDataTable(), AnimInstance->GetRequiredBonesOnAnyThread());
 	
-	// @todo: should we always use the PlayTimeOffset to extract the root transform?
-	if (PlayTimeOffset != 0.f)
+	if (PrecalculatedRootTransformOrigin)
+	{
+		// using the PrecalculatedRootTransformOrigin if provided. useful to synchronize multiple roles preview actors
+		RootTransformOrigin = *PrecalculatedRootTransformOrigin;
+	}
+	else
 	{
 		RootTransformOrigin = MirrorDataCache.MirrorTransform(Sampler.ExtractRootTransform(PlayTimeOffset));
 	}
@@ -190,134 +200,206 @@ void FDatabasePreviewActor::Destroy()
 	}
 }
 
-bool FDatabasePreviewActor::DrawPreviewActor(const UPoseSearchDatabase* PoseSearchDatabase, bool bDisplayRootMotionSpeed, bool bDisplayBlockTransition, TConstArrayView<float> QueryVector)
+bool FDatabasePreviewActor::DrawPreviewActors(TArrayView<FDatabasePreviewActor> PreviewActors, const UPoseSearchDatabase* PoseSearchDatabase, bool bDisplayRootMotionSpeed, bool bDisplayBlockTransition, TConstArrayView<float> QueryVector)
 {
-	if (!PoseSearchDatabase->GetSearchIndex().IsValidPoseIndex(GetCurrentPoseIndex()))
-	{
-		return false;
-	}
+	using namespace UE::PoseSearch;
+
+	UWorld* CommonWorld = nullptr;
+	int32 CommonCurrentPoseIndex = INDEX_NONE;
+#if DO_CHECK
+	int32 CommonIndexAssetIndex = INDEX_NONE;
+#endif // DO_CHECK
+
+	TArray<const USkinnedMeshComponent*> Meshes;
+	FRoleToIndex RoleToIndex;
+	TArray<FArchivedPoseHistory> ArchivedPoseHistories;
+	TArray<const IPoseHistory*> PoseHistories;
+
+	const int32 NumPreviewActors = PreviewActors.Num();
 
-	const UDebugSkelMeshComponent* Mesh = GetDebugSkelMeshComponent();
-	if (!Mesh)
+	Meshes.Reserve(NumPreviewActors);
+	RoleToIndex.Reserve(NumPreviewActors);
+	ArchivedPoseHistories.Reserve(NumPreviewActors);
+	PoseHistories.Reserve(NumPreviewActors);
+
+	for (FDatabasePreviewActor& PreviewActor : PreviewActors)
 	{
-		return false;
+		if (!PoseSearchDatabase->GetSearchIndex().IsValidPoseIndex(PreviewActor.GetCurrentPoseIndex()))
+		{
+			return false;
+		}
+
+		const UDebugSkelMeshComponent* Mesh = PreviewActor.GetDebugSkelMeshComponent();
+		if (!Mesh)
+		{
+			return false;
+		}
+
+		if (!CommonWorld)
+		{
+			CommonWorld = Mesh->GetWorld();
+		}
+		else if (CommonWorld != Mesh->GetWorld())
+		{
+			return false;
+		}
+
+		// making sure PreviewActors are consistent with each other
+		if (CommonCurrentPoseIndex == INDEX_NONE)
+		{
+			CommonCurrentPoseIndex = PreviewActor.GetCurrentPoseIndex();
+		}
+		else if (CommonCurrentPoseIndex != PreviewActor.GetCurrentPoseIndex())
+		{
+			checkNoEntry();
+			return false;
+		}
+
+#if DO_CHECK
+		if (CommonIndexAssetIndex == INDEX_NONE)
+		{
+			CommonIndexAssetIndex = PreviewActor.IndexAssetIndex;
+		}
+		else if (CommonIndexAssetIndex != PreviewActor.IndexAssetIndex)
+		{
+			checkNoEntry();
+			return false;
+		}
+#endif // DO_CHECK
+
+		RoleToIndex.Add(PreviewActor.ActorRole) = Meshes.Num();
+		Meshes.Add(Mesh);
+
+		const FTransform RootMotionTransform = PreviewActor.RootBoneTransformCurrentQuantizedTime * PreviewActor.RootTransformCurrentQuantizedTime;
+
+		// @todo: reconstruct ArchivedPoseHistory::BoneToTransformMap and ArchivedPoseHistory::Entries if needed
+		FArchivedPoseHistory& ArchivedPoseHistory = ArchivedPoseHistories.AddDefaulted_GetRef();
+		// @todo: reconstruct multiple trajectory samples if needed
+		// reconstructing the FPoseSearchQueryTrajectory with only one sample with AccumulatedSeconds at zero
+		FPoseSearchQueryTrajectorySample& TrajectorySample = ArchivedPoseHistory.Trajectory.Samples.AddDefaulted_GetRef();
+		TrajectorySample.SetTransform(RootMotionTransform);
+		TrajectorySample.AccumulatedSeconds = 0.f;
+
+		PoseHistories.Add(&ArchivedPoseHistory);
 	}
 
-	const FTransform RootBoneTransformCurrentQuantizedTimeWorld = RootBoneTransformCurrentQuantizedTime * RootTransformCurrentQuantizedTime;
-	UE::PoseSearch::FDebugDrawParams DrawParams(Mesh->GetWorld(), Mesh, RootBoneTransformCurrentQuantizedTimeWorld, PoseSearchDatabase);
-	DrawParams.DrawFeatureVector(GetCurrentPoseIndex());
+	UE::PoseSearch::FDebugDrawParams DrawParams(Meshes, PoseHistories, RoleToIndex, PoseSearchDatabase);
+	DrawParams.DrawFeatureVector(CommonCurrentPoseIndex);
 
 	if (!QueryVector.IsEmpty())
 	{
 		DrawParams.DrawFeatureVector(QueryVector);
 	}
 
-	const FSearchIndex& SearchIndex = PoseSearchDatabase->GetSearchIndex();
-	const FSearchIndexAsset& IndexAsset = SearchIndex.Assets[IndexAssetIndex];
-
-	const FMirrorDataCache MirrorDataCache(Mesh->PreviewInstance->GetMirrorDataTable(), Mesh->PreviewInstance->GetRequiredBonesOnAnyThread());
-	if (bDisplayRootMotionSpeed || bDisplayBlockTransition)
+	for (FDatabasePreviewActor& PreviewActor : PreviewActors)
 	{
-		// initializing SampledRootMotion if required
-		if (SampledRootMotion.IsEmpty())
+		const UDebugSkelMeshComponent* Mesh = PreviewActor.GetDebugSkelMeshComponent();
+		const FSearchIndex& SearchIndex = PoseSearchDatabase->GetSearchIndex();
+		const FSearchIndexAsset& IndexAsset = SearchIndex.Assets[PreviewActor.IndexAssetIndex];
+
+		const FMirrorDataCache MirrorDataCache(Mesh->PreviewInstance->GetMirrorDataTable(), Mesh->PreviewInstance->GetRequiredBonesOnAnyThread());
+		if (bDisplayRootMotionSpeed || bDisplayBlockTransition)
 		{
-			const int NumPoses = IndexAsset.GetNumPoses();
-			if (NumPoses > 1)
+			// initializing SampledRootMotion if required
+			if (PreviewActor.SampledRootMotion.IsEmpty())
 			{
-				SampledRootMotion.SetNumUninitialized(NumPoses);
-				SampledRootMotionSpeed.SetNumUninitialized(NumPoses);
-
-				for (int32 Index = 0; Index < NumPoses; ++Index)
+				const int NumPoses = IndexAsset.GetNumPoses();
+				if (NumPoses > 1)
 				{
-					const int32 IndexAssetPoseIdx = Index + IndexAsset.GetFirstPoseIdx();
-					const float IndexAssetPoseTime = IndexAsset.GetTimeFromPoseIndex(IndexAssetPoseIdx, PoseSearchDatabase->Schema->SampleRate);
-					FTransform IndexAssetPoseTransform = MirrorDataCache.MirrorTransform(Sampler.ExtractRootTransform(IndexAssetPoseTime));
-					IndexAssetPoseTransform.SetToRelativeTransform(RootTransformOrigin);
+					PreviewActor.SampledRootMotion.SetNumUninitialized(NumPoses);
+					PreviewActor.SampledRootMotionSpeed.SetNumUninitialized(NumPoses);
 
-					SampledRootMotion[Index] = IndexAssetPoseTransform.GetTranslation();
-				}
+					for (int32 Index = 0; Index < NumPoses; ++Index)
+					{
+						const int32 IndexAssetPoseIdx = Index + IndexAsset.GetFirstPoseIdx();
+						const float IndexAssetPoseTime = IndexAsset.GetTimeFromPoseIndex(IndexAssetPoseIdx, PoseSearchDatabase->Schema->SampleRate);
+						FTransform IndexAssetPoseTransform = MirrorDataCache.MirrorTransform(PreviewActor.Sampler.ExtractRootTransform(IndexAssetPoseTime));
+						IndexAssetPoseTransform.SetToRelativeTransform(PreviewActor.RootTransformOrigin);
 
-				for (int32 Index = 1; Index < NumPoses; ++Index)
-				{
-					const FVector& Start = SampledRootMotion[Index - 1];
-					const FVector& End = SampledRootMotion[Index];
-					SampledRootMotionSpeed[Index] = (Start - End).Length() * PoseSearchDatabase->Schema->PermutationsSampleRate;
+						PreviewActor.SampledRootMotion[Index] = IndexAssetPoseTransform.GetTranslation();
+					}
+
+					for (int32 Index = 1; Index < NumPoses; ++Index)
+					{
+						const FVector& Start = PreviewActor.SampledRootMotion[Index - 1];
+						const FVector& End = PreviewActor.SampledRootMotion[Index];
+						PreviewActor.SampledRootMotionSpeed[Index] = (Start - End).Length() * PoseSearchDatabase->Schema->PermutationsSampleRate;
+					}
+					PreviewActor.SampledRootMotionSpeed[0] = PreviewActor.SampledRootMotionSpeed[1];
 				}
-				SampledRootMotionSpeed[0] = SampledRootMotionSpeed[1];
 			}
 		}
-	}
 
-	if (bDisplayRootMotionSpeed)
-	{
-		// drawing PreviewActor.SampledRootMotion
-		const int32 SampledRootMotionNum = SampledRootMotion.Num();
-		if (SampledRootMotionNum > 1)
+		if (bDisplayRootMotionSpeed)
 		{
-			for (int32 Index = 0; Index < SampledRootMotion.Num(); ++Index)
+			// drawing PreviewActor.SampledRootMotion
+			const int32 SampledRootMotionNum = PreviewActor.SampledRootMotion.Num();
+			if (SampledRootMotionNum > 1)
 			{
-				const FVector& EndDown = SampledRootMotion[Index];
-				const FVector EndUp = EndDown + (SampledRootMotionSpeed[Index] * FVector::UpVector);
-
-				DrawParams.DrawLine(EndDown, EndUp, FColor::Black);
-				if (Index > 0)
+				for (int32 Index = 0; Index < PreviewActor.SampledRootMotion.Num(); ++Index)
 				{
-					const FColor RootMotionColor = Index % 2 == 0 ? FColor::Purple : FColor::Orange;
-					const FVector& StartDown = SampledRootMotion[Index - 1];
-					const FVector StartUp = StartDown + (SampledRootMotionSpeed[Index - 1] * FVector::UpVector);
-					DrawParams.DrawLine(StartDown, EndDown, RootMotionColor);
-					DrawParams.DrawLine(StartUp, EndUp, RootMotionColor);
+					const FVector& EndDown = PreviewActor.SampledRootMotion[Index];
+					const FVector EndUp = EndDown + (PreviewActor.SampledRootMotionSpeed[Index] * FVector::UpVector);
+
+					DrawParams.DrawLine(EndDown, EndUp, FColor::Black);
+					if (Index > 0)
+					{
+						const FColor RootMotionColor = Index % 2 == 0 ? FColor::Purple : FColor::Orange;
+						const FVector& StartDown = PreviewActor.SampledRootMotion[Index - 1];
+						const FVector StartUp = StartDown + (PreviewActor.SampledRootMotionSpeed[Index - 1] * FVector::UpVector);
+						DrawParams.DrawLine(StartDown, EndDown, RootMotionColor);
+						DrawParams.DrawLine(StartUp, EndUp, RootMotionColor);
+					}
 				}
 			}
 		}
-	}
 
-	if (bDisplayBlockTransition)
-	{
-		const int NumPoses = IndexAsset.GetNumPoses();
-		if (NumPoses == SampledRootMotion.Num())
+		if (bDisplayBlockTransition)
 		{
-			for (int32 Index = 0; Index < NumPoses; ++Index)
+			const int NumPoses = IndexAsset.GetNumPoses();
+			if (NumPoses == PreviewActor.SampledRootMotion.Num())
 			{
-				const int32 IndexAssetPoseIdx = Index + IndexAsset.GetFirstPoseIdx();
-				if (SearchIndex.PoseMetadata[IndexAssetPoseIdx].IsBlockTransition())
-				{
-					DrawParams.DrawPoint(SampledRootMotion[Index], FColor::Red);
-				}
-				else
+				for (int32 Index = 0; Index < NumPoses; ++Index)
 				{
-					DrawParams.DrawPoint(SampledRootMotion[Index], FColor::Green);
+					const int32 IndexAssetPoseIdx = Index + IndexAsset.GetFirstPoseIdx();
+					if (SearchIndex.PoseMetadata[IndexAssetPoseIdx].IsBlockTransition())
+					{
+						DrawParams.DrawPoint(PreviewActor.SampledRootMotion[Index], FColor::Red);
+					}
+					else
+					{
+						DrawParams.DrawPoint(PreviewActor.SampledRootMotion[Index], FColor::Green);
+					}
 				}
 			}
 		}
-	}
 
 
 #if ENABLE_ANIM_DEBUG
-	const float DebugDrawSamplerSize = CVarDatabasePreviewDebugDrawSamplerSize.GetValueOnAnyThread();
-	if (DebugDrawSamplerSize > UE_KINDA_SMALL_NUMBER)
-	{
-		// drawing the pose extracted from the Sampler to visually compare with the pose features and the mesh drawing
-		FMemMark Mark(FMemStack::Get());
-		FCompactPose Pose;
-		Pose.SetBoneContainer(&GetAnimPreviewInstance()->GetRequiredBonesOnAnyThread());
-		
-		Sampler.ExtractPose(CurrentTime, Pose);
-		MirrorDataCache.MirrorPose(Pose);
+		const float DebugDrawSamplerSize = CVarDatabasePreviewDebugDrawSamplerSize.GetValueOnAnyThread();
+		if (DebugDrawSamplerSize > UE_KINDA_SMALL_NUMBER)
+		{
+			// drawing the pose extracted from the Sampler to visually compare with the pose features and the mesh drawing
+			FMemMark Mark(FMemStack::Get());
+			FCompactPose Pose;
+			Pose.SetBoneContainer(&PreviewActor.GetAnimPreviewInstance()->GetRequiredBonesOnAnyThread());
 
-		const FTransform RootTransform = MirrorDataCache.MirrorTransform(Sampler.ExtractRootTransform(CurrentTime));
+			PreviewActor.Sampler.ExtractPose(PreviewActor.CurrentTime, Pose);
+			MirrorDataCache.MirrorPose(Pose);
 
-		FCSPose<FCompactPose> ComponentSpacePose;
-		ComponentSpacePose.InitPose(MoveTemp(Pose));
+			const FTransform RootTransform = MirrorDataCache.MirrorTransform(PreviewActor.Sampler.ExtractRootTransform(PreviewActor.CurrentTime));
 
-		for (int32 BoneIndex = 0; BoneIndex < ComponentSpacePose.GetPose().GetNumBones(); ++BoneIndex)
-		{
-			const FTransform BoneWorldTransforms = ComponentSpacePose.GetComponentSpaceTransform(FCompactPoseBoneIndex(BoneIndex)) * RootTransform;
-			DrawParams.DrawPoint(BoneWorldTransforms.GetTranslation(), FColor::Red, DebugDrawSamplerSize);
+			FCSPose<FCompactPose> ComponentSpacePose;
+			ComponentSpacePose.InitPose(MoveTemp(Pose));
+
+			for (int32 BoneIndex = 0; BoneIndex < ComponentSpacePose.GetPose().GetNumBones(); ++BoneIndex)
+			{
+				const FTransform BoneWorldTransforms = ComponentSpacePose.GetComponentSpaceTransform(FCompactPoseBoneIndex(BoneIndex)) * RootTransform;
+				DrawParams.DrawPoint(BoneWorldTransforms.GetTranslation(), FColor::Red, DebugDrawSamplerSize);
+			}
 		}
-	}
 #endif // ENABLE_ANIM_DEBUG
-
+	}
 	return true;
 }
 
@@ -432,9 +514,12 @@ void FDatabaseViewModel::Tick(float DeltaSeconds)
 		{
 			if (FAsyncPoseSearchDatabasesManagement::RequestAsyncBuildIndex(Database, ERequestAsyncBuildFlag::ContinueRequest))
 			{
-				for (FDatabasePreviewActor& PreviewActor : GetPreviewActors())
+				for (TArray<FDatabasePreviewActor>& PreviewActorGroup : PreviewActors)
 				{
-					PreviewActor.UpdatePreviewActor(Database, PlayTime, bQuantizeAnimationToPoseData);
+					for (FDatabasePreviewActor& PreviewActor : PreviewActorGroup)
+					{
+						PreviewActor.UpdatePreviewActor(Database, PlayTime, bQuantizeAnimationToPoseData);
+					}
 				}
 			}
 		}
@@ -450,9 +535,12 @@ void FDatabaseViewModel::RemovePreviewActors()
 	bIsEditorSelection = true;
 	bDrawQueryVector = false;
 
-	for (FDatabasePreviewActor& PreviewActor : PreviewActors)
+	for (TArray<FDatabasePreviewActor>& PreviewActorGroup : PreviewActors)
 	{
-		PreviewActor.Destroy();
+		for (FDatabasePreviewActor& PreviewActor : PreviewActorGroup)
+		{
+			PreviewActor.Destroy();
+		}
 	}
 
 	PreviewActors.Reset();
@@ -464,7 +552,7 @@ void FDatabaseViewModel::AddSequenceToDatabase(UAnimSequence* AnimSequence)
 	{
 		FPoseSearchDatabaseSequence NewAsset;
 		NewAsset.Sequence = AnimSequence;
-		Database->AnimationAssets.Add(FInstancedStruct::Make(NewAsset));
+		Database->AddAnimationAsset(FInstancedStruct::Make(NewAsset));
 	}
 }
 
@@ -474,7 +562,7 @@ void FDatabaseViewModel::AddBlendSpaceToDatabase(UBlendSpace* BlendSpace)
 	{
 		FPoseSearchDatabaseBlendSpace NewAsset;
 		NewAsset.BlendSpace = BlendSpace;
-		Database->AnimationAssets.Add(FInstancedStruct::Make(NewAsset));
+		Database->AddAnimationAsset(FInstancedStruct::Make(NewAsset));
 	}
 }
 
@@ -484,7 +572,7 @@ void FDatabaseViewModel::AddAnimCompositeToDatabase(UAnimComposite* AnimComposit
 	{
 		FPoseSearchDatabaseAnimComposite NewAsset;
 		NewAsset.AnimComposite = AnimComposite;
-		Database->AnimationAssets.Add(FInstancedStruct::Make(NewAsset));
+		Database->AddAnimationAsset(FInstancedStruct::Make(NewAsset));
 	}
 }
 
@@ -494,7 +582,16 @@ void FDatabaseViewModel::AddAnimMontageToDatabase(UAnimMontage* AnimMontage)
 	{
 		FPoseSearchDatabaseAnimMontage NewAsset;
 		NewAsset.AnimMontage = AnimMontage;
-		Database->AnimationAssets.Add(FInstancedStruct::Make(NewAsset));
+		Database->AddAnimationAsset(FInstancedStruct::Make(NewAsset));
+	}
+}
+
+void FDatabaseViewModel::AddMultiSequenceToDatabase()
+{
+	if (UPoseSearchDatabase* Database = GetPoseSearchDatabase())
+	{
+		FPoseSearchDatabaseMultiSequence NewAsset;
+		Database->AddAnimationAsset(FInstancedStruct::Make(NewAsset));
 	}
 }
 
@@ -527,7 +624,7 @@ bool FDatabaseViewModel::DeleteFromDatabase(int32 AnimationAssetIndex)
 				}
 			}
 
-			Database->AnimationAssets.RemoveAt(AnimationAssetIndex);
+			Database->RemoveAnimationAssetAt(AnimationAssetIndex);
 			Database->Modify();
 
 			return true;
@@ -610,23 +707,47 @@ int32 FDatabaseViewModel::SetSelectedNode(int32 PoseIdx, bool bClearSelection, b
 				const uint32 IndexAssetIndex = SearchIndex.PoseMetadata[PoseIdx].GetAssetIndex();
 				if (SearchIndex.Assets.IsValidIndex(IndexAssetIndex))
 				{
-					FDatabasePreviewActor PreviewActor;
-					if (PreviewActor.SpawnPreviewActor(GetWorld(), Database, IndexAssetIndex, PoseIdx))
+					const FSearchIndexAsset& IndexAsset = SearchIndex.Assets[IndexAssetIndex];
+					const FPoseSearchDatabaseAnimationAssetBase* DatabaseAnimationAsset = Database->GetAnimationAssetBase(IndexAsset.GetSourceAssetIdx());
+					check(DatabaseAnimationAsset);
+					int32 PreviewActorGroupIndex = INDEX_NONE;
+					bool bIsSynchronizedRootTransformOriginValid = false;
+					FTransform SynchronizedRootTransformOrigin = FTransform::Identity;
+					for (int32 RoleIndex = 0; RoleIndex < DatabaseAnimationAsset->GetNumRoles(); ++RoleIndex)
 					{
-						const FSearchIndexAsset& IndexAsset = SearchIndex.Assets[IndexAssetIndex];
-						MaxPreviewPlayLength = FMath::Max(MaxPreviewPlayLength, IndexAsset.GetLastSampleTime(Database->Schema->SampleRate) - PreviewActor.GetPlayTimeOffset());
-						MinPreviewPlayLength = FMath::Min(MinPreviewPlayLength, IndexAsset.GetFirstSampleTime(Database->Schema->SampleRate) - PreviewActor.GetPlayTimeOffset());
-						PreviewActors.Add(PreviewActor);
-						SelectedSourceAssetIdx = IndexAsset.GetSourceAssetIdx();
+						FDatabasePreviewActor PreviewActor;
+						const UE::PoseSearch::FRole Role = DatabaseAnimationAsset->GetRole(RoleIndex);
+						const FTransform& RootTransformOrigin = DatabaseAnimationAsset->GetRootTransformOriginForRole(Role);
+						if (PreviewActor.SpawnPreviewActor(GetWorld(), Database, IndexAssetIndex, Role, RootTransformOrigin, bIsSynchronizedRootTransformOriginValid ? &SynchronizedRootTransformOrigin : nullptr, PoseIdx))
+						{
+							if (PreviewActorGroupIndex == INDEX_NONE)
+							{
+								PreviewActorGroupIndex = PreviewActors.AddDefaulted();
+							}
+
+							if (!bIsSynchronizedRootTransformOriginValid)
+							{
+								bIsSynchronizedRootTransformOriginValid = true;
+								SynchronizedRootTransformOrigin = PreviewActor.GetRootTransformOrigin();
+							}
+
+							MaxPreviewPlayLength = FMath::Max(MaxPreviewPlayLength, IndexAsset.GetLastSampleTime(Database->Schema->SampleRate) - PreviewActor.GetPlayTimeOffset());
+							MinPreviewPlayLength = FMath::Min(MinPreviewPlayLength, IndexAsset.GetFirstSampleTime(Database->Schema->SampleRate) - PreviewActor.GetPlayTimeOffset());
+							PreviewActors[PreviewActorGroupIndex].Add(PreviewActor);
+							SelectedSourceAssetIdx = IndexAsset.GetSourceAssetIdx();
+						}
 					}
 				}
 			}
 
 			DatabaseDataDetails.Pin()->Reconstruct();
 
-			for (FDatabasePreviewActor& PreviewActor : GetPreviewActors())
+			for (TArray<FDatabasePreviewActor>& PreviewActorGroup : PreviewActors)
 			{
-				PreviewActor.UpdatePreviewActor(Database, PlayTime, bQuantizeAnimationToPoseData);
+				for (FDatabasePreviewActor& PreviewActor : PreviewActorGroup)
+				{
+					PreviewActor.UpdatePreviewActor(Database, PlayTime, bQuantizeAnimationToPoseData);
+				}
 			}
 
 			SetPlayTime(0.f, false);
@@ -656,21 +777,45 @@ void FDatabaseViewModel::SetSelectedNodes(const TArrayView<TSharedPtr<FDatabaseA
 			for (int32 IndexAssetIndex = 0; IndexAssetIndex < SearchIndex.Assets.Num(); ++IndexAssetIndex)
 			{
 				const FSearchIndexAsset& IndexAsset = SearchIndex.Assets[IndexAssetIndex];
-				if (const int32* SelectedNodesIndex = AssociatedAssetIndices.Find(IndexAsset.GetSourceAssetIdx()))
+				if (AssociatedAssetIndices.Find(IndexAsset.GetSourceAssetIdx()))
 				{
-					FDatabasePreviewActor PreviewActor;
-					if (PreviewActor.SpawnPreviewActor(GetWorld(), Database, IndexAssetIndex))
+					const FPoseSearchDatabaseAnimationAssetBase* DatabaseAnimationAsset = Database->GetAnimationAssetBase(IndexAsset.GetSourceAssetIdx());
+					check(DatabaseAnimationAsset);
+					int32 PreviewActorGroupIndex = INDEX_NONE;
+					bool bIsSynchronizedRootTransformOriginValid = false;
+					FTransform SynchronizedRootTransformOrigin = FTransform::Identity;
+					for (int32 RoleIndex = 0; RoleIndex < DatabaseAnimationAsset->GetNumRoles(); ++RoleIndex)
 					{
-						MaxPreviewPlayLength = FMath::Max(MaxPreviewPlayLength, IndexAsset.GetLastSampleTime(Database->Schema->SampleRate) - IndexAsset.GetFirstSampleTime(Database->Schema->SampleRate));
-						PreviewActors.Add(PreviewActor);
+						FDatabasePreviewActor PreviewActor;
+						const UE::PoseSearch::FRole Role = DatabaseAnimationAsset->GetRole(RoleIndex);
+						const FTransform& RootTransformOrigin = DatabaseAnimationAsset->GetRootTransformOriginForRole(Role);
+						if (PreviewActor.SpawnPreviewActor(GetWorld(), Database, IndexAssetIndex, Role, RootTransformOrigin, bIsSynchronizedRootTransformOriginValid ? &SynchronizedRootTransformOrigin : nullptr))
+						{
+							if (PreviewActorGroupIndex == INDEX_NONE)
+							{
+								PreviewActorGroupIndex = PreviewActors.AddDefaulted();
+							}
+
+							if (!bIsSynchronizedRootTransformOriginValid)
+							{
+								bIsSynchronizedRootTransformOriginValid = true;
+								SynchronizedRootTransformOrigin = PreviewActor.GetRootTransformOrigin();
+							}
+
+							MaxPreviewPlayLength = FMath::Max(MaxPreviewPlayLength, IndexAsset.GetLastSampleTime(Database->Schema->SampleRate) - IndexAsset.GetFirstSampleTime(Database->Schema->SampleRate));
+							PreviewActors[PreviewActorGroupIndex].Add(PreviewActor);
+						}
 					}
 				}
 			}
 
 			DatabaseDataDetails.Pin()->Reconstruct();
-			for (FDatabasePreviewActor& PreviewActor : GetPreviewActors())
+			for (TArray<FDatabasePreviewActor>& PreviewActorGroup : PreviewActors)
 			{
-				PreviewActor.UpdatePreviewActor(Database, PlayTime, bQuantizeAnimationToPoseData);
+				for (FDatabasePreviewActor& PreviewActor : PreviewActorGroup)
+				{
+					PreviewActor.UpdatePreviewActor(Database, PlayTime, bQuantizeAnimationToPoseData);
+				}
 			}
 		}
 
@@ -681,9 +826,17 @@ void FDatabaseViewModel::SetSelectedNodes(const TArrayView<TSharedPtr<FDatabaseA
 void FDatabaseViewModel::ProcessSelectedActor(AActor* Actor)
 {
 	SelectedActorIndexAssetIndex = INDEX_NONE;
-	if (const FDatabasePreviewActor* SelectedPreviewActor = PreviewActors.FindByPredicate([Actor](const FDatabasePreviewActor& PreviewActor) { return PreviewActor.GetActor() == Actor; }))
+
+	for (const TArray<FDatabasePreviewActor>& PreviewActorGroup : PreviewActors)
 	{
-		SelectedActorIndexAssetIndex = SelectedPreviewActor->GetIndexAssetIndex();
+		for (const FDatabasePreviewActor& PreviewActor : PreviewActorGroup)
+		{
+			if (PreviewActor.GetActor() == Actor)
+			{
+				SelectedActorIndexAssetIndex = PreviewActor.GetIndexAssetIndex();
+				return;
+			}
+		}
 	}
 }
 
@@ -736,9 +889,12 @@ void FDatabaseViewModel::SetPlayTime(float NewPlayTime, bool bInTickPlayTime)
 		{
 			if (FAsyncPoseSearchDatabasesManagement::RequestAsyncBuildIndex(Database, ERequestAsyncBuildFlag::ContinueRequest))
 			{
-				for (FDatabasePreviewActor& PreviewActor : GetPreviewActors())
+				for (TArray<FDatabasePreviewActor>& PreviewActorGroup : PreviewActors)
 				{
-					PreviewActor.UpdatePreviewActor(Database, PlayTime, bQuantizeAnimationToPoseData);
+					for (FDatabasePreviewActor& PreviewActor : PreviewActorGroup)
+					{
+						PreviewActor.UpdatePreviewActor(Database, PlayTime, bQuantizeAnimationToPoseData);
+					}
 				}
 			}
 		}
@@ -752,16 +908,19 @@ bool FDatabaseViewModel::GetAnimationTime(int32 SourceAssetIdx, float& CurrentPl
 		if (FAsyncPoseSearchDatabasesManagement::RequestAsyncBuildIndex(Database, ERequestAsyncBuildFlag::ContinueRequest))
 		{
 			const FSearchIndex& SearchIndex = Database->GetSearchIndex();
-			for (const FDatabasePreviewActor& PreviewActor : GetPreviewActors())
+			for (const TArray<FDatabasePreviewActor>& PreviewActorGroup : PreviewActors)
 			{
-				if (PreviewActor.GetIndexAssetIndex() >= 0 && PreviewActor.GetIndexAssetIndex() < SearchIndex.Assets.Num())
+				for (const FDatabasePreviewActor& PreviewActor : PreviewActorGroup)
 				{
-					const FSearchIndexAsset& IndexAsset = SearchIndex.Assets[PreviewActor.GetIndexAssetIndex()];
-					if (IndexAsset.GetSourceAssetIdx() == SourceAssetIdx)
+					if (PreviewActor.GetIndexAssetIndex() >= 0 && PreviewActor.GetIndexAssetIndex() < SearchIndex.Assets.Num())
 					{
-						CurrentPlayTime = PreviewActor.GetSampler().ToNormalizedTime(PlayTime + IndexAsset.GetFirstSampleTime(Database->Schema->SampleRate) + PreviewActor.GetPlayTimeOffset());
-						BlendParameters = IndexAsset.GetBlendParameters();
-						return true;
+						const FSearchIndexAsset& IndexAsset = SearchIndex.Assets[PreviewActor.GetIndexAssetIndex()];
+						if (IndexAsset.GetSourceAssetIdx() == SourceAssetIdx)
+						{
+							CurrentPlayTime = PreviewActor.GetSampler().ToNormalizedTime(PlayTime + IndexAsset.GetFirstSampleTime(Database->Schema->SampleRate) + PreviewActor.GetPlayTimeOffset());
+							BlendParameters = IndexAsset.GetBlendParameters();
+							return true;
+						}
 					}
 				}
 			}
diff --git a/Engine/Plugins/Experimental/Animation/PoseSearch/Source/Editor/Private/PoseSearchDebuggerDatabaseView.cpp b/Engine/Plugins/Experimental/Animation/PoseSearch/Source/Editor/Private/PoseSearchDebuggerDatabaseView.cpp
--- a/Engine/Plugins/Experimental/Animation/PoseSearch/Source/Editor/Private/PoseSearchDebuggerDatabaseView.cpp
+++ b/Engine/Plugins/Experimental/Animation/PoseSearch/Source/Editor/Private/PoseSearchDebuggerDatabaseView.cpp
@@ -35,7 +35,7 @@ public:
 		TLabelBuilder LabelBuilder;
 		for (const FTraceMotionMatchingStateDatabaseEntry& DbEntry : DatabaseEntries)
 		{
-			const UPoseSearchDatabase* Database = FTraceMotionMatchingState::GetObjectFromId<UPoseSearchDatabase>(DbEntry.DatabaseId);
+			const UPoseSearchDatabase* Database = FTraceMotionMatchingStateMessage::GetObjectFromId<UPoseSearchDatabase>(DbEntry.DatabaseId);
 			if (FAsyncPoseSearchDatabasesManagement::RequestAsyncBuildIndex(Database, ERequestAsyncBuildFlag::ContinueRequest))
 			{
 				for (const TObjectPtr<UPoseSearchFeatureChannel>& ChannelPtr : Database->Schema->GetChannels())
@@ -219,21 +219,15 @@ static void AddUnfilteredDatabaseRow(const UPoseSearchDatabase* Database,
 		const FInstancedStruct& DatabaseAssetStruct = Database->GetAnimationAssetStruct(*SearchIndexAsset);
 		if (const FPoseSearchDatabaseAnimationAssetBase* DatabaseAsset = DatabaseAssetStruct.GetPtr<FPoseSearchDatabaseAnimationAssetBase>())
 		{
+			const float PlayLength = DatabaseAsset->GetPlayLength();
+			UObject* AnimationAsset = DatabaseAsset->GetAnimationAsset();
+
 			Row->AssetName = DatabaseAsset->GetName();
-			Row->AssetPath = DatabaseAsset->GetAnimationAsset() ? DatabaseAsset->GetAnimationAsset()->GetPathName() : "";
+			Row->AssetPath = AnimationAsset ? AnimationAsset->GetPathName() : "";
 			Row->bLooping = DatabaseAsset->IsLooping();
 			Row->BlendParameters = SearchIndexAsset->GetBlendParameters();
-			Row->AnimFrame = 0;
-			Row->AnimPercentage = 0.0f;
-
-			if (const FPoseSearchDatabaseAnimationAssetBase* DatabaseAnimationAssetBase = DatabaseAssetStruct.GetPtr<FPoseSearchDatabaseAnimationAssetBase>())
-			{
-				if (const UAnimSequenceBase* SequenceBase = Cast<UAnimSequenceBase>(DatabaseAnimationAssetBase->GetAnimationAsset()))
-				{
-					Row->AnimFrame = SequenceBase->GetFrameAtTime(Time);
-					Row->AnimPercentage = Time / SequenceBase->GetPlayLength();
-				}
-			}
+			Row->AnimFrame = DatabaseAsset->GetFrameAtTime(Time);
+			Row->AnimPercentage = FMath::IsNearlyZero(PlayLength) ? 0.f : Time / PlayLength;
 		}
 	}
 }
@@ -262,7 +256,7 @@ void SDebuggerDatabaseView::Update(const FTraceMotionMatchingStateMessage& State
 	TArray<uint32> PoseToPCAValuesVectorIndexes;
 	for (const FTraceMotionMatchingStateDatabaseEntry& DbEntry : State.DatabaseEntries)
 	{
-		const UPoseSearchDatabase* Database = FTraceMotionMatchingState::GetObjectFromId<UPoseSearchDatabase>(DbEntry.DatabaseId);
+		const UPoseSearchDatabase* Database = FTraceMotionMatchingStateMessage::GetObjectFromId<UPoseSearchDatabase>(DbEntry.DatabaseId);
 		if (FAsyncPoseSearchDatabasesManagement::RequestAsyncBuildIndex(Database, ERequestAsyncBuildFlag::ContinueRequest))
 		{
 			const FSearchIndex& SearchIndex = Database->GetSearchIndex();
diff --git a/Engine/Plugins/Experimental/Animation/PoseSearch/Source/Editor/Private/PoseSearchDebuggerView.cpp b/Engine/Plugins/Experimental/Animation/PoseSearch/Source/Editor/Private/PoseSearchDebuggerView.cpp
--- a/Engine/Plugins/Experimental/Animation/PoseSearch/Source/Editor/Private/PoseSearchDebuggerView.cpp
+++ b/Engine/Plugins/Experimental/Animation/PoseSearch/Source/Editor/Private/PoseSearchDebuggerView.cpp
@@ -12,6 +12,7 @@
 #include "PoseSearchDebuggerDatabaseView.h"
 #include "PoseSearchDebuggerReflection.h"
 #include "PoseSearchDebuggerViewModel.h"
+#include "PoseSearchEditor.h"
 #include "PropertyEditorModule.h"
 #include "Subsystems/AssetEditorSubsystem.h"
 #include "Trace/PoseSearchTraceProvider.h"
@@ -178,19 +179,20 @@ void SDebuggerView::Tick(const FGeometry& AllottedGeometry, const double InCurre
 	}
 
 	const UWorld* DebuggerWorld = FDebugger::GetWorld();
-    check(DebuggerWorld);
-	
+	check(DebuggerWorld);
+
 	// @TODO: Handle editor world when those features are enabled for the Rewind Debugger
 	// Currently prevents debug draw remnants from stopped world
 	if (DebuggerWorld->WorldType != EWorldType::PIE)
 	{
 		return;
 	}
-	
+
 	const bool bSameTime = FMath::Abs(TimeMarker - PreviousTimeMarker) < DOUBLE_SMALL_NUMBER;
 	PreviousTimeMarker = TimeMarker;
 
 	TSharedPtr<FDebuggerViewModel> Model = ViewModel.Get();
+	check(Model.IsValid());
 
 	// We haven't reached the update point yet
 	if (CurrentConsecutiveFrames < ConsecutiveFramesUpdateThreshold)
@@ -222,9 +224,147 @@ void SDebuggerView::Tick(const FGeometry& AllottedGeometry, const double InCurre
 		}
 	}
 
-	// Draw visualization every tick
-	DrawVisualization();
-	
+	// Draw features
+	if (const FTraceMotionMatchingStateMessage* State = Model->GetMotionMatchingState())
+	{
+		FRoleToIndex RoleToIndex;
+		TArray<const USkinnedMeshComponent*> Meshes;
+		TArray<const IPoseHistory*> PoseHistories;
+
+		const int32 NumRoles = State->Roles.Num();
+
+		RoleToIndex.Reserve(NumRoles);
+		Meshes.SetNum(NumRoles);
+		PoseHistories.SetNum(NumRoles);
+
+		for (int32 RoleIndex = 0; RoleIndex < NumRoles; ++RoleIndex)
+		{
+			const uint64 ActorSkeletalMeshComponentId = State->SkeletalMeshComponentIds[RoleIndex];
+			if (const TWeakObjectPtr<AActor>* ActorPtr = Model->GetDebugDrawActors().Find(ActorSkeletalMeshComponentId))
+			{
+				if (ActorPtr->IsValid())
+				{
+					const FRole& Role = State->Roles[RoleIndex];
+
+					for (UActorComponent* ActorComponent : (*ActorPtr)->GetInstanceComponents())
+					{
+						if (UPoseSearchMeshComponent* PoseSearchMeshComponent = Cast<UPoseSearchMeshComponent>(ActorComponent))
+						{
+							RoleToIndex.Add(Role) = RoleIndex;
+							Meshes[RoleIndex] = PoseSearchMeshComponent;
+							PoseHistories[RoleIndex] = &State->PoseHistories[RoleIndex];
+							break;
+						}
+					}
+				}
+			}
+		}
+
+		// checking if all roles have been resolved properly
+		if (RoleToIndex.Num() != NumRoles)
+		{
+			return;
+		}
+
+		// Draw world space trajectory
+#if ENABLE_ANIM_DEBUG
+		const bool bDrawTrajectory = Model->GetDrawTrajectory();
+		const bool bDrawHistory = Model->GetDrawHistory();
+
+		if (bDrawTrajectory || bDrawHistory)
+		{
+			for (int32 RoleIndex = 0; RoleIndex < NumRoles; ++RoleIndex)
+			{
+				UWorld* World = Meshes[RoleIndex]->GetWorld();
+				const FPoseSearchQueryTrajectory& Trajectory = State->PoseHistories[RoleIndex].Trajectory;
+				
+				if (bDrawTrajectory)
+				{
+					Trajectory.DebugDrawTrajectory(World);
+				}
+
+				if (bDrawHistory)
+				{
+					State->PoseHistories[RoleIndex].DebugDraw(World, FColor::Red);
+				}
+			}
+		}
+#endif
+
+		// Draw query vector
+		if (Model->GetDrawQuery())
+		{
+			if (const UPoseSearchDatabase* CurrentDatabase = Model->GetCurrentDatabase())
+			{
+				for (const FTraceMotionMatchingStateDatabaseEntry& DbEntry : State->DatabaseEntries)
+				{
+					const UPoseSearchDatabase* Database = FTraceMotionMatchingStateMessage::GetObjectFromId<UPoseSearchDatabase>(DbEntry.DatabaseId);
+					if (Database && Database == CurrentDatabase &&
+						FAsyncPoseSearchDatabasesManagement::RequestAsyncBuildIndex(CurrentDatabase, ERequestAsyncBuildFlag::ContinueRequest) &&
+						DbEntry.QueryVector.Num() == Database->Schema->SchemaCardinality)
+					{
+						FDebugDrawParams DrawParams(Meshes, PoseHistories, RoleToIndex, CurrentDatabase, EDebugDrawFlags::DrawQuery);
+						DrawParams.DrawFeatureVector(DbEntry.QueryVector);
+						break;
+					}
+				}
+			}
+		}
+
+		// Draw selected poses
+		const TSharedPtr<SListView<TSharedRef<FDebuggerDatabaseRowData>>>& DatabaseRows = DatabaseView->GetDatabaseRows();
+		TArray<TSharedRef<FDebuggerDatabaseRowData>> SelectedRows = DatabaseRows->GetSelectedItems();
+
+		// Draw any selected database vectors
+		constexpr int32 MaxRowsToDraw = 250;
+
+		const int32 NumRowsToDraw = FMath::Min(MaxRowsToDraw, SelectedRows.Num());
+		for (int32 RowIdx = 0; RowIdx < NumRowsToDraw; ++RowIdx)
+		{
+			const TSharedRef<FDebuggerDatabaseRowData>& Row = SelectedRows[RowIdx];
+			const UPoseSearchDatabase* RowDatabase = Row->SharedData->SourceDatabase.Get();
+			if (FAsyncPoseSearchDatabasesManagement::RequestAsyncBuildIndex(RowDatabase, ERequestAsyncBuildFlag::ContinueRequest))
+			{
+				FDebugDrawParams DrawParams(Meshes, PoseHistories, RoleToIndex, RowDatabase);
+				DrawParams.DrawFeatureVector(Row->PoseIdx);
+			}
+		}
+
+		// Draw active pose
+		TArray<TSharedRef<FDebuggerDatabaseRowData>> ActiveRows = DatabaseView->GetActiveRow()->GetSelectedItems();
+
+		// Active row should only have 0 or 1
+		check(ActiveRows.Num() < 2);
+
+		if (!ActiveRows.IsEmpty())
+		{
+			const UPoseSearchDatabase* Database = ActiveRows[0]->SharedData->SourceDatabase.Get();
+			if (Database && FAsyncPoseSearchDatabasesManagement::RequestAsyncBuildIndex(Database, ERequestAsyncBuildFlag::ContinueRequest))
+			{
+				// Use the motion-matching state's pose idx, as the active row may be update-throttled at this point
+				FDebugDrawParams DrawParams(Meshes, PoseHistories, RoleToIndex, Database);
+				DrawParams.DrawFeatureVector(ActiveRows[0]->PoseIdx);
+			}
+		}
+
+
+		// Draw continuing pose
+		TArray<TSharedRef<FDebuggerDatabaseRowData>> ContinuingRows = DatabaseView->GetContinuingPoseRow()->GetSelectedItems();
+
+		// ContinuingPose row should only have 0 or 1
+		check(ContinuingRows.Num() < 2);
+
+		if (!ContinuingRows.IsEmpty())
+		{
+			const UPoseSearchDatabase* Database = ContinuingRows[0]->SharedData->SourceDatabase.Get();
+			if (Database && FAsyncPoseSearchDatabasesManagement::RequestAsyncBuildIndex(Database, ERequestAsyncBuildFlag::ContinueRequest))
+			{
+				FDebugDrawParams DrawParams(Meshes, PoseHistories, RoleToIndex, Database);
+				DrawParams.DrawFeatureVector(ContinuingRows[0]->PoseIdx);
+			}
+		}
+	}
+
 	// synchronizing the model DrawQuery state with all the open PoseSearchDatabaseEditor(s)
 	const bool bDrawQuery = Model->GetDrawQuery();
 	if (UAssetEditorSubsystem* AssetEditorSS = GEditor->GetEditorSubsystem<UAssetEditorSubsystem>())
@@ -250,6 +390,7 @@ void SDebuggerView::Tick(const FGeometry& AllottedGeometry, const double InCurre
 bool SDebuggerView::UpdateNodeSelection()
 {
 	TSharedPtr<FDebuggerViewModel> Model = ViewModel.Get();
+	check(Model.IsValid());
 
 	// Update selection view if no node selected
 	bool bNodeSelected = SelectedNodeId != INDEX_NONE;
@@ -297,108 +438,11 @@ void SDebuggerView::UpdateViews() const
 	}
 }
 
-void SDebuggerView::DrawVisualization() const
-{
-	if (const FTraceMotionMatchingStateMessage* State = ViewModel.Get()->GetMotionMatchingState())
-	{
-		const UWorld* DebuggerWorld = FDebugger::GetWorld();
-		check(DebuggerWorld);
-
-		DrawFeatures(*DebuggerWorld, *State, ViewModel.Get()->GetRootBoneTransform(), ViewModel.Get()->GetMeshComponent());
-	}
-}
-
 TArray<TSharedRef<FDebuggerDatabaseRowData>> SDebuggerView::GetSelectedDatabaseRows() const
 {
 	return DatabaseView->GetDatabaseRows()->GetSelectedItems();
 }
 
-void SDebuggerView::DrawFeatures(const UWorld& DebuggerWorld, const FTraceMotionMatchingStateMessage& State, const FTransform& RootBoneWorldTransform, const USkinnedMeshComponent* Mesh, int32 MaxRowsToDraw) const
-{
-	// Draw world space trajectory
-#if ENABLE_ANIM_DEBUG
-	if (ViewModel.Get()->GetDrawTrajectory())
-	{
-		State.Trajectory.DebugDrawTrajectory(&DebuggerWorld);
-	}
-#endif
-
-	// Draw query vector
-	if (ViewModel.Get()->GetDrawQuery())
-	{
-		if (const UPoseSearchDatabase* CurrentDatabase = ViewModel.Get()->GetCurrentDatabase())
-		{
-			for (const FTraceMotionMatchingStateDatabaseEntry& DbEntry : State.DatabaseEntries)
-			{
-				const UPoseSearchDatabase* Database = FTraceMotionMatchingState::GetObjectFromId<UPoseSearchDatabase>(DbEntry.DatabaseId);
-				if (Database && Database == CurrentDatabase && 
-					FAsyncPoseSearchDatabasesManagement::RequestAsyncBuildIndex(CurrentDatabase, ERequestAsyncBuildFlag::ContinueRequest) &&
-					DbEntry.QueryVector.Num() == Database->Schema->SchemaCardinality)
-				{
-					FDebugDrawParams DrawParams(&DebuggerWorld, Mesh, RootBoneWorldTransform, CurrentDatabase, EDebugDrawFlags::DrawQuery);
-					DrawParams.DrawFeatureVector(DbEntry.QueryVector);
-					break;
-				}
-			}
-		}
-	}
-
-	// Draw selected poses
-	const TSharedPtr<SListView<TSharedRef<FDebuggerDatabaseRowData>>>& DatabaseRows = DatabaseView->GetDatabaseRows();
-	TArray<TSharedRef<FDebuggerDatabaseRowData>> SelectedRows = DatabaseRows->GetSelectedItems();
-	
-	// Draw any selected database vectors
-	const int32 NumRowsToDraw = FMath::Min(MaxRowsToDraw, SelectedRows.Num());
-	for (int32 RowIdx = 0; RowIdx < NumRowsToDraw; ++RowIdx)
-	{
-		const TSharedRef<FDebuggerDatabaseRowData>& Row = SelectedRows[RowIdx];
-		const UPoseSearchDatabase* RowDatabase = Row->SharedData->SourceDatabase.Get();
-		if (FAsyncPoseSearchDatabasesManagement::RequestAsyncBuildIndex(RowDatabase, ERequestAsyncBuildFlag::ContinueRequest))
-		{
-			FDebugDrawParams DrawParams(&DebuggerWorld, Mesh, RootBoneWorldTransform, RowDatabase);
-			DrawParams.DrawFeatureVector(Row->PoseIdx);
-		}
-	}
-
-	// Draw active pose
-	{
-		TArray<TSharedRef<FDebuggerDatabaseRowData>> ActiveRows = DatabaseView->GetActiveRow()->GetSelectedItems();
-
-		// Active row should only have 0 or 1
-		check(ActiveRows.Num() < 2);
-
-		if (!ActiveRows.IsEmpty())
-		{
-			const UPoseSearchDatabase* Database = ActiveRows[0]->SharedData->SourceDatabase.Get();
-			if (Database && FAsyncPoseSearchDatabasesManagement::RequestAsyncBuildIndex(Database, ERequestAsyncBuildFlag::ContinueRequest))
-			{
-				// Use the motion-matching state's pose idx, as the active row may be update-throttled at this point
-				FDebugDrawParams DrawParams(&DebuggerWorld, Mesh, RootBoneWorldTransform, Database);
-				DrawParams.DrawFeatureVector(ActiveRows[0]->PoseIdx);
-			}
-		}
-	}
-
-
-	// Draw continuing pose
-	{
-		TArray<TSharedRef<FDebuggerDatabaseRowData>> ContinuingRows = DatabaseView->GetContinuingPoseRow()->GetSelectedItems();
-
-		// ContinuingPose row should only have 0 or 1
-		check(ContinuingRows.Num() < 2);
-
-		if (!ContinuingRows.IsEmpty())
-		{
-			const UPoseSearchDatabase* Database = ContinuingRows[0]->SharedData->SourceDatabase.Get();
-			if (Database && FAsyncPoseSearchDatabasesManagement::RequestAsyncBuildIndex(Database, ERequestAsyncBuildFlag::ContinueRequest))
-			{
-				FDebugDrawParams DrawParams(&DebuggerWorld, Mesh, RootBoneWorldTransform, Database);
-				DrawParams.DrawFeatureVector(ContinuingRows[0]->PoseIdx);
-			}
-		}
-	}
-}
-
 int32 SDebuggerView::SelectView() const
 {
 	// Currently recording
@@ -414,6 +458,7 @@ int32 SDebuggerView::SelectView() const
 	}
 
 	const TSharedPtr<FDebuggerViewModel> Model = ViewModel.Get();
+	check(Model.IsValid());
 
 	const bool bNoActiveNodes = Model->GetNodesNum() == 0;
 	const bool bNodeSelectedWithoutData = SelectedNodeId != INDEX_NONE && Model->GetMotionMatchingState() == nullptr;
@@ -437,6 +482,8 @@ int32 SDebuggerView::SelectView() const
 void SDebuggerView::OnPoseSelectionChanged(const UPoseSearchDatabase* Database, int32 DbPoseIdx, float Time)
 {
 	const TSharedPtr<FDebuggerViewModel> Model = ViewModel.Get();
+	check(Model.IsValid());
+
 	if (const FTraceMotionMatchingStateMessage* State = Model->GetMotionMatchingState())
 	{
 		DetailsView->Update(*State);
@@ -569,6 +616,33 @@ TSharedRef<SHorizontalBox> SDebuggerView::GenerateReturnButtonView()
 				]
 			]
 		]
+		+SHorizontalBox::Slot()
+		.VAlign(VAlign_Top)
+		.HAlign(HAlign_Left)
+		.Padding(64, 5, 0, 0)
+		.AutoWidth()
+		[
+			SNew(SHorizontalBox)
+
+			+ SHorizontalBox::Slot()
+			.AutoWidth()
+			.Padding(0, 5, 0, 0)
+			[
+				SNew(SCheckBox)
+				.IsChecked_Lambda([this]
+				{
+					return ViewModel.Get()->GetDrawHistory() ? ECheckBoxState::Checked : ECheckBoxState::Unchecked; 
+				})
+				.OnCheckStateChanged_Lambda([this](ECheckBoxState State)
+				{
+					ViewModel.Get()->SetDrawHistory(State == ECheckBoxState::Checked);
+				})
+				[
+					SNew(STextBlock)
+					.Text(LOCTEXT("PoseSearchDebuggerDrawHistory", "Draw History"))
+				]
+			]
+		]
 		+ SHorizontalBox::Slot()
 		.VAlign(VAlign_Top)
 		.HAlign(HAlign_Left)
diff --git a/Engine/Plugins/Experimental/Animation/PoseSearch/Source/Editor/Private/PoseSearchDebuggerView.h b/Engine/Plugins/Experimental/Animation/PoseSearch/Source/Editor/Private/PoseSearchDebuggerView.h
--- a/Engine/Plugins/Experimental/Animation/PoseSearch/Source/Editor/Private/PoseSearchDebuggerView.h
+++ b/Engine/Plugins/Experimental/Animation/PoseSearch/Source/Editor/Private/PoseSearchDebuggerView.h
@@ -78,17 +78,12 @@ public:
 	TArray<TSharedRef<FDebuggerDatabaseRowData>> GetSelectedDatabaseRows() const;
 
 private:
-	/** Called each frame to draw features of the query vector & database selections */
-	void DrawFeatures(const UWorld& DebuggerWorld, const FTraceMotionMatchingStateMessage& State, const FTransform& RootBoneWorldTransform, const USkinnedMeshComponent* Mesh, int32 MaxRowsToDraw = 250) const;
-	
 	/** Check if a node selection was made, true if a node is selected */
 	bool UpdateNodeSelection();
 
 	/** Update the database and details views */
 	void UpdateViews() const;
 
-	void DrawVisualization() const;
-	
 	/** Returns an int32 appropriate to the index of our widget selector */
 	int32 SelectView() const;
 
diff --git a/Engine/Plugins/Experimental/Animation/PoseSearch/Source/Editor/Private/PoseSearchDebuggerViewModel.cpp b/Engine/Plugins/Experimental/Animation/PoseSearch/Source/Editor/Private/PoseSearchDebuggerViewModel.cpp
--- a/Engine/Plugins/Experimental/Animation/PoseSearch/Source/Editor/Private/PoseSearchDebuggerViewModel.cpp
+++ b/Engine/Plugins/Experimental/Animation/PoseSearch/Source/Editor/Private/PoseSearchDebuggerViewModel.cpp
@@ -24,9 +24,12 @@ FDebuggerViewModel::FDebuggerViewModel(uint64 InAnimInstanceId)
 
 FDebuggerViewModel::~FDebuggerViewModel()
 {
-	if (DebugDrawActor != nullptr)
+	for (TPair<uint64, TWeakObjectPtr<AActor>>& DebugDrawActorPair : DebugDrawActors)
 	{
-		DebugDrawActor->Destroy();
+		if (DebugDrawActorPair.Value.IsValid())
+		{
+			DebugDrawActorPair.Value->Destroy();
+		}
 	}
 }
 
@@ -64,51 +67,9 @@ const FTransform& FDebuggerViewModel::GetRootBoneTransform() const
 }
 
 void FDebuggerViewModel::OnUpdate()
-{
-	if (!bSkeletonInitialized)
-	{
-		UWorld* World = RewindDebugger.Get()->GetWorldToVisualize();
-		FActorSpawnParameters ActorSpawnParameters;
-		ActorSpawnParameters.bHideFromSceneOutliner = false;
-		ActorSpawnParameters.ObjectFlags |= RF_Transient;
-		DebugDrawActor = World->SpawnActor<AActor>(ActorSpawnParameters);
-		DebugDrawActor->SetActorLabel(TEXT("PoseSearch"));
-		DebugDrawMeshComponent = NewObject<UPoseSearchMeshComponent>(DebugDrawActor.Get());
-		DebugDrawActor->AddInstanceComponent(DebugDrawMeshComponent.Get());
-		DebugDrawMeshComponent->RegisterComponentWithWorld(World);
-	
-		FWorldDelegates::OnWorldCleanup.AddRaw(this, &FDebuggerViewModel::OnWorldCleanup);
-		bSkeletonInitialized = true;
-	}
-
-	UpdateFromTimeline();
-}
-
-void FDebuggerViewModel::OnUpdateNodeSelection(int32 InNodeId)
-{
-	if (InNodeId == INDEX_NONE)
-	{
-		return;
-	}
-
-	// Find node in all motion matching states this frame
-	ActiveMotionMatchingStateIdx = INDEX_NONE;
-	const int32 NodesNum = NodeIds.Num();
-	for (int32 i = 0; i < NodesNum; ++i)
-	{
-		if (NodeIds[i] == InNodeId)
-		{
-			ActiveMotionMatchingStateIdx = i;
-			break;
-		}
-	}
-}
-
-void FDebuggerViewModel::UpdateFromTimeline()
 {
 	NodeIds.Empty();
 	MotionMatchingStates.Empty();
-	SkeletalMeshComponentId = 0;
 
 	// Get provider and validate
 	const TraceServices::IAnalysisSession* Session = RewindDebugger.Get()->GetAnalysisSession();
@@ -124,7 +85,7 @@ void FDebuggerViewModel::UpdateFromTimeline()
 	const double TraceTime = RewindDebugger.Get()->CurrentTraceTime();
 	TraceServices::FFrame Frame;
 	ReadFrameProvider(*Session).GetFrameFromTime(TraceFrameType_Game, TraceTime, Frame);
-	PoseSearchProvider->EnumerateMotionMatchingStateTimelines(AnimInstanceId, [&](const FTraceProvider::FMotionMatchingStateTimeline& InTimeline)
+	PoseSearchProvider->EnumerateMotionMatchingStateTimelines(AnimInstanceId, [&Frame, this](const FTraceProvider::FMotionMatchingStateTimeline& InTimeline)
 	{
 		const FTraceMotionMatchingStateMessage* Message = nullptr;
 
@@ -135,66 +96,138 @@ void FDebuggerViewModel::UpdateFromTimeline()
 		});
 		if (Message)
 		{
-			NodeIds.Add(Message->NodeId);
+			check(Message->Roles.Num() == Message->SkeletalMeshComponentIds.Num());
 
-			// @todo: figure out if we can avoid this copy
+			NodeIds.Add(Message->NodeId);
 			MotionMatchingStates.Add(*Message);
-			SkeletalMeshComponentId = Message->SkeletalMeshComponentId;
 		}
 	});
-	/** No active motion matching state as no messages were read */
-	if (SkeletalMeshComponentId == 0)
-	{
-		return;
-	}
-	AnimationProvider->ReadSkeletalMeshPoseTimeline(SkeletalMeshComponentId, [&](const IAnimationProvider::SkeletalMeshPoseTimeline& TimelineData, bool bHasCurves)
+	
+	// adding all the missing DebugDrawActors, and tracking them via UsedDebugDrawActorIds
+	TSet<uint64> UsedDebugDrawActorIds;
+	for (const FTraceMotionMatchingStateMessage& MotionMatchingState : MotionMatchingStates)
 	{
-		TimelineData.EnumerateEvents(Frame.StartTime, Frame.EndTime, [&](double InStartTime, double InEndTime, uint32 InDepth, const FSkeletalMeshPoseMessage& PoseMessage) -> TraceServices::EEventEnumerate
+		for (uint64 SkeletalMeshComponentId : MotionMatchingState.SkeletalMeshComponentIds)
 		{
-			const FSkeletalMeshInfo* SkeletalMeshInfo = AnimationProvider->FindSkeletalMeshInfo(PoseMessage.MeshId);
-			const FObjectInfo* SkeletalMeshObjectInfo = GameplayProvider->FindObjectInfo(PoseMessage.MeshId);
-			if (!SkeletalMeshInfo || !SkeletalMeshObjectInfo)
-			{
-				return TraceServices::EEventEnumerate::Stop;
-			}
+			UsedDebugDrawActorIds.Add(SkeletalMeshComponentId);
 
-			UPoseSearchMeshComponent* PoseSearchMeshComponent = DebugDrawMeshComponent.Get();
-			USkeletalMesh* SkeletalMesh = TSoftObjectPtr<USkeletalMesh>(FSoftObjectPath(SkeletalMeshObjectInfo->PathName)).LoadSynchronous();
-			if (SkeletalMesh)
+			bool bNeedToAdd = true;
+			if (TWeakObjectPtr<AActor>* DebugDrawActor = DebugDrawActors.Find(SkeletalMeshComponentId))
 			{
-				PoseSearchMeshComponent->SetSkinnedAssetAndUpdate(SkeletalMesh, true);
+				if (DebugDrawActor->IsValid())
+				{
+					bNeedToAdd = false;
+				}
 			}
-			FTransform ComponentWorldTransform;
-			// Active skeleton is simply the traced bone transforms
-			TArray<FTransform>& ComponentSpaceTransforms = PoseSearchMeshComponent->GetEditableComponentSpaceTransforms();
-			AnimationProvider->GetSkeletalMeshComponentSpacePose(PoseMessage, *SkeletalMeshInfo, ComponentWorldTransform, ComponentSpaceTransforms);
-
-			check(ComponentWorldTransform.Equals(PoseMessage.ComponentToWorld));
 
-			if (!ComponentSpaceTransforms.IsEmpty())
+			if (bNeedToAdd)
 			{
-				RootBoneWorldTransform = ComponentSpaceTransforms[RootBoneIndexType] * ComponentWorldTransform;
+				UWorld* World = RewindDebugger.Get()->GetWorldToVisualize();
+				FActorSpawnParameters ActorSpawnParameters;
+				ActorSpawnParameters.bHideFromSceneOutliner = false;
+				ActorSpawnParameters.ObjectFlags |= RF_Transient;
+				AActor* DebugDrawActor = World->SpawnActor<AActor>(ActorSpawnParameters);
+				DebugDrawActor->SetActorLabel(TEXT("PoseSearch"));
+				UPoseSearchMeshComponent* DebugDrawMeshComponent = NewObject<UPoseSearchMeshComponent>(DebugDrawActor);
+				DebugDrawActor->AddInstanceComponent(DebugDrawMeshComponent);
+				DebugDrawMeshComponent->RegisterComponentWithWorld(World);
+
+				DebugDrawActors.Add(SkeletalMeshComponentId) = DebugDrawActor;
 			}
-			else
+		}
+	}
+
+	// cleaning up the no longer used actors (with a SkeletalMeshComponentId not in UsedDebugDrawActorIds)
+	for (auto DebugDrawActorsIt = DebugDrawActors.CreateIterator(); DebugDrawActorsIt; ++DebugDrawActorsIt)
+	{
+		if (!UsedDebugDrawActorIds.Find(DebugDrawActorsIt.Key()))
+		{
+			if (DebugDrawActorsIt.Value().IsValid())
 			{
-				RootBoneWorldTransform = ComponentWorldTransform;
+				DebugDrawActorsIt.Value()->Destroy();
 			}
+			DebugDrawActorsIt.RemoveCurrent();
+		}
+	}
 
-			PoseSearchMeshComponent->Initialize(ComponentWorldTransform);
+	/** No active motion matching state as no messages were read */
+	if (MotionMatchingStates.IsEmpty())
+	{
+		return;
+	}
 
-			return TraceServices::EEventEnumerate::Stop;
-		});
-	});
+	for (TPair<uint64, TWeakObjectPtr<AActor>>& DebugDrawActorPair : DebugDrawActors)
+	{
+		const uint64 SkeletalMeshComponentId = DebugDrawActorPair.Key;
+		if (AActor* DebugDrawActor = DebugDrawActorPair.Value.Get())
+		{
+			for (UActorComponent* ActorComponent : DebugDrawActor->GetInstanceComponents())
+			{
+				if (UPoseSearchMeshComponent* PoseSearchMeshComponent = Cast<UPoseSearchMeshComponent>(ActorComponent))
+				{
+					AnimationProvider->ReadSkeletalMeshPoseTimeline(SkeletalMeshComponentId, [&Frame, AnimationProvider, GameplayProvider, PoseSearchMeshComponent](const IAnimationProvider::SkeletalMeshPoseTimeline& TimelineData, bool bHasCurves)
+						{
+							TimelineData.EnumerateEvents(Frame.StartTime, Frame.EndTime, [AnimationProvider, GameplayProvider, PoseSearchMeshComponent](double InStartTime, double InEndTime, uint32 InDepth, const FSkeletalMeshPoseMessage& PoseMessage) -> TraceServices::EEventEnumerate
+								{
+									const FSkeletalMeshInfo* SkeletalMeshInfo = AnimationProvider->FindSkeletalMeshInfo(PoseMessage.MeshId);
+									const FObjectInfo* SkeletalMeshObjectInfo = GameplayProvider->FindObjectInfo(PoseMessage.MeshId);
+									if (!SkeletalMeshInfo || !SkeletalMeshObjectInfo)
+									{
+										return TraceServices::EEventEnumerate::Stop;
+									}
+
+									USkeletalMesh* SkeletalMesh = TSoftObjectPtr<USkeletalMesh>(FSoftObjectPath(SkeletalMeshObjectInfo->PathName)).LoadSynchronous();
+									if (SkeletalMesh)
+									{
+										PoseSearchMeshComponent->SetSkinnedAssetAndUpdate(SkeletalMesh, true);
+									}
+									FTransform ComponentWorldTransform;
+									// Active skeleton is simply the traced bone transforms
+									TArray<FTransform>& ComponentSpaceTransforms = PoseSearchMeshComponent->GetEditableComponentSpaceTransforms();
+									AnimationProvider->GetSkeletalMeshComponentSpacePose(PoseMessage, *SkeletalMeshInfo, ComponentWorldTransform, ComponentSpaceTransforms);
+
+									check(ComponentWorldTransform.Equals(PoseMessage.ComponentToWorld));
+
+									// @todo: remove RootBoneWorldTransform and use Trajectoris instead
+									//if (!ComponentSpaceTransforms.IsEmpty())
+									//{
+									//	RootBoneWorldTransform = ComponentSpaceTransforms[RootBoneIndexType] * ComponentWorldTransform;
+									//}
+									//else
+									//{
+									//	RootBoneWorldTransform = ComponentWorldTransform;
+									//}
+
+									PoseSearchMeshComponent->Initialize(ComponentWorldTransform);
+
+									return TraceServices::EEventEnumerate::Stop;
+								});
+						});
+					break;
+				}
+			}
+		}
+	}
 }
 
-const USkinnedMeshComponent* FDebuggerViewModel::GetMeshComponent() const
+void FDebuggerViewModel::OnUpdateNodeSelection(int32 InNodeId)
 {
-	return DebugDrawMeshComponent.Get();
-}
+	if (InNodeId == INDEX_NONE)
+	{
+		return;
+	}
 
-void FDebuggerViewModel::OnWorldCleanup(UWorld* InWorld, bool bSessionEnded, bool bCleanupResources)
-{
-	bSkeletonInitialized = false;
+	// Find node in all motion matching states this frame
+	ActiveMotionMatchingStateIdx = INDEX_NONE;
+	const int32 NodesNum = NodeIds.Num();
+	for (int32 i = 0; i < NodesNum; ++i)
+	{
+		if (NodeIds[i] == InNodeId)
+		{
+			ActiveMotionMatchingStateIdx = i;
+			break;
+		}
+	}
 }
 
 } // namespace UE::PoseSearch
diff --git a/Engine/Plugins/Experimental/Animation/PoseSearch/Source/Editor/Private/PoseSearchDebuggerViewModel.h b/Engine/Plugins/Experimental/Animation/PoseSearch/Source/Editor/Private/PoseSearchDebuggerViewModel.h
--- a/Engine/Plugins/Experimental/Animation/PoseSearch/Source/Editor/Private/PoseSearchDebuggerViewModel.h
+++ b/Engine/Plugins/Experimental/Animation/PoseSearch/Source/Editor/Private/PoseSearchDebuggerViewModel.h
@@ -26,6 +26,8 @@ public:
 
 	// Used for view callbacks
     const FTraceMotionMatchingStateMessage* GetMotionMatchingState() const;
+	const TMap<uint64, TWeakObjectPtr<AActor>>& GetDebugDrawActors() const { return DebugDrawActors; }
+
 	const UPoseSearchDatabase* GetCurrentDatabase() const;
 	const TArray<int32>* GetNodeIds() const;
 	int32 GetNodesNum() const;
@@ -44,18 +46,14 @@ public:
 	bool GetDrawQuery() const { return bDrawQuery; }
 
 	void SetDrawTrajectory(bool bInDrawTrajectory) { bDrawTrajectory = bInDrawTrajectory; }
-
 	bool GetDrawTrajectory() const { return bDrawTrajectory; }
 
-	/** Callback to reset debug skeletons for the active world */
-	void OnWorldCleanup(UWorld* InWorld, bool bSessionEnded, bool bCleanupResources);
-
-	const USkinnedMeshComponent* GetMeshComponent() const;
+	void SetDrawHistory(bool bInDrawHistory) { bDrawHistory = bInDrawHistory; }
+	bool GetDrawHistory() const { return bDrawHistory; }
 
 private:
-	/** Update the list of states for this frame */
-	void UpdateFromTimeline();
 
+	// @todo: duplicate member, already contained in MotionMatchingStates(FTraceMessage)::NodeId
 	/** List of all Node IDs associated with motion matching states */
 	TArray<int32> NodeIds;
 	
@@ -65,9 +63,7 @@ private:
 	/** Currently active motion matching state index based on node selection in the view */
 	int32 ActiveMotionMatchingStateIdx = INDEX_NONE;
 
-	/** Current Skeletal Mesh Component Id for the AnimInstance */
-	uint64 SkeletalMeshComponentId = 0;
-
+	// @todo: duplicate information already present in MotionMatchingStates::Trajectories
 	/** Currently active root bone transform */
 	FTransform RootBoneWorldTransform = FTransform::Identity;
 
@@ -77,20 +73,13 @@ private:
 	/** Anim Instance associated with this debugger instance */
 	uint64 AnimInstanceId = 0;
 
-	/** Actor object for the skeleton */
-	TWeakObjectPtr<AActor> DebugDrawActor;
+	/** Actor object populated by the timeline mapped from the MotionMatchingStates.SkeletalMeshComponentIds. used as input for additional bone transforms to draw channels */
+	TMap<uint64, TWeakObjectPtr<AActor>> DebugDrawActors;
 
-	/** Derived skeletal mesh for setting the skeleton in the scene */
-	TWeakObjectPtr<UPoseSearchMeshComponent> DebugDrawMeshComponent;
-
-	/** Whether the skeleton have been initialized for this world */
-	bool bSkeletonInitialized = false;
-	
 	bool bIsVerbose = false;
-
 	bool bDrawQuery = true;	
-
 	bool bDrawTrajectory = false;
+	bool bDrawHistory = false;
 	
 	/** Limits some public API */
 	friend class FDebugger;
diff --git a/Engine/Plugins/Experimental/Animation/PoseSearch/Source/Editor/Private/PoseSearchSchemaFactory.cpp b/Engine/Plugins/Experimental/Animation/PoseSearch/Source/Editor/Private/PoseSearchSchemaFactory.cpp
--- a/Engine/Plugins/Experimental/Animation/PoseSearch/Source/Editor/Private/PoseSearchSchemaFactory.cpp
+++ b/Engine/Plugins/Experimental/Animation/PoseSearch/Source/Editor/Private/PoseSearchSchemaFactory.cpp
@@ -65,7 +65,7 @@ UObject* UPoseSearchSchemaFactory::FactoryCreateNew(UClass* Class, UObject* InPa
 	if (TargetSkeleton)
 	{
 		Schema = NewObject<UPoseSearchSchema>(InParent, Class, Name, Flags);
-		Schema->Skeleton = TargetSkeleton;
+		Schema->AddSkeleton(TargetSkeleton);
 		Schema->AddDefaultChannels();
 	}
 	return Schema;
diff --git a/Engine/Plugins/Experimental/Animation/PoseSearch/Source/Editor/Private/SPoseSearchDatabaseViewport.cpp b/Engine/Plugins/Experimental/Animation/PoseSearch/Source/Editor/Private/SPoseSearchDatabaseViewport.cpp
--- a/Engine/Plugins/Experimental/Animation/PoseSearch/Source/Editor/Private/SPoseSearchDatabaseViewport.cpp
+++ b/Engine/Plugins/Experimental/Animation/PoseSearch/Source/Editor/Private/SPoseSearchDatabaseViewport.cpp
@@ -10,6 +10,7 @@
 
 #include "SSimpleTimeSlider.h"
 #include "Viewports.h"
+#include "Widgets/Images/SImage.h"
 #include "Widgets/Input/SButton.h"
 
 namespace UE::PoseSearch
@@ -158,21 +159,26 @@ namespace UE::PoseSearch
 			FMultiBoxCustomization::None, 
 			nullptr, true);
 
-		auto AddToolBarButton = [&ToolBarBuilder](FName ButtonStyleName, FOnButtonClickedEvent& OnClicked)
-		{
-			ToolBarBuilder.AddToolBarWidget(
-				SNew(SButton)
-				.ButtonStyle(FAppStyle::Get(), ButtonStyleName)
-				.OnClicked_Lambda([&OnClicked]() 
-				{
-					if (OnClicked.IsBound())
-					{
-						OnClicked.Execute();
-						return FReply::Handled();
-					}
-					return FReply::Unhandled();
-				}));
-		};
+		auto AddToolBarButton = [&ToolBarBuilder](FName ButtonImageName, FOnButtonClickedEvent& OnClicked)
+			{
+				ToolBarBuilder.AddToolBarWidget(
+					SNew(SButton)
+					.ButtonStyle(FAppStyle::Get(), "Animation.PlayControlsButton")
+					.OnClicked_Lambda([&OnClicked]()
+						{
+							if (OnClicked.IsBound())
+							{
+								OnClicked.Execute();
+								return FReply::Handled();
+							}
+							return FReply::Unhandled();
+						})
+					[
+						SNew(SImage)
+							.ColorAndOpacity(FSlateColor::UseSubduedForeground())
+							.Image(FAppStyle::Get().GetBrush(ButtonImageName))
+					]);
+			};
 
 		//ToolBarBuilder.SetStyle(&FAppStyle::Get(), "PaletteToolBar");
 		ToolBarBuilder.BeginSection("Preview");
diff --git a/Engine/Plugins/Experimental/Animation/PoseSearch/Source/Editor/Private/Trace/PoseSearchTraceAnalyzer.cpp b/Engine/Plugins/Experimental/Animation/PoseSearch/Source/Editor/Private/Trace/PoseSearchTraceAnalyzer.cpp
--- a/Engine/Plugins/Experimental/Animation/PoseSearch/Source/Editor/Private/Trace/PoseSearchTraceAnalyzer.cpp
+++ b/Engine/Plugins/Experimental/Animation/PoseSearch/Source/Editor/Private/Trace/PoseSearchTraceAnalyzer.cpp
@@ -24,7 +24,7 @@ void FTraceAnalyzer::OnAnalysisBegin(const FOnAnalysisContext& Context)
 	ANSICHAR MotionMatchingStateName[NAME_SIZE];
 
 	FTraceLogger::Name.GetPlainANSIString(LoggerName);
-	FTraceMotionMatchingState::Name.GetPlainANSIString(MotionMatchingStateName);
+	FTraceMotionMatchingStateMessage::Name.GetPlainANSIString(MotionMatchingStateName);
 
 	Builder.RouteEvent(RouteId_MotionMatchingState, LoggerName, MotionMatchingStateName);
 }
diff --git a/Engine/Plugins/Experimental/Animation/PoseSearch/Source/Editor/Private/Trace/PoseSearchTraceProvider.cpp b/Engine/Plugins/Experimental/Animation/PoseSearch/Source/Editor/Private/Trace/PoseSearchTraceProvider.cpp
--- a/Engine/Plugins/Experimental/Animation/PoseSearch/Source/Editor/Private/Trace/PoseSearchTraceProvider.cpp
+++ b/Engine/Plugins/Experimental/Animation/PoseSearch/Source/Editor/Private/Trace/PoseSearchTraceProvider.cpp
@@ -39,11 +39,4 @@ void FTraceProvider::AppendMotionMatchingState(const FTraceMotionMatchingStateMe
 	Session.UpdateDurationSeconds(InTime);
 }
 
-FArchive& operator<<(FArchive& Ar, FTraceMotionMatchingStateMessage& State)
-{
-	Ar << static_cast<FTraceMessage&>(State);
-	Ar << static_cast<FTraceMotionMatchingState&>(State);
-	return Ar;
-}
-
 } // namespace UE::PoseSearch
diff --git a/Engine/Plugins/Experimental/Animation/PoseSearch/Source/Editor/Private/Trace/PoseSearchTraceProvider.h b/Engine/Plugins/Experimental/Animation/PoseSearch/Source/Editor/Private/Trace/PoseSearchTraceProvider.h
--- a/Engine/Plugins/Experimental/Animation/PoseSearch/Source/Editor/Private/Trace/PoseSearchTraceProvider.h
+++ b/Engine/Plugins/Experimental/Animation/PoseSearch/Source/Editor/Private/Trace/PoseSearchTraceProvider.h
@@ -10,12 +10,6 @@
 namespace UE::PoseSearch
 {
 
-/** Motion matching state message container */
-struct FTraceMotionMatchingStateMessage : FTraceMotionMatchingState, FTraceMessage
-{
-};
-FArchive& operator<<(FArchive& Ar, FTraceMotionMatchingStateMessage& State);
-
 /**
  * Provider to the widgets for pose search functionality, largely mimicking FAnimationProvider
  */
diff --git a/Engine/Plugins/Experimental/Animation/PoseSearch/Source/Editor/Public/PoseSearchDatabaseAssetTree.h b/Engine/Plugins/Experimental/Animation/PoseSearch/Source/Editor/Public/PoseSearchDatabaseAssetTree.h
--- a/Engine/Plugins/Experimental/Animation/PoseSearch/Source/Editor/Public/PoseSearchDatabaseAssetTree.h
+++ b/Engine/Plugins/Experimental/Animation/PoseSearch/Source/Editor/Public/PoseSearchDatabaseAssetTree.h
@@ -89,6 +89,7 @@ namespace UE::PoseSearch
 		void OnAddBlendSpace(bool bFinalizeChanges = true);
 		void OnAddAnimComposite(bool bFinalizeChanges = true);
 		void OnAddAnimMontage(bool bFinalizeChanges = true);
+		void OnAddMultiSequence(bool bFinalizeChanges = true);
 
 		void OnDeleteAsset(TSharedPtr<FDatabaseAssetTreeNode> Node, bool bFinalizeChanges = true);
 		void CreateCommandList();
diff --git a/Engine/Plugins/Experimental/Animation/PoseSearch/Source/Editor/Public/PoseSearchDatabaseEditorReflection.h b/Engine/Plugins/Experimental/Animation/PoseSearch/Source/Editor/Public/PoseSearchDatabaseEditorReflection.h
--- a/Engine/Plugins/Experimental/Animation/PoseSearch/Source/Editor/Public/PoseSearchDatabaseEditorReflection.h
+++ b/Engine/Plugins/Experimental/Animation/PoseSearch/Source/Editor/Public/PoseSearchDatabaseEditorReflection.h
@@ -145,6 +145,35 @@ public:
 #endif
 };
 
+USTRUCT()
+struct FPoseSearchDatabaseMultiSequenceEx : public FPoseSearchDatabaseMultiSequence
+{
+	GENERATED_BODY()
+
+	// Is this animation set as looping? If you want to change this you need to change it in the base Anim Sequence.
+	// Changing the sampling range will disable looping
+	UPROPERTY(VisibleAnywhere, Category="Sequence", meta = (DisplayPriority = 10))
+	bool bLooping = false;
+
+	// Does this animation have root motion enabled ? If you want to change this you need to change it in the base Anim Sequence.
+	UPROPERTY(VisibleAnywhere, Category="Sequence", meta = (DisplayPriority = 11))
+	bool bHasRootMotion = false;
+};
+
+UCLASS()
+class UPoseSearchDatabaseMultiSequenceReflection : public UPoseSearchDatabaseReflectionBase
+{
+	GENERATED_BODY()
+
+public:
+	UPROPERTY(EditAnywhere, Category = "Selected Sequence")
+	FPoseSearchDatabaseMultiSequenceEx MultiSequence;
+
+#if WITH_EDITOR
+	virtual void PostEditChangeProperty(struct FPropertyChangedEvent& PropertyChangedEvent) override;
+#endif
+};
+
 USTRUCT()
 struct FPoseSearchDatabaseMemoryStats
 {
diff --git a/Engine/Plugins/Experimental/Animation/PoseSearch/Source/Editor/Public/PoseSearchDatabaseViewModel.h b/Engine/Plugins/Experimental/Animation/PoseSearch/Source/Editor/Public/PoseSearchDatabaseViewModel.h
--- a/Engine/Plugins/Experimental/Animation/PoseSearch/Source/Editor/Public/PoseSearchDatabaseViewModel.h
+++ b/Engine/Plugins/Experimental/Animation/PoseSearch/Source/Editor/Public/PoseSearchDatabaseViewModel.h
@@ -7,6 +7,7 @@
 #include "MovieSceneFwd.h"
 #include "PoseSearch/PoseSearchAssetSampler.h"
 #include "PoseSearch/PoseSearchMirrorDataCache.h"
+#include "PoseSearch/PoseSearchRole.h"
 #include "PoseSearchDatabasePreviewScene.h"
 #include "UObject/GCObject.h"
 
@@ -28,9 +29,10 @@ namespace UE::PoseSearch
 	struct FDatabasePreviewActor
 	{
 	public:
-		bool SpawnPreviewActor(UWorld* World, const UPoseSearchDatabase* PoseSearchDatabase, int32 IndexAssetIdx, int32 PoseIdxForTimeOffset = INDEX_NONE);
+		bool SpawnPreviewActor(UWorld* World, const UPoseSearchDatabase* PoseSearchDatabase, int32 IndexAssetIdx, const FRole& Role, const FTransform& SamplerRootTransformOrigin, const FTransform* PrecalculatedRootTransformOrigin, int32 PoseIdxForTimeOffset = INDEX_NONE);
 		void UpdatePreviewActor(const UPoseSearchDatabase* PoseSearchDatabase, float PlayTime, bool bQuantizeAnimationToPoseData);
-		bool DrawPreviewActor(const UPoseSearchDatabase* PoseSearchDatabase, bool bDisplayRootMotionSpeed, bool bDisplayBlockTransition, TConstArrayView<float> QueryVector);
+		static bool DrawPreviewActors(TArrayView<FDatabasePreviewActor> PreviewActors, const UPoseSearchDatabase* PoseSearchDatabase, bool bDisplayRootMotionSpeed, bool bDisplayBlockTransition, TConstArrayView<float> QueryVector);
+
 		void Destroy();
 
 		const UDebugSkelMeshComponent* GetDebugSkelMeshComponent() const;
@@ -41,6 +43,7 @@ namespace UE::PoseSearch
 		int32 GetIndexAssetIndex() const { return IndexAssetIndex; }
 		int32 GetCurrentPoseIndex() const { return CurrentPoseIndex; }
 		float GetPlayTimeOffset() const { return PlayTimeOffset; }
+		const FTransform& GetRootTransformOrigin() const { return RootTransformOrigin;  }
 		
 	private:
 		UAnimPreviewInstance* GetAnimPreviewInstanceInternal();
@@ -63,6 +66,8 @@ namespace UE::PoseSearch
 
 		TArray<FVector> SampledRootMotion;
 		TArray<float> SampledRootMotionSpeed;
+
+		FRole ActorRole = DefaultRole;
 	};
 
 	class FDatabaseViewModel : public TSharedFromThis<FDatabaseViewModel>, public FGCObject
@@ -91,8 +96,8 @@ namespace UE::PoseSearch
 
 		void Tick(float DeltaSeconds);
 
-		TArray<FDatabasePreviewActor>& GetPreviewActors() { return PreviewActors; }
-		const TArray<FDatabasePreviewActor>& GetPreviewActors() const { return PreviewActors; }
+		const TArray<TArray<FDatabasePreviewActor>>& GetPreviewActors() const { return PreviewActors; }
+		TArray<TArray<FDatabasePreviewActor>>& GetPreviewActors() { return PreviewActors; }
 
 		void ToggleDisplayRootMotionSpeed() { bDisplayRootMotionSpeed = !bDisplayRootMotionSpeed;	}
 		bool IsDisplayRootMotionSpeedChecked() const { return bDisplayRootMotionSpeed; }
@@ -110,6 +115,8 @@ namespace UE::PoseSearch
 		void AddBlendSpaceToDatabase(UBlendSpace* BlendSpace);
 		void AddAnimCompositeToDatabase(UAnimComposite* AnimComposite);
 		void AddAnimMontageToDatabase(UAnimMontage* AnimMontage);
+		void AddMultiSequenceToDatabase();
+
 		bool DeleteFromDatabase(int32 AnimationAssetIndex);
 
 		void SetDisableReselection(int32 AnimationAssetIndex, bool bEnabled);
@@ -151,7 +158,7 @@ namespace UE::PoseSearch
 		TWeakPtr<SDatabaseDataDetails> DatabaseDataDetails;
 
 		/** Actors to be displayed in the preview viewport */
-		TArray<FDatabasePreviewActor> PreviewActors;
+		TArray<TArray<FDatabasePreviewActor>> PreviewActors;
 		
 		/** From zero to the play length of the longest preview */
 		float MaxPreviewPlayLength = 0.f;
