motion matching - multicharacter WIP

diff --git a/Engine/Plugins/Experimental/Animation/PoseSearch/Source/Runtime/Private/AnimNode_PoseSearchHistoryCollector.cpp b/Engine/Plugins/Experimental/Animation/PoseSearch/Source/Runtime/Private/AnimNode_PoseSearchHistoryCollector.cpp
--- a/Engine/Plugins/Experimental/Animation/PoseSearch/Source/Runtime/Private/AnimNode_PoseSearchHistoryCollector.cpp
+++ b/Engine/Plugins/Experimental/Animation/PoseSearch/Source/Runtime/Private/AnimNode_PoseSearchHistoryCollector.cpp
@@ -37,6 +37,13 @@ private:
 /////////////////////////////////////////////////////
 // FAnimNode_PoseSearchHistoryCollector_Base
 
+void FAnimNode_PoseSearchHistoryCollector_Base::Initialize_AnyThread(const FAnimationInitializeContext& Context)
+{
+	Super::Initialize_AnyThread(Context);
+
+	PoseHistory.Initialize_AnyThread(PoseCount, SamplingInterval);
+}
+
 void FAnimNode_PoseSearchHistoryCollector_Base::CacheBones_AnyThread(const FAnimationCacheBonesContext& Context)
 {
 	DECLARE_SCOPE_HIERARCHICAL_COUNTER_ANIMNODE(Update_AnyThread);
@@ -66,7 +73,7 @@ void FAnimNode_PoseSearchHistoryCollector_Base::CacheBones_AnyThread(const FAnim
 		}
 	}
 
-	PoseHistory.CacheBones_AnyThread(PoseCount, SamplingInterval, RequiredBones);
+	PoseHistory.CacheBones_AnyThread(RequiredBones);
 
 	if (bInitializeWithRefPose)
 	{
@@ -76,7 +83,7 @@ void FAnimNode_PoseSearchHistoryCollector_Base::CacheBones_AnyThread(const FAnim
 		Pose.SetBoneContainer(&Context.AnimInstanceProxy->GetRequiredBones());
 		FCSPose<FCompactPose> ComponentSpacePose;
 		ComponentSpacePose.InitPose(Pose);
-		PoseHistory.EvaluateComponentSpace_AnyThread(0.f, ComponentSpacePose, bGenerateTrajectory, bStoreScales);
+		PoseHistory.EvaluateComponentSpace_AnyThread(0.f, ComponentSpacePose, bStoreScales);
 	}
 }
 
@@ -84,7 +91,17 @@ void FAnimNode_PoseSearchHistoryCollector_Base::Update_AnyThread(const FAnimatio
 {
 	GetEvaluateGraphExposedInputs().Execute(Context);
 
-	const bool bNeedsReset = bResetOnBecomingRelevant && UpdateCounter.HasEverBeenUpdated() && !UpdateCounter.WasSynchronizedCounter(Context.AnimInstanceProxy->GetUpdateCounter());
+	PoseHistory.Update_AnyThread(Context.GetDeltaTime(), bGenerateTrajectory ? FPoseSearchQueryTrajectory() : Trajectory, TrajectorySpeedMultiplier);
+
+	UpdateCounter.SynchronizeWith(Context.AnimInstanceProxy->GetUpdateCounter());
+}
+
+void FAnimNode_PoseSearchHistoryCollector_Base::PreUpdate(const UAnimInstance* InAnimInstance)
+{
+	Super::PreUpdate(InAnimInstance);
+	
+	const bool bNeedsReset = bResetOnBecomingRelevant && UpdateCounter.HasEverBeenUpdated() && !UpdateCounter.WasSynchronizedCounter(InAnimInstance->GetUpdateCounter());
+	const float DeltaTime = InAnimInstance->GetDeltaSeconds();
 
 	FPoseSearchTrajectoryData::FSampling TrajectoryDataSampling;
 	TrajectoryDataSampling.NumHistorySamples = FMath::Max(PoseCount, TrajectoryHistoryCount);
@@ -92,9 +109,7 @@ void FAnimNode_PoseSearchHistoryCollector_Base::Update_AnyThread(const FAnimatio
 	TrajectoryDataSampling.NumPredictionSamples = TrajectoryPredictionCount;
 	TrajectoryDataSampling.SecondsPerPredictionSample = PredictionSamplingInterval;
 
-	PoseHistory.Update_AnyThread(Context.GetDeltaTime(), Trajectory, TrajectorySpeedMultiplier, bGenerateTrajectory, *Context.AnimInstanceProxy, TrajectoryData, TrajectoryDataSampling, bNeedsReset);
-
-	UpdateCounter.SynchronizeWith(Context.AnimInstanceProxy->GetUpdateCounter());
+	PoseHistory.PreUpdate(InAnimInstance, DeltaTime, bGenerateTrajectory, TrajectoryData, TrajectoryDataSampling, bNeedsReset);
 }
 
 /////////////////////////////////////////////////////
@@ -126,14 +141,14 @@ void FAnimNode_PoseSearchHistoryCollector::Evaluate_AnyThread(FPoseContext& Outp
 
 	FCSPose<FCompactPose> ComponentSpacePose;
 	ComponentSpacePose.InitPose(Output.Pose);
-	PoseHistory.EvaluateComponentSpace_AnyThread(Output.AnimInstanceProxy->GetDeltaSeconds(), ComponentSpacePose, bGenerateTrajectory, bStoreScales);
+	PoseHistory.EvaluateComponentSpace_AnyThread(Output.AnimInstanceProxy->GetDeltaSeconds(), ComponentSpacePose, bStoreScales);
 
 #if ENABLE_DRAW_DEBUG && ENABLE_ANIM_DEBUG
-		FColor Color;
+	FColor Color;
 #if WITH_EDITORONLY_DATA
-		Color = DebugColor.ToFColor(true);
+	Color = DebugColor.ToFColor(true);
 #else // WITH_EDITORONLY_DATA
-		Color = FLinearColor::Red.ToFColor(true);
+	Color = FLinearColor::Red.ToFColor(true);
 #endif // WITH_EDITORONLY_DATA
 	PoseHistory.DebugDraw(*Output.AnimInstanceProxy, Color);
 #endif // ENABLE_DRAW_DEBUG && ENABLE_ANIM_DEBUG
@@ -180,14 +195,14 @@ void FAnimNode_PoseSearchComponentSpaceHistoryCollector::EvaluateComponentSpace_
 	Super::EvaluateComponentSpace_AnyThread(Output);
 	Source.EvaluateComponentSpace(Output);
 
-	PoseHistory.EvaluateComponentSpace_AnyThread(Output.AnimInstanceProxy->GetDeltaSeconds(), Output.Pose, bGenerateTrajectory, bStoreScales);
+	PoseHistory.EvaluateComponentSpace_AnyThread(Output.AnimInstanceProxy->GetDeltaSeconds(), Output.Pose, bStoreScales);
 
 #if ENABLE_DRAW_DEBUG && ENABLE_ANIM_DEBUG
-		FColor Color;
+	FColor Color;
 #if WITH_EDITORONLY_DATA
-		Color = DebugColor.ToFColor(true);
+	Color = DebugColor.ToFColor(true);
 #else // WITH_EDITORONLY_DATA
-		Color = FLinearColor::Red.ToFColor(true);
+	Color = FLinearColor::Red.ToFColor(true);
 #endif // WITH_EDITORONLY_DATA
 	PoseHistory.DebugDraw(*Output.AnimInstanceProxy, Color);
 #endif // ENABLE_DRAW_DEBUG && ENABLE_ANIM_DEBUG
diff --git a/Engine/Plugins/Experimental/Animation/PoseSearch/Source/Runtime/Private/PoseSearchAssetIndexer.cpp b/Engine/Plugins/Experimental/Animation/PoseSearch/Source/Runtime/Private/PoseSearchAssetIndexer.cpp
--- a/Engine/Plugins/Experimental/Animation/PoseSearch/Source/Runtime/Private/PoseSearchAssetIndexer.cpp
+++ b/Engine/Plugins/Experimental/Animation/PoseSearch/Source/Runtime/Private/PoseSearchAssetIndexer.cpp
@@ -81,26 +81,140 @@ static FSamplingParam WrapOrClampSamplingParam(bool bCanWrap, float SamplingPara
 
 //////////////////////////////////////////////////////////////////////////
 // FAssetSamplingContext
-FAssetSamplingContext::FAssetSamplingContext(const UPoseSearchDatabase& Database, const FBoneContainer& BoneContainer)
-: FMirrorDataCache(Database.Schema ? Database.Schema->MirrorDataTable : nullptr, BoneContainer)
+FAssetSamplingContext::FAssetSamplingContext(const UPoseSearchDatabase& Database)
 {
 	check(Database.Schema);
 	BaseCostBias = Database.BaseCostBias;
 	LoopingCostBias = Database.LoopingCostBias;
 }
 
+//////////////////////////////////////////////////////////////////////////
+// FAnimationAssetSamplers
+void FAnimationAssetSamplers::Reset()
+{
+	AnimationAssetSamplers.Reset();
+	MirrorDataCaches.Reset();
+}
+
+int32 FAnimationAssetSamplers::Num() const
+{
+	check(AnimationAssetSamplers.Num() == MirrorDataCaches.Num());
+	return AnimationAssetSamplers.Num();
+}
+
+float FAnimationAssetSamplers::GetPlayLength() const
+{
+	float PlayLength = 0.f;
+	for (const FAnimationAssetSampler* Sampler : AnimationAssetSamplers)
+	{
+		PlayLength = FMath::Max(PlayLength, Sampler->GetPlayLength());
+	}
+	return PlayLength;
+}
+
+bool FAnimationAssetSamplers::IsLoopable() const
+{
+	float CommonPlayLength = -1.f;
+	for (const FAnimationAssetSampler* Sampler : AnimationAssetSamplers)
+	{
+		if (!Sampler->IsLoopable())
+		{
+			return false;
+		}
+
+		if (CommonPlayLength < 0.f)
+		{
+			CommonPlayLength = Sampler->GetPlayLength();
+		}
+		else if (!FMath::IsNearlyEqual(CommonPlayLength, Sampler->GetPlayLength()))
+		{
+			return false;
+		}
+	}
+	return true;
+}
+
+void FAnimationAssetSamplers::ExtractPoseSearchNotifyStates(float Time, TFunction<bool(UAnimNotifyState_PoseSearchBase*)> ProcessPoseSearchBase) const
+{
+	for (const FAnimationAssetSampler* Sampler : AnimationAssetSamplers)
+	{
+		Sampler->ExtractPoseSearchNotifyStates(Time, ProcessPoseSearchBase);
+	}
+}
+
+bool FAnimationAssetSamplers::ProcessAllAnimNotifyEvents(TFunction<bool(TConstArrayView<FAnimNotifyEvent>)> ProcessAnimNotifyEvents) const
+{
+	for (const FAnimationAssetSampler* Sampler : AnimationAssetSamplers)
+	{
+		if (ProcessAnimNotifyEvents(Sampler->GetAllAnimNotifyEvents()))
+		{
+			return true;
+		}
+	}
+	return false;
+}
+
+const FString FAnimationAssetSamplers::GetAssetName() const
+{
+	FString Name;
+	Name.Reserve(256);
+	bool bAddComma = false;
+	for (const FAnimationAssetSampler* Sampler : AnimationAssetSamplers)
+	{
+		if (bAddComma)
+		{
+			Name += ", ";
+		}
+		else
+		{
+			bAddComma = true;
+		}
+
+		Name += GetNameSafe(Sampler->GetAsset());
+	}
+	return Name;
+}
+
+FTransform FAnimationAssetSamplers::ExtractRootTransform(float Time, int32 RoleIndex) const
+{
+	return AnimationAssetSamplers[RoleIndex]->ExtractRootTransform(Time);
+}
+
+FTransform FAnimationAssetSamplers::GetTotalRootTransform(int32 RoleIndex) const
+{
+	return AnimationAssetSamplers[RoleIndex]->GetTotalRootTransform();
+}
+
+void FAnimationAssetSamplers::ExtractPose(float Time, FCompactPose& OutPose, int32 RoleIndex) const
+{
+	AnimationAssetSamplers[RoleIndex]->ExtractPose(Time, OutPose);
+}
+
+FTransform FAnimationAssetSamplers::MirrorTransform(const FTransform& InTransform, int32 RoleIndex) const
+{
+	return MirrorDataCaches[RoleIndex]->MirrorTransform(InTransform);
+}
+
+void FAnimationAssetSamplers::MirrorPose(FCompactPose& Pose, int32 RoleIndex) const
+{
+	MirrorDataCaches[RoleIndex]->MirrorPose(Pose);
+}
+
 //////////////////////////////////////////////////////////////////////////
 // FAssetIndexer
-FAssetIndexer::FAssetIndexer(const FBoneContainer& InBoneContainer, const FSearchIndexAsset& InSearchIndexAsset, const FAssetSamplingContext& InSamplingContext,
-	const UPoseSearchSchema& InSchema, const FAnimationAssetSampler& InAssetSampler, const FFloatInterval& InExtrapolationTimeInterval)
-: BoneContainer(InBoneContainer)
+FAssetIndexer::FAssetIndexer(const TConstArrayView<FBoneContainer> InBoneContainers, const FSearchIndexAsset& InSearchIndexAsset, const FAssetSamplingContext& InSamplingContext,
+	const UPoseSearchSchema& InSchema, const FAnimationAssetSamplers& InAssetSamplers, const FRoleToIndex& InRoleToIndex, const FFloatInterval& InExtrapolationTimeInterval)
+: BoneContainers(InBoneContainers)
 , CachedEntries()
 , SearchIndexAsset(InSearchIndexAsset)
 , SamplingContext(InSamplingContext)
 , Schema(InSchema)
-, AssetSampler(InAssetSampler)
+, AssetSamplers(InAssetSamplers)
+, RoleToIndex(InRoleToIndex)
 , ExtrapolationTimeInterval(InExtrapolationTimeInterval)
 {
+	check(BoneContainers.Num() == AssetSamplers.Num() && BoneContainers.Num() == RoleToIndex.Num());
+	check(IsValid(RoleToIndex));
 }
 
 void FAssetIndexer::AssignWorkingData(int32 InStartPoseIdx, TArrayView<float> InOutFeatureVectorTable, TArrayView<FPoseMetadata> InOutPoseMetadata)
@@ -114,19 +228,17 @@ void FAssetIndexer::AssignWorkingData(int32 InStartPoseIdx, TArrayView<float> In
 
 void FAssetIndexer::Process(int32 AssetIdx)
 {
-	check(Schema.IsValid());
-
 	bProcessFailed = false;
 
 	// Generate pose metadata
-	const float SequenceLength = AssetSampler.GetPlayLength();
+	const float PlayLength = AssetSamplers.GetPlayLength();
 	for (int32 SampleIdx = GetBeginSampleIdx(); SampleIdx != GetEndSampleIdx(); ++SampleIdx)
 	{
-		const float SampleTime = FMath::Min(CalculateSampleTime(SampleIdx), SequenceLength);
+		const float SampleTime = FMath::Min(CalculateSampleTime(SampleIdx), PlayLength);
 		float CostAddend = SamplingContext.BaseCostBias;
 		bool bBlockTransition = false;
 
-		AssetSampler.ExtractPoseSearchNotifyStates(SampleTime, [&bBlockTransition, &CostAddend](const UAnimNotifyState_PoseSearchBase* PoseSearchNotify)
+		AssetSamplers.ExtractPoseSearchNotifyStates(SampleTime, [&bBlockTransition, &CostAddend](const UAnimNotifyState_PoseSearchBase* PoseSearchNotify)
 			{
 				if (PoseSearchNotify->GetClass()->IsChildOf<UAnimNotifyState_PoseSearchBlockTransition>())
 				{
@@ -139,7 +251,7 @@ void FAssetIndexer::Process(int32 AssetIdx)
 				return true;
 			});
 
-		if (AssetSampler.IsLoopable())
+		if (AssetSamplers.IsLoopable())
 		{
 			CostAddend += SamplingContext.LoopingCostBias;
 		}
@@ -173,47 +285,52 @@ void FAssetIndexer::ComputeStats()
 {
 	Stats = FStats();
 
-	for (int32 SampleIdx = GetBeginSampleIdx(); SampleIdx != GetEndSampleIdx(); ++SampleIdx)
+	for (const FRoleToIndexPair& RoleToIndexPair : RoleToIndex)
 	{
-		const float SampleTime = FMath::Min(CalculateSampleTime(SampleIdx), AssetSampler.GetPlayLength());
+		const FRole& Role = RoleToIndexPair.Key;
+		const int32 RoleIndex = RoleToIndexPair.Value;
+		const FBoneReference& RootBoneReference = Schema.GetBoneReferences(Role)[RootSchemaBoneIdx];
 
-		bool AnyClamped = false;
-		const FTransform TrajTransformsPast = GetTransform(SampleTime - FiniteDelta, AnyClamped);
-		if (!AnyClamped)
+		for (int32 SampleIdx = GetBeginSampleIdx(); SampleIdx != GetEndSampleIdx(); ++SampleIdx)
 		{
-			const FTransform TrajTransformsPresent = GetTransform(SampleTime, AnyClamped);
+			const float SampleTime = FMath::Min(CalculateSampleTime(SampleIdx), AssetSamplers.GetPlayLength());
+
+			bool AnyClamped = false;
+			const FTransform TrajTransformsPast = GetTransform(SampleTime - FiniteDelta, RoleIndex, AnyClamped, RootBoneReference);
 			if (!AnyClamped)
 			{
-				const FTransform TrajTransformsFuture = GetTransform(SampleTime + FiniteDelta, AnyClamped);
+				const FTransform TrajTransformsPresent = GetTransform(SampleTime, RoleIndex, AnyClamped, RootBoneReference);
 				if (!AnyClamped)
 				{
-					// if any transform is clamped we just skip the sample entirely
-					const FVector LinearVelocityPresent = (TrajTransformsPresent.GetTranslation() - TrajTransformsPast.GetTranslation()) / FiniteDelta;
-					const FVector LinearVelocityFuture = (TrajTransformsFuture.GetTranslation() - TrajTransformsPresent.GetTranslation()) / FiniteDelta;
-					const FVector LinearAcceleration = (LinearVelocityFuture - LinearVelocityPresent) / FiniteDelta;
+					const FTransform TrajTransformsFuture = GetTransform(SampleTime + FiniteDelta, RoleIndex, AnyClamped, RootBoneReference);
+					if (!AnyClamped)
+					{
+						// if any transform is clamped we just skip the sample entirely
+						const FVector LinearVelocityPresent = (TrajTransformsPresent.GetTranslation() - TrajTransformsPast.GetTranslation()) / FiniteDelta;
+						const FVector LinearVelocityFuture = (TrajTransformsFuture.GetTranslation() - TrajTransformsPresent.GetTranslation()) / FiniteDelta;
+						const FVector LinearAcceleration = (LinearVelocityFuture - LinearVelocityPresent) / FiniteDelta;
 
-					const float Speed = LinearVelocityPresent.Length();
-					const float Acceleration = LinearAcceleration.Length();
+						const float Speed = LinearVelocityPresent.Length();
+						const float Acceleration = LinearAcceleration.Length();
 
-					Stats.AccumulatedSpeed += Speed;
-					Stats.MaxSpeed = FMath::Max(Stats.MaxSpeed, Speed);
+						Stats.AccumulatedSpeed += Speed;
+						Stats.MaxSpeed = FMath::Max(Stats.MaxSpeed, Speed);
 
-					Stats.AccumulatedAcceleration += Acceleration;
-					Stats.MaxAcceleration = FMath::Max(Stats.MaxAcceleration, Acceleration);
+						Stats.AccumulatedAcceleration += Acceleration;
+						Stats.MaxAcceleration = FMath::Max(Stats.MaxAcceleration, Acceleration);
 
-					++Stats.NumAccumulatedSamples;
+						++Stats.NumAccumulatedSamples;
+					}
 				}
 			}
 		}
 	}
 }
 
-FAssetIndexer::FSampleInfo FAssetIndexer::GetSampleInfo(float SampleTime) const
+void FAssetIndexer::GetSampleInfo(float SampleTime, int32 RoleIndex, FTransform& OutRootTransform, float& OutClipTime, bool& bOutClamped) const
 {
-	FSampleInfo Sample;
-
-	const float PlayLength = AssetSampler.GetPlayLength();
-	const bool bCanWrap = AssetSampler.IsLoopable();
+	const float PlayLength = AssetSamplers.GetPlayLength();
+	const bool bCanWrap = AssetSamplers.IsLoopable();
 
 	float MainRelativeTime = SampleTime;
 	if (SampleTime < 0.0f && bCanWrap)
@@ -227,22 +344,23 @@ FAssetIndexer::FSampleInfo FAssetIndexer::GetSampleInfo(float SampleTime) const
 
 	if (FMath::Abs(SamplingParam.Extrapolation) > SMALL_NUMBER)
 	{
-		Sample.bClamped = true;
-		Sample.ClipTime = SamplingParam.WrappedParam + SamplingParam.Extrapolation;
-		Sample.RootTransform = AssetSampler.ExtractRootTransform(Sample.ClipTime);
+		bOutClamped = true;
+		OutClipTime = SamplingParam.WrappedParam + SamplingParam.Extrapolation;
+		OutRootTransform = AssetSamplers.ExtractRootTransform(OutClipTime, RoleIndex);
 	}
 	else
 	{
-		Sample.ClipTime = SamplingParam.WrappedParam;
-		Sample.RootTransform = FTransform::Identity;
+		bOutClamped = false;
+		OutClipTime = SamplingParam.WrappedParam;
+		OutRootTransform = FTransform::Identity;
 
 		// Find the remaining motion deltas after wrapping
-		FTransform RootMotionRemainder = AssetSampler.ExtractRootTransform(Sample.ClipTime);
+		FTransform RootMotionRemainder = AssetSamplers.ExtractRootTransform(OutClipTime, RoleIndex);
 
-		const bool bNegativeSampleTime = SampleTime < 0.f;
+  		const bool bNegativeSampleTime = SampleTime < 0.f;
 		if (SamplingParam.NumCycles > 0 || bNegativeSampleTime)
 		{
-			const FTransform RootMotionLast = AssetSampler.GetTotalRootTransform();
+			const FTransform RootMotionLast = AssetSamplers.GetTotalRootTransform(RoleIndex);
 
 			// Determine how to accumulate motion for every cycle of the anim. If the sample
 			// had to be clamped, this motion will end up not getting applied below.
@@ -264,22 +382,20 @@ FAssetIndexer::FSampleInfo FAssetIndexer::GetSampleInfo(float SampleTime) const
 			int32 CyclesRemaining = SamplingParam.NumCycles;
 			while (CyclesRemaining--)
 			{
-				Sample.RootTransform = RootMotionPerCycle * Sample.RootTransform;
+				OutRootTransform = RootMotionPerCycle * OutRootTransform;
 			}
 		}
 
-		Sample.RootTransform = RootMotionRemainder * Sample.RootTransform;
+		OutRootTransform = RootMotionRemainder * OutRootTransform;
 	}
-
-	return Sample;
 }
 
-FTransform FAssetIndexer::MirrorTransform(const FTransform& Transform) const
+FTransform FAssetIndexer::MirrorTransform(const FTransform& Transform, int32 RoleIndex) const
 {
-	return SearchIndexAsset.IsMirrored() ? SamplingContext.MirrorTransform(Transform) : Transform;
+	return SearchIndexAsset.IsMirrored() ? AssetSamplers.MirrorTransform(Transform, RoleIndex) : Transform;
 }
 
-FAssetIndexer::CachedEntry& FAssetIndexer::GetEntry(float SampleTime)
+FAssetIndexer::FCachedEntry& FAssetIndexer::GetEntry(float SampleTime)
 {
 	using namespace UE::Anim;
 	
@@ -290,122 +406,129 @@ FAssetIndexer::CachedEntry& FAssetIndexer::GetEntry(float SampleTime)
 
 	SampleTime = FMath::Clamp(SampleTime, ExtrapolationTimeInterval.Min, ExtrapolationTimeInterval.Max);
 
-	CachedEntry* Entry = bDisableCaching ? nullptr : CachedEntries.Find(SampleTime);
+	FCachedEntry* Entry = bDisableCaching ? nullptr : CachedEntries.Find(SampleTime);
 	if (!Entry)
 	{
 		Entry = &CachedEntries.Add(SampleTime);
 		Entry->SampleTime = SampleTime;
 
-		if (!BoneContainer.IsValid())
-		{
-			UE_LOG(LogPoseSearch,
-				Warning,
-				TEXT("Invalid BoneContainer encountered in FAssetIndexer::GetEntry. Asset: %s. Schema: %s. BoneContainerAsset: %s. NumBoneIndices: %d"),
-				*GetNameSafe(AssetSampler.GetAsset()),
-				*GetNameSafe(&Schema),
-				*GetNameSafe(BoneContainer.GetAsset()),
-				BoneContainer.GetCompactPoseNumBones());
-		}
+		const bool bLoopable = AssetSamplers.IsLoopable();
+		const float PlayLength = AssetSamplers.GetPlayLength();
+		const int32 AssetSamplersNum = AssetSamplers.Num();
 
-		const FAssetIndexer::FSampleInfo Sample = GetSampleInfo(SampleTime);
-		float CurrentTime = Sample.ClipTime;
-
-		const bool bLoopable = AssetSampler.IsLoopable();
-		const float PlayLength = AssetSampler.GetPlayLength();
-
-		if (!bLoopable)
+		Entry->RootTransform.SetNum(AssetSamplersNum);
+		Entry->ComponentSpacePose.SetNum(AssetSamplersNum);
+		for (int32 RoleIndex = 0; RoleIndex < AssetSamplers.AnimationAssetSamplers.Num(); ++RoleIndex)
 		{
-			CurrentTime = FMath::Clamp(CurrentTime, 0.f, PlayLength);
-		}
+			if (!BoneContainers[RoleIndex].IsValid())
+			{
+				UE_LOG(LogPoseSearch,
+					Warning,
+					TEXT("Invalid BoneContainer encountered in FAssetIndexer::GetEntry. Asset: %s. Schema: %s. BoneContainerAsset: %s. NumBoneIndices: %d"),
+					*AssetSamplers.GetAssetName(),
+					*GetNameSafe(&Schema),
+					*GetNameSafe(BoneContainers[RoleIndex].GetAsset()),
+					BoneContainers[RoleIndex].GetCompactPoseNumBones());
+			}
 
-		FMemMark Mark(FMemStack::Get());
-		FCompactPose Pose;
-		Pose.SetBoneContainer(&BoneContainer);
-		AssetSampler.ExtractPose(CurrentTime, Pose);
+			FTransform SampleRootTransform;
+			bool bSampleClamped;
+			float CurrentTime;
 
-#if ENABLE_ANIM_DEBUG
-		const int32 NumIterations = CVarMotionMatchTestExtractPoseDeterminismNumIterations.GetValueOnAnyThread();
-		for (int32 IterationIndex = 0; IterationIndex < NumIterations; ++IterationIndex)
-		{
-			FCompactPose TestPose;
-			TestPose.SetBoneContainer(&BoneContainer);
-			AssetSampler.ExtractPose(CurrentTime, TestPose);
+			GetSampleInfo(SampleTime, RoleIndex, SampleRootTransform, CurrentTime, bSampleClamped);
 
-			const TConstArrayView<FTransform> Bones = Pose.GetBones();
-			const TConstArrayView<FTransform> TestBones = TestPose.GetBones();
-			if (Bones.Num() != TestBones.Num())
+			if (!bLoopable)
 			{
-				UE_LOG(LogPoseSearch, Warning, TEXT("FAssetIndexer::GetEntry - ExtractPose is not deterministic"));
+				CurrentTime = FMath::Clamp(CurrentTime, 0.f, PlayLength);
 			}
-			else
+
+			FMemMark Mark(FMemStack::Get());
+			FCompactPose Pose;
+			Pose.SetBoneContainer(&BoneContainers[RoleIndex]);
+			AssetSamplers.ExtractPose(CurrentTime, Pose, RoleIndex);
+
+#if ENABLE_ANIM_DEBUG
+			const int32 NumIterations = CVarMotionMatchTestExtractPoseDeterminismNumIterations.GetValueOnAnyThread();
+			for (int32 IterationIndex = 0; IterationIndex < NumIterations; ++IterationIndex)
 			{
-				for (int32 BoneIndex = 0; BoneIndex < Bones.Num(); ++BoneIndex)
+				FCompactPose TestPose;
+				TestPose.SetBoneContainer(&BoneContainers[RoleIndex]);
+				AssetSamplers.ExtractPose(CurrentTime, TestPose, RoleIndex);
+
+				const TConstArrayView<FTransform> Bones = Pose.GetBones();
+				const TConstArrayView<FTransform> TestBones = TestPose.GetBones();
+				if (Bones.Num() != TestBones.Num())
+				{
+					UE_LOG(LogPoseSearch, Warning, TEXT("FAssetIndexer::GetEntry - ExtractPose is not deterministic"));
+				}
+				else
 				{
-					if (FMemory::Memcmp(&Bones[BoneIndex], &TestBones[BoneIndex], sizeof(FTransform)) != 0)
+					for (int32 BoneIndex = 0; BoneIndex < Bones.Num(); ++BoneIndex)
 					{
-						UE_LOG(LogPoseSearch, Warning, TEXT("FAssetIndexer::GetEntry - ExtractPose is not deterministic"));
+						if (FMemory::Memcmp(&Bones[BoneIndex], &TestBones[BoneIndex], sizeof(FTransform)) != 0)
+						{
+							UE_LOG(LogPoseSearch, Warning, TEXT("FAssetIndexer::GetEntry - ExtractPose is not deterministic"));
+						}
 					}
 				}
 			}
-		}
 #endif // ENABLE_ANIM_DEBUG
 
-		if (SearchIndexAsset.IsMirrored())
-		{
-			SamplingContext.MirrorPose(Pose);
-		}
+			if (SearchIndexAsset.IsMirrored())
+			{
+				AssetSamplers.MirrorPose(Pose, RoleIndex);
+			}
 
-		FCSPose<FCompactPose> StackComponentSpacePose;
-		StackComponentSpacePose.InitPose(MoveTemp(Pose));
-		Entry->ComponentSpacePose.CopyPose(StackComponentSpacePose);
+			FCSPose<FCompactPose> StackComponentSpacePose;
+			StackComponentSpacePose.InitPose(MoveTemp(Pose));
+			Entry->ComponentSpacePose[RoleIndex].CopyPose(StackComponentSpacePose);
 
-		Entry->RootTransform = Sample.RootTransform;
-		Entry->bClamped = Sample.bClamped;
+			Entry->RootTransform[RoleIndex] = SampleRootTransform;
+			Entry->bClamped |= bSampleClamped;
+		}
 	}
 
 	return *Entry;
 }
 
 // returns the transform in component space for the bone indexed by Schema->BoneReferences[SchemaBoneIdx] at SampleTime seconds
-FTransform FAssetIndexer::GetComponentSpaceTransform(float SampleTime, bool& bClamped, int8 SchemaBoneIdx)
+FTransform FAssetIndexer::GetComponentSpaceTransform(float SampleTime, const FRole& Role, bool& bClamped, int8 SchemaBoneIdx)
 {
 	using namespace UE::PoseSearch;
 
-	CachedEntry& Entry = GetEntry(SampleTime);
+	FCachedEntry& Entry = GetEntry(SampleTime);
 	bClamped = Entry.bClamped;
 
-	return CalculateComponentSpaceTransform(Entry, SchemaBoneIdx);
+	const int32 RoleIndex = RoleToIndex[Role];
+	const FBoneReference& BoneReference = Schema.GetBoneReferences(Role)[SchemaBoneIdx];
+	return CalculateComponentSpaceTransform(Entry, BoneReference, RoleIndex);
 }
 
 // returns the transform in animation space for the bone indexed by Schema->BoneReferences[SchemaBoneIdx] at SampleTime seconds
-FTransform FAssetIndexer::GetTransform(float SampleTime, bool& bClamped, int8 SchemaBoneIdx)
+FTransform FAssetIndexer::GetTransform(float SampleTime, const FRole& Role, bool& bClamped, int8 SchemaBoneIdx)
 {
 	using namespace UE::PoseSearch;
 
-	CachedEntry& Entry = GetEntry(SampleTime);
+	FCachedEntry& Entry = GetEntry(SampleTime);
 	bClamped = Entry.bClamped;
 
-	return CalculateComponentSpaceTransform(Entry, SchemaBoneIdx) * MirrorTransform(Entry.RootTransform);
+	const int32 RoleIndex = RoleToIndex[Role];
+	const FBoneReference& BoneReference = Schema.GetBoneReferences(Role)[SchemaBoneIdx];
+	return CalculateComponentSpaceTransform(Entry, BoneReference, RoleIndex) * MirrorTransform(Entry.RootTransform[RoleIndex], RoleIndex);
 }
 
 // returns the transform in animation space for the BoneReference at SampleTime seconds
-FTransform FAssetIndexer::GetTransform(float SampleTime, bool& bClamped, const FBoneReference& BoneReference)
+FTransform FAssetIndexer::GetTransform(float SampleTime, int32 RoleIndex, bool& bClamped, const FBoneReference& BoneReference)
 {
-	CachedEntry& Entry = GetEntry(SampleTime);
+	FCachedEntry& Entry = GetEntry(SampleTime);
 	bClamped = Entry.bClamped;
-
-	return CalculateComponentSpaceTransform(Entry, BoneReference) * MirrorTransform(Entry.RootTransform);
+	return CalculateComponentSpaceTransform(Entry, BoneReference, RoleIndex) * MirrorTransform(Entry.RootTransform[RoleIndex], RoleIndex);
 }
 
-FTransform FAssetIndexer::CalculateComponentSpaceTransform(FAssetIndexer::CachedEntry& Entry, int8 SchemaBoneIdx)
+FTransform FAssetIndexer::CalculateComponentSpaceTransform(FAssetIndexer::FCachedEntry& Entry, const FBoneReference& BoneReference, int32 RoleIndex)
 {
-	return CalculateComponentSpaceTransform(Entry, Schema.BoneReferences[SchemaBoneIdx]);
-}
-
-FTransform FAssetIndexer::CalculateComponentSpaceTransform(FAssetIndexer::CachedEntry& Entry, const FBoneReference& BoneReference)
-{
-	const FCompactPoseBoneIndex CompactBoneIndex = BoneContainer.MakeCompactPoseIndex(FMeshPoseBoneIndex(BoneReference.BoneIndex));
-	return Entry.ComponentSpacePose.GetComponentSpaceTransform(CompactBoneIndex);
+	const FCompactPoseBoneIndex CompactBoneIndex = BoneContainers[RoleIndex].MakeCompactPoseIndex(FMeshPoseBoneIndex(BoneReference.BoneIndex));
+	return Entry.ComponentSpacePose[RoleIndex].GetComponentSpaceTransform(CompactBoneIndex);
 }
 
 float FAssetIndexer::CalculateSampleTime(int32 SampleIdx) const
@@ -413,7 +536,7 @@ float FAssetIndexer::CalculateSampleTime(int32 SampleIdx) const
 	return SampleIdx / float(Schema.SampleRate);
 }
 
-bool FAssetIndexer::GetSampleRotation(FQuat& OutSampleRotation, float SampleTimeOffset, float OriginTimeOffset, int32 SampleIdx, int8 SchemaSampleBoneIdx, int8 SchemaOriginBoneIdx, EPermutationTimeType PermutationTimeType, int32 SamplingAttributeId)
+bool FAssetIndexer::GetSampleRotation(FQuat& OutSampleRotation, float SampleTimeOffset, float OriginTimeOffset, int32 SampleIdx, int8 SchemaSampleBoneIdx, int8 SchemaOriginBoneIdx, const FRole& SampleRole, const FRole& OriginRole, EPermutationTimeType PermutationTimeType, int32 SamplingAttributeId)
 {
 	using namespace UE::PoseSearch;
 
@@ -431,61 +554,65 @@ bool FAssetIndexer::GetSampleRotation(FQuat& OutSampleRotation, float SampleTime
 		UE_LOG(LogPoseSearch,
 			Error,
 			TEXT("FAssetIndexer::GetSampleRotation: support for non root origin bones not implemented (bone: '%s', schema: '%s'"), 
-			*Schema.BoneReferences[SchemaOriginBoneIdx].BoneName.ToString(),
+			*Schema.GetBoneReferences(OriginRole)[SchemaOriginBoneIdx].BoneName.ToString(),
 			*GetNameSafe(&Schema));
 	}
 
-	bool bUnused;
 	if (SamplingAttributeId >= 0)
 	{
-		const TConstArrayView<FAnimNotifyEvent> AnimNotifyEvents = AssetSampler.GetAllAnimNotifyEvents();
-		if (!AnimNotifyEvents.IsEmpty())
-		{
-			for (const FAnimNotifyEvent& AnimNotifyEvent : AnimNotifyEvents)
+		if (AssetSamplers.ProcessAllAnimNotifyEvents([&OutSampleRotation, &SampleRole, &OriginRole, SamplingAttributeId, this, OriginTime](const TConstArrayView<FAnimNotifyEvent> AnimNotifyEvents)
 			{
-				if (const UAnimNotifyState_PoseSearchSamplingAttribute* SamplingAttribute = Cast<UAnimNotifyState_PoseSearchSamplingAttribute>(AnimNotifyEvent.NotifyStateClass))
+				for (const FAnimNotifyEvent& AnimNotifyEvent : AnimNotifyEvents)
 				{
-					if (SamplingAttribute->SamplingAttributeId == SamplingAttributeId)
+					if (const UAnimNotifyState_PoseSearchSamplingAttribute* SamplingAttribute = Cast<UAnimNotifyState_PoseSearchSamplingAttribute>(AnimNotifyEvent.NotifyStateClass))
 					{
-						if (SamplingAttribute->Bone.BoneName != NAME_None)
+						if (SamplingAttribute->SamplingAttributeId == SamplingAttributeId)
 						{
-							// @todo perhaps cache the initialized bone to speed up this method
-							FBoneReference TempBoneReference = SamplingAttribute->Bone;
-							TempBoneReference.Initialize(BoneContainer.GetSkeletonAsset());
-							if (TempBoneReference.HasValidSetup())
+							bool bUnused;
+							if (SamplingAttribute->Bone.BoneName != NAME_None)
 							{
-								const float SamplingAttributeTime = AnimNotifyEvent.GetTime();
-								const FTransform RootBoneTransform = GetTransform(OriginTime, bUnused, RootSchemaBoneIdx);
-								const FTransform SamplingAttributeBoneTransform = GetTransform(SamplingAttributeTime, bUnused, TempBoneReference);
-								OutSampleRotation = RootBoneTransform.InverseTransformRotation(SamplingAttributeBoneTransform.GetRotation());
-								return true;
+								// @todo perhaps cache the initialized bone to speed up this method
+								FBoneReference TempBoneReference = SamplingAttribute->Bone;
+								const int32 SampleRoleIndex = RoleToIndex[SampleRole];
+								TempBoneReference.Initialize(BoneContainers[SampleRoleIndex].GetSkeletonAsset());
+								if (TempBoneReference.HasValidSetup())
+								{
+									const float SamplingAttributeTime = AnimNotifyEvent.GetTime();
+									const FTransform RootBoneTransform = GetTransform(OriginTime, OriginRole, bUnused, RootSchemaBoneIdx);
+									const FTransform SamplingAttributeBoneTransform = GetTransform(SamplingAttributeTime, SampleRoleIndex, bUnused, TempBoneReference);
+									OutSampleRotation = RootBoneTransform.InverseTransformRotation(SamplingAttributeBoneTransform.GetRotation());
+									return true;
+								}
+
+								UE_LOG(LogPoseSearch, Error, TEXT("FAssetIndexer::GetSampleRotation: required UAnimNotifyState_PoseSearchSamplingAttribute in '%s' has an invalid Bone"), *AssetSamplers.GetAssetName());
+								return false;
 							}
 
-							UE_LOG(LogPoseSearch, Error, TEXT("FAssetIndexer::GetSampleRotation: required UAnimNotifyState_PoseSearchSamplingAttribute in '%s' has an invalid Bone"), *GetNameSafe(AssetSampler.GetAsset()));
-							OutSampleRotation = FQuat::Identity;
-							return false;
+							const FTransform RootBoneTransform = GetTransform(OriginTime, OriginRole, bUnused, RootSchemaBoneIdx);
+							OutSampleRotation = RootBoneTransform.InverseTransformRotation(SamplingAttribute->Rotation);
+							return true;
 						}
-						
-						const FTransform RootBoneTransform = GetTransform(OriginTime, bUnused, RootSchemaBoneIdx);
-						OutSampleRotation = RootBoneTransform.InverseTransformRotation(SamplingAttribute->Rotation);
-						return true;
 					}
 				}
-			}
+				return false;
+			}))
+		{
+			return true;
 		}
 
-		UE_LOG(LogPoseSearch, Error, TEXT("FAssetIndexer::GetSamplePositionInternal: required UAnimNotifyState_PoseSearchSamplingAttribute not found in '%s'"), *GetNameSafe(AssetSampler.GetAsset()));
+		UE_LOG(LogPoseSearch, Error, TEXT("FAssetIndexer::GetSamplePositionInternal: required UAnimNotifyState_PoseSearchSamplingAttribute not found in '%s'"), *AssetSamplers.GetAssetName());
 		OutSampleRotation = FQuat::Identity;
 		return false;
 	}
 
-	const FTransform RootBoneTransform = GetTransform(OriginTime, bUnused, RootSchemaBoneIdx);
-	const FTransform SampleBoneTransform = GetTransform(SampleTime, bUnused, SchemaSampleBoneIdx);
+	bool bUnused;
+	const FTransform RootBoneTransform = GetTransform(OriginTime, OriginRole, bUnused, RootSchemaBoneIdx);
+	const FTransform SampleBoneTransform = GetTransform(SampleTime, SampleRole, bUnused, SchemaSampleBoneIdx);
 	OutSampleRotation = RootBoneTransform.InverseTransformRotation(SampleBoneTransform.GetRotation());
 	return true;
 }
 
-bool FAssetIndexer::GetSamplePosition(FVector& OutSamplePosition, float SampleTimeOffset, float OriginTimeOffset, int32 SampleIdx, int8 SchemaSampleBoneIdx, int8 SchemaOriginBoneIdx, EPermutationTimeType PermutationTimeType, int32 SamplingAttributeId)
+bool FAssetIndexer::GetSamplePosition(FVector& OutSamplePosition, float SampleTimeOffset, float OriginTimeOffset, int32 SampleIdx, int8 SchemaSampleBoneIdx, int8 SchemaOriginBoneIdx, const FRole& SampleRole, const FRole& OriginRole, EPermutationTimeType PermutationTimeType, int32 SamplingAttributeId)
 {
 	float PermutationSampleTimeOffset = 0.f;
 	float PermutationOriginTimeOffset = 0.f;
@@ -496,70 +623,74 @@ bool FAssetIndexer::GetSamplePosition(FVector& OutSamplePosition, float SampleTi
 	const float OriginTime = Time + OriginTimeOffset + PermutationOriginTimeOffset;
 	
 	bool bUnused;
-	return GetSamplePositionInternal(OutSamplePosition, SampleTime, OriginTime, bUnused, SchemaSampleBoneIdx, SchemaOriginBoneIdx, SamplingAttributeId);
+	return GetSamplePositionInternal(OutSamplePosition, SampleTime, OriginTime, bUnused, SchemaSampleBoneIdx, SchemaOriginBoneIdx, SampleRole, OriginRole, SamplingAttributeId);
 }
 
-bool FAssetIndexer::GetSamplePositionInternal(FVector& OutSamplePosition, float SampleTime, float OriginTime, bool& bClamped, int8 SchemaSampleBoneIdx, int8 SchemaOriginBoneIdx, int32 SamplingAttributeId)
+bool FAssetIndexer::GetSamplePositionInternal(FVector& OutSamplePosition, float SampleTime, float OriginTime, bool& bClamped, int8 SchemaSampleBoneIdx, int8 SchemaOriginBoneIdx, const FRole& SampleRole, const FRole& OriginRole, int32 SamplingAttributeId)
 {
 	using namespace UE::PoseSearch;
 
-	bool bUnused;
 	if (SamplingAttributeId >= 0)
 	{
-		const TConstArrayView<FAnimNotifyEvent> AnimNotifyEvents = AssetSampler.GetAllAnimNotifyEvents();
-		if (!AnimNotifyEvents.IsEmpty())
-		{
-			for (const FAnimNotifyEvent& AnimNotifyEvent : AnimNotifyEvents)
+		if (AssetSamplers.ProcessAllAnimNotifyEvents([&OutSamplePosition, OriginTime, &bClamped, SchemaOriginBoneIdx, &SampleRole, OriginRole, SamplingAttributeId, this](const TConstArrayView<FAnimNotifyEvent> AnimNotifyEvents)
 			{
-				if (const UAnimNotifyState_PoseSearchSamplingAttribute* SamplingAttribute = Cast<UAnimNotifyState_PoseSearchSamplingAttribute>(AnimNotifyEvent.NotifyStateClass))
+				for (const FAnimNotifyEvent& AnimNotifyEvent : AnimNotifyEvents)
 				{
-					if (SamplingAttribute->SamplingAttributeId == SamplingAttributeId)
+					if (const UAnimNotifyState_PoseSearchSamplingAttribute* SamplingAttribute = Cast<UAnimNotifyState_PoseSearchSamplingAttribute>(AnimNotifyEvent.NotifyStateClass))
 					{
-						if (SamplingAttribute->Bone.BoneName != NAME_None)
+						if (SamplingAttribute->SamplingAttributeId == SamplingAttributeId)
 						{
-							// @todo perhaps cache the initialized bone to speed up this method
-							FBoneReference TempBoneReference = SamplingAttribute->Bone;
-							TempBoneReference.Initialize(BoneContainer.GetSkeletonAsset());
-							if (TempBoneReference.HasValidSetup())
+							bool bUnused;
+							if (SamplingAttribute->Bone.BoneName != NAME_None)
 							{
-								const float SamplingAttributeTime = AnimNotifyEvent.GetTime();
-								const FTransform RootBoneTransform = GetTransform(OriginTime, bUnused, RootSchemaBoneIdx);
-								const FTransform SamplingAttributeBoneTransform = GetTransform(SamplingAttributeTime, bClamped, TempBoneReference);
-								if (SchemaOriginBoneIdx == RootSchemaBoneIdx)
+								// @todo perhaps cache the initialized bone to speed up this method
+								FBoneReference TempBoneReference = SamplingAttribute->Bone;
+								const int32 SampleRoleIndex = RoleToIndex[SampleRole];
+								TempBoneReference.Initialize(BoneContainers[SampleRoleIndex].GetSkeletonAsset());
+								if (TempBoneReference.HasValidSetup())
 								{
-									OutSamplePosition = RootBoneTransform.InverseTransformPosition(SamplingAttributeBoneTransform.GetTranslation());
-								}
-								else
-								{
-									bool bOriginClamped;
-									const FTransform OriginBoneTransform = GetTransform(OriginTime, bOriginClamped, SchemaOriginBoneIdx);
-									bClamped |= bOriginClamped;
-									const FVector DeltaBoneTranslation = SamplingAttributeBoneTransform.GetTranslation() - OriginBoneTransform.GetTranslation();
-									OutSamplePosition = RootBoneTransform.InverseTransformVector(DeltaBoneTranslation);
+									const float SamplingAttributeTime = AnimNotifyEvent.GetTime();
+									const FTransform RootBoneTransform = GetTransform(OriginTime, OriginRole, bUnused, RootSchemaBoneIdx);
+									const FTransform SamplingAttributeBoneTransform = GetTransform(SamplingAttributeTime, SampleRoleIndex, bClamped, TempBoneReference);
+									if (SchemaOriginBoneIdx == RootSchemaBoneIdx)
+									{
+										OutSamplePosition = RootBoneTransform.InverseTransformPosition(SamplingAttributeBoneTransform.GetTranslation());
+									}
+									else
+									{
+										bool bOriginClamped;
+										const FTransform OriginBoneTransform = GetTransform(OriginTime, OriginRole, bOriginClamped, SchemaOriginBoneIdx);
+										bClamped |= bOriginClamped;
+										const FVector DeltaBoneTranslation = SamplingAttributeBoneTransform.GetTranslation() - OriginBoneTransform.GetTranslation();
+										OutSamplePosition = RootBoneTransform.InverseTransformVector(DeltaBoneTranslation);
+									}
+									return true;
 								}
-								return true;
+
+								UE_LOG(LogPoseSearch, Error, TEXT("FAssetIndexer::GetSamplePositionInternal: required UAnimNotifyState_PoseSearchSamplingAttribute in '%s' has an invalid Bone"), *AssetSamplers.GetAssetName());
+								return false;
 							}
-							
-							UE_LOG(LogPoseSearch, Error, TEXT("FAssetIndexer::GetSamplePositionInternal: required UAnimNotifyState_PoseSearchSamplingAttribute in '%s' has an invalid Bone"), *GetNameSafe(AssetSampler.GetAsset()));
-							OutSamplePosition = FVector::ZeroVector;
-							return false;
+
+							const FTransform RootBoneTransform = GetTransform(OriginTime, OriginRole, bUnused, RootSchemaBoneIdx);
+							OutSamplePosition = RootBoneTransform.InverseTransformPosition(SamplingAttribute->Position);
+							return true;
 						}
-						
-						const FTransform RootBoneTransform = GetTransform(OriginTime, bUnused, RootSchemaBoneIdx);
-						OutSamplePosition = RootBoneTransform.InverseTransformPosition(SamplingAttribute->Position);
-						return true;
 					}
 				}
-			}
+				return false;
+			}))
+		{
+			return true;
 		}
 
-		UE_LOG(LogPoseSearch, Error, TEXT("FAssetIndexer::GetSamplePositionInternal: required UAnimNotifyState_PoseSearchSamplingAttribute not found in '%s'"), *GetNameSafe(AssetSampler.GetAsset()));
+		UE_LOG(LogPoseSearch, Error, TEXT("FAssetIndexer::GetSamplePositionInternal: required UAnimNotifyState_PoseSearchSamplingAttribute not found in '%s'"), *AssetSamplers.GetAssetName());
 		OutSamplePosition = FVector::ZeroVector;
 		return false;
 	}
 
-	const FTransform RootBoneTransform = GetTransform(OriginTime, bUnused, RootSchemaBoneIdx);
-	const FTransform SampleBoneTransform = GetTransform(SampleTime, bClamped, SchemaSampleBoneIdx);
+	bool bUnused;
+	const FTransform RootBoneTransform = GetTransform(OriginTime, OriginRole, bUnused, RootSchemaBoneIdx);
+	const FTransform SampleBoneTransform = GetTransform(SampleTime, SampleRole, bClamped, SchemaSampleBoneIdx);
 	if (SchemaOriginBoneIdx == RootSchemaBoneIdx)
 	{
 		OutSamplePosition = RootBoneTransform.InverseTransformPosition(SampleBoneTransform.GetTranslation());
@@ -567,14 +698,14 @@ bool FAssetIndexer::GetSamplePositionInternal(FVector& OutSamplePosition, float
 	}
 
 	bool bOriginClamped;
-	const FTransform OriginBoneTransform = GetTransform(OriginTime, bOriginClamped, SchemaOriginBoneIdx);
+	const FTransform OriginBoneTransform = GetTransform(OriginTime, OriginRole, bOriginClamped, SchemaOriginBoneIdx);
 	bClamped |= bOriginClamped;
 	const FVector DeltaBoneTranslation = SampleBoneTransform.GetTranslation() - OriginBoneTransform.GetTranslation();
 	OutSamplePosition = RootBoneTransform.InverseTransformVector(DeltaBoneTranslation);
 	return true;
 }
 
-bool FAssetIndexer::GetSampleVelocity(FVector& OutSampleVelocity, float SampleTimeOffset, float OriginTimeOffset, int32 SampleIdx, int8 SchemaSampleBoneIdx, int8 SchemaOriginBoneIdx, bool bUseCharacterSpaceVelocities, EPermutationTimeType PermutationTimeType, int32 SamplingAttributeId)
+bool FAssetIndexer::GetSampleVelocity(FVector& OutSampleVelocity, float SampleTimeOffset, float OriginTimeOffset, int32 SampleIdx, int8 SchemaSampleBoneIdx, int8 SchemaOriginBoneIdx, const FRole& SampleRole, const FRole& OriginRole, bool bUseCharacterSpaceVelocities, EPermutationTimeType PermutationTimeType, int32 SamplingAttributeId)
 {
 	float PermutationSampleTimeOffset = 0.f;
 	float PermutationOriginTimeOffset = 0.f;
@@ -583,71 +714,75 @@ bool FAssetIndexer::GetSampleVelocity(FVector& OutSampleVelocity, float SampleTi
 	const float Time = CalculateSampleTime(SampleIdx);
 	const float SampleTime = Time + SampleTimeOffset + PermutationSampleTimeOffset;
 	const float OriginTime = Time + OriginTimeOffset + PermutationOriginTimeOffset;
-
-	bool bUnused, bClampedPast;
-	FVector BonePositionPast, BonePositionPresent;
+	
 	if (SamplingAttributeId >= 0)
 	{
-		const TConstArrayView<FAnimNotifyEvent> AnimNotifyEvents = AssetSampler.GetAllAnimNotifyEvents();
-		if (!AnimNotifyEvents.IsEmpty())
-		{
-			for (const FAnimNotifyEvent& AnimNotifyEvent : AnimNotifyEvents)
+		if (AssetSamplers.ProcessAllAnimNotifyEvents([&OutSampleVelocity, SchemaSampleBoneIdx, SchemaOriginBoneIdx, bUseCharacterSpaceVelocities, &SampleRole, &OriginRole, SamplingAttributeId, this, OriginTime](const TConstArrayView<FAnimNotifyEvent> AnimNotifyEvents)
 			{
-				if (const UAnimNotifyState_PoseSearchSamplingAttribute* SamplingAttribute = Cast<UAnimNotifyState_PoseSearchSamplingAttribute>(AnimNotifyEvent.NotifyStateClass))
+				for (const FAnimNotifyEvent& AnimNotifyEvent : AnimNotifyEvents)
 				{
-					if (SamplingAttribute->SamplingAttributeId == SamplingAttributeId)
+					if (const UAnimNotifyState_PoseSearchSamplingAttribute* SamplingAttribute = Cast<UAnimNotifyState_PoseSearchSamplingAttribute>(AnimNotifyEvent.NotifyStateClass))
 					{
-						if (SamplingAttribute->Bone.BoneName != NAME_None)
+						if (SamplingAttribute->SamplingAttributeId == SamplingAttributeId)
 						{
-							// @todo perhaps cache the initialized bone to speed up this method
-							FBoneReference TempBoneReference = SamplingAttribute->Bone;
-							TempBoneReference.Initialize(BoneContainer.GetSkeletonAsset());
-							if (TempBoneReference.HasValidSetup())
+							bool bUnused, bClampedPast;
+							FVector BonePositionPast, BonePositionPresent;
+							if (SamplingAttribute->Bone.BoneName != NAME_None)
 							{
-								const float SamplingAttributeTime = AnimNotifyEvent.GetTime();
-
-								if (GetSamplePositionInternal(BonePositionPast, SamplingAttributeTime - FiniteDelta, bUseCharacterSpaceVelocities ? OriginTime - FiniteDelta : OriginTime, bClampedPast, SchemaSampleBoneIdx, SchemaOriginBoneIdx, SamplingAttributeId) &&
-									GetSamplePositionInternal(BonePositionPresent, SamplingAttributeTime, OriginTime, bUnused, SchemaSampleBoneIdx, SchemaOriginBoneIdx, SamplingAttributeId))
+								// @todo: FIX ME TempBoneReference is not used!
+								// @todo perhaps cache the initialized bone to speed up this method
+								FBoneReference TempBoneReference = SamplingAttribute->Bone;
+								const int32 SampleRoleIndex = RoleToIndex[SampleRole];
+								TempBoneReference.Initialize(BoneContainers[SampleRoleIndex].GetSkeletonAsset());
+								if (TempBoneReference.HasValidSetup())
 								{
-									if (!bClampedPast)
-									{
-										OutSampleVelocity = (BonePositionPresent - BonePositionPast) / FiniteDelta;
-										return true;
-									}
+									const float SamplingAttributeTime = AnimNotifyEvent.GetTime();
 
-									FVector BonePositionFuture;
-									if (GetSamplePositionInternal(BonePositionFuture, SamplingAttributeTime + FiniteDelta, bUseCharacterSpaceVelocities ? OriginTime + FiniteDelta : OriginTime, bUnused, SchemaSampleBoneIdx, SchemaOriginBoneIdx, SamplingAttributeId))
+									if (GetSamplePositionInternal(BonePositionPast, SamplingAttributeTime - FiniteDelta, bUseCharacterSpaceVelocities ? OriginTime - FiniteDelta : OriginTime, bClampedPast, SchemaSampleBoneIdx, SchemaOriginBoneIdx, SampleRole, OriginRole, SamplingAttributeId) &&
+										GetSamplePositionInternal(BonePositionPresent, SamplingAttributeTime, OriginTime, bUnused, SchemaSampleBoneIdx, SchemaOriginBoneIdx, SampleRole, OriginRole, SamplingAttributeId))
 									{
-										OutSampleVelocity = (BonePositionFuture - BonePositionPresent) / FiniteDelta;
-										return true;
+										if (!bClampedPast)
+										{
+											OutSampleVelocity = (BonePositionPresent - BonePositionPast) / FiniteDelta;
+											return true;
+										}
+
+										FVector BonePositionFuture;
+										if (GetSamplePositionInternal(BonePositionFuture, SamplingAttributeTime + FiniteDelta, bUseCharacterSpaceVelocities ? OriginTime + FiniteDelta : OriginTime, bUnused, SchemaSampleBoneIdx, SchemaOriginBoneIdx, SampleRole, OriginRole, SamplingAttributeId))
+										{
+											OutSampleVelocity = (BonePositionFuture - BonePositionPresent) / FiniteDelta;
+											return true;
+										}
 									}
+
+									return false;
 								}
 
-								// @todo: should we log some error message?
-								OutSampleVelocity = FVector::ZeroVector;
+								UE_LOG(LogPoseSearch, Error, TEXT("FAssetIndexer::GetSampleVelocity: required UAnimNotifyState_PoseSearchSamplingAttribute in '%s' has an invalid Bone"), *AssetSamplers.GetAssetName());
 								return false;
 							}
 
-							UE_LOG(LogPoseSearch, Error, TEXT("FAssetIndexer::GetSampleVelocity: required UAnimNotifyState_PoseSearchSamplingAttribute in '%s' has an invalid Bone"), *GetNameSafe(AssetSampler.GetAsset()));
-							OutSampleVelocity = FVector::ZeroVector;
-							return false;
+							const FTransform RootBoneTransform = GetTransform(OriginTime, OriginRole, bUnused, RootSchemaBoneIdx);
+							OutSampleVelocity = RootBoneTransform.InverseTransformPosition(SamplingAttribute->LinearVelocity);
+							return true;
 						}
-						
-						const FTransform RootBoneTransform = GetTransform(OriginTime, bUnused, RootSchemaBoneIdx);
-						OutSampleVelocity = RootBoneTransform.InverseTransformPosition(SamplingAttribute->LinearVelocity);
-						return true;
 					}
 				}
-			}
+				return false;
+			}))
+		{
+			return true;
 		}
 
-		UE_LOG(LogPoseSearch, Error, TEXT("FAssetIndexer::GetSampleVelocity: required UAnimNotifyState_PoseSearchSamplingAttribute not found in '%s'"), *GetNameSafe(AssetSampler.GetAsset()));
+		UE_LOG(LogPoseSearch, Error, TEXT("FAssetIndexer::GetSampleVelocity: required UAnimNotifyState_PoseSearchSamplingAttribute not found in '%s'"), *AssetSamplers.GetAssetName());
 		OutSampleVelocity = FVector::ZeroVector;
 		return false;
 	}
 
-	if (GetSamplePositionInternal(BonePositionPast, SampleTime - FiniteDelta, bUseCharacterSpaceVelocities ? OriginTime - FiniteDelta : OriginTime, bClampedPast, SchemaSampleBoneIdx, SchemaOriginBoneIdx, INDEX_NONE) &&
-		GetSamplePositionInternal(BonePositionPresent, SampleTime, OriginTime, bUnused, SchemaSampleBoneIdx, SchemaOriginBoneIdx, INDEX_NONE))
+	bool bUnused, bClampedPast;
+	FVector BonePositionPast, BonePositionPresent;
+	if (GetSamplePositionInternal(BonePositionPast, SampleTime - FiniteDelta, bUseCharacterSpaceVelocities ? OriginTime - FiniteDelta : OriginTime, bClampedPast, SchemaSampleBoneIdx, SchemaOriginBoneIdx, SampleRole, OriginRole, INDEX_NONE) &&
+		GetSamplePositionInternal(BonePositionPresent, SampleTime, OriginTime, bUnused, SchemaSampleBoneIdx, SchemaOriginBoneIdx, SampleRole, OriginRole, INDEX_NONE))
 	{
 		if (!bClampedPast)
 		{
@@ -656,7 +791,7 @@ bool FAssetIndexer::GetSampleVelocity(FVector& OutSampleVelocity, float SampleTi
 		}
 
 		FVector BonePositionFuture;
-		if (GetSamplePositionInternal(BonePositionFuture, SampleTime + FiniteDelta, bUseCharacterSpaceVelocities ? OriginTime + FiniteDelta : OriginTime, bUnused, SchemaSampleBoneIdx, SchemaOriginBoneIdx, INDEX_NONE))
+		if (GetSamplePositionInternal(BonePositionFuture, SampleTime + FiniteDelta, bUseCharacterSpaceVelocities ? OriginTime + FiniteDelta : OriginTime, bUnused, SchemaSampleBoneIdx, SchemaOriginBoneIdx, SampleRole, OriginRole, INDEX_NONE))
 		{
 			OutSampleVelocity = (BonePositionFuture - BonePositionPresent) / FiniteDelta;
 			return true;
@@ -716,11 +851,11 @@ void FAssetIndexer::CompareCachedEntries(const FAssetIndexer& Other) const
 	}
 	else
 	{
-		for (const TPair<float, CachedEntry>& Pair : CachedEntries)
+		for (const TPair<float, FCachedEntry>& Pair : CachedEntries)
 		{
-			if (const CachedEntry* OtherEntry = Other.CachedEntries.Find(Pair.Key))
+			if (const FCachedEntry* OtherEntry = Other.CachedEntries.Find(Pair.Key))
 			{
-				const CachedEntry* Entry = &Pair.Value;
+				const FCachedEntry* Entry = &Pair.Value;
 				if (Entry->SampleTime != OtherEntry->SampleTime)
 				{
 					UE_LOG(LogPoseSearch, Warning, TEXT("CompareCachedEntries - FAssetIndexer::CachedEntries::SampleTime is not deterministic (%f, %f)"), Entry->SampleTime, OtherEntry->SampleTime);
@@ -731,29 +866,49 @@ void FAssetIndexer::CompareCachedEntries(const FAssetIndexer& Other) const
 					UE_LOG(LogPoseSearch, Warning, TEXT("CompareCachedEntries - FAssetIndexer::CachedEntries::bClamped is not deterministic"));
 				}
 
-				if (FMemory::Memcmp(&Entry->RootTransform, &OtherEntry->RootTransform, sizeof(FTransform)) != 0)
+				if (Entry->RootTransform.Num() != OtherEntry->RootTransform.Num())
 				{
-					UE_LOG(LogPoseSearch, Warning, TEXT("CompareCachedEntries - FAssetIndexer::CachedEntries::RootTransform is not deterministic"));
+					UE_LOG(LogPoseSearch, Warning, TEXT("CompareCachedEntries - FAssetIndexer::CachedEntries::RootTransform::Num is not deterministic"));
 				}
-
-				if (Entry->ComponentSpacePose.GetComponentSpaceFlags() != OtherEntry->ComponentSpacePose.GetComponentSpaceFlags())
+				else
 				{
-					UE_LOG(LogPoseSearch, Warning, TEXT("CompareCachedEntries - FAssetIndexer::CachedEntries::ComponentSpacePose::ComponentSpaceFlags is not deterministic"));
+					for (int32 RoleIndex = 0; RoleIndex < Entry->RootTransform.Num(); ++RoleIndex)
+					{
+						if (FMemory::Memcmp(&Entry->RootTransform[RoleIndex], &OtherEntry->RootTransform[RoleIndex], sizeof(FTransform)) != 0)
+						{
+							UE_LOG(LogPoseSearch, Warning, TEXT("CompareCachedEntries - FAssetIndexer::CachedEntries::RootTransform[%d] is not deterministic"), RoleIndex);
+						}
+					}
 				}
 
-				const TConstArrayView<FTransform> Bones = Entry->ComponentSpacePose.GetPose().GetBones();
-				const TConstArrayView<FTransform> OtherBones = OtherEntry->ComponentSpacePose.GetPose().GetBones();
-				if (Bones.Num() != OtherBones.Num())
+				if (Entry->ComponentSpacePose.Num() != OtherEntry->ComponentSpacePose.Num())
 				{
-					UE_LOG(LogPoseSearch, Warning, TEXT("CompareCachedEntries - FAssetIndexer::CachedEntries::ComponentSpacePose::Bones is not deterministic"));
+					UE_LOG(LogPoseSearch, Warning, TEXT("CompareCachedEntries - FAssetIndexer::CachedEntries::ComponentSpacePose::Num is not deterministic"));
 				}
 				else
 				{
-					for (int32 Index = 0; Index < Bones.Num(); ++Index)
+					for (int32 RoleIndex = 0; RoleIndex < Entry->ComponentSpacePose.Num(); ++RoleIndex)
 					{
-						if (FMemory::Memcmp(&Bones[Index], &OtherBones[Index], sizeof(FTransform)) != 0)
+						if (Entry->ComponentSpacePose[RoleIndex].GetComponentSpaceFlags() != OtherEntry->ComponentSpacePose[RoleIndex].GetComponentSpaceFlags())
 						{
-							UE_LOG(LogPoseSearch, Warning, TEXT("CompareCachedEntries - FAssetIndexer::CachedEntries::Bones[%d] is not deterministic"), Index);
+							UE_LOG(LogPoseSearch, Warning, TEXT("CompareCachedEntries - FAssetIndexer::CachedEntries::ComponentSpacePose[%d]::ComponentSpaceFlags is not deterministic"), RoleIndex);
+						}
+
+						const TConstArrayView<FTransform> Bones = Entry->ComponentSpacePose[RoleIndex].GetPose().GetBones();
+						const TConstArrayView<FTransform> OtherBones = OtherEntry->ComponentSpacePose[RoleIndex].GetPose().GetBones();
+						if (Bones.Num() != OtherBones.Num())
+						{
+							UE_LOG(LogPoseSearch, Warning, TEXT("CompareCachedEntries - FAssetIndexer::CachedEntries::ComponentSpacePose[%d]::Bones is not deterministic"), RoleIndex);
+						}
+						else
+						{
+							for (int32 BoneIndex = 0; BoneIndex < Bones.Num(); ++BoneIndex)
+							{
+								if (FMemory::Memcmp(&Bones[BoneIndex], &OtherBones[BoneIndex], sizeof(FTransform)) != 0)
+								{
+									UE_LOG(LogPoseSearch, Warning, TEXT("CompareCachedEntries - FAssetIndexer::CachedEntries::ComponentSpacePose[%d]::Bones[%d] is not deterministic"), RoleIndex, BoneIndex);
+								}
+							}
 						}
 					}
 				}
diff --git a/Engine/Plugins/Experimental/Animation/PoseSearch/Source/Runtime/Private/PoseSearchAssetSampler.cpp b/Engine/Plugins/Experimental/Animation/PoseSearch/Source/Runtime/Private/PoseSearchAssetSampler.cpp
--- a/Engine/Plugins/Experimental/Animation/PoseSearch/Source/Runtime/Private/PoseSearchAssetSampler.cpp
+++ b/Engine/Plugins/Experimental/Animation/PoseSearch/Source/Runtime/Private/PoseSearchAssetSampler.cpp
@@ -143,7 +143,7 @@ static FTransform ExtractBlendSpaceRootMotion(float StartTime, float DeltaTime,
 	return RootMotionParams.GetRootMotionTransform();
 }
 
-static void ProcessRootTransform(const UBlendSpace* BlendSpace, const FVector& BlendParameters, float CachedPlayLength, const FBoneContainer& BoneContainer,
+static void ProcessRootTransform(const UBlendSpace* BlendSpace, const FVector& BlendParameters, float CachedPlayLength,
 	int32 RootTransformSamplingRate, bool bIsLoopable, TArray<FTransform>& AccumulatedRootTransform)
 {
 	// Pre-compute root motion
@@ -192,14 +192,15 @@ static int32 GetHighestWeightSample(const TArray<struct FBlendSampleData>& Sampl
 
 //////////////////////////////////////////////////////////////////////////
 // FAssetSamplerBase
-FAnimationAssetSampler::FAnimationAssetSampler(TObjectPtr<const UAnimationAsset> InAnimationAsset, const FVector& InBlendParameters, int32 InRootTransformSamplingRate)
+FAnimationAssetSampler::FAnimationAssetSampler(TObjectPtr<const UAnimationAsset> InAnimationAsset, const FTransform& InRootTransformOrigin, const FVector& InBlendParameters, int32 InRootTransformSamplingRate)
 {
-	Init(InAnimationAsset, InBlendParameters, InRootTransformSamplingRate);
+	Init(InAnimationAsset, InRootTransformOrigin, InBlendParameters, InRootTransformSamplingRate);
 }
 
-void FAnimationAssetSampler::Init(TObjectPtr<const UAnimationAsset> InAnimationAsset, const FVector& InBlendParameters, int32 InRootTransformSamplingRate)
+void FAnimationAssetSampler::Init(TObjectPtr<const UAnimationAsset> InAnimationAsset, const FTransform& InRootTransformOrigin, const FVector& InBlendParameters, int32 InRootTransformSamplingRate)
 {
 	AnimationAssetPtr = InAnimationAsset;
+	RootTransformOrigin = InRootTransformOrigin;
 	BlendParameters = InBlendParameters;
 	RootTransformSamplingRate = InRootTransformSamplingRate;
 	CachedPlayLength = GetPlayLength(AnimationAssetPtr.Get(), BlendParameters);
@@ -524,14 +525,14 @@ FTransform FAnimationAssetSampler::ExtractRootTransform(float Time) const
 		}
 	}
 
-	return RootTransform;
+	return RootTransform * RootTransformOrigin;
 }
 
-void FAnimationAssetSampler::Process(const FBoneContainer& BoneContainer)
+void FAnimationAssetSampler::Process()
 {
 	if (const UBlendSpace* BlendSpace = Cast<UBlendSpace>(AnimationAssetPtr.Get()))
 	{
-		ProcessRootTransform(BlendSpace, BlendParameters, CachedPlayLength, BoneContainer, RootTransformSamplingRate, IsLoopable(), AccumulatedRootTransform);
+		ProcessRootTransform(BlendSpace, BlendParameters, CachedPlayLength, RootTransformSamplingRate, IsLoopable(), AccumulatedRootTransform);
 	}
 }
 
diff --git a/Engine/Plugins/Experimental/Animation/PoseSearch/Source/Runtime/Private/PoseSearchContext.cpp b/Engine/Plugins/Experimental/Animation/PoseSearch/Source/Runtime/Private/PoseSearchContext.cpp
--- a/Engine/Plugins/Experimental/Animation/PoseSearch/Source/Runtime/Private/PoseSearchContext.cpp
+++ b/Engine/Plugins/Experimental/Animation/PoseSearch/Source/Runtime/Private/PoseSearchContext.cpp
@@ -16,29 +16,33 @@ namespace UE::PoseSearch
 #if ENABLE_DRAW_DEBUG
 //////////////////////////////////////////////////////////////////////////
 // FDebugDrawParams
-FDebugDrawParams::FDebugDrawParams(FAnimInstanceProxy* InAnimInstanceProxy, const FTransform& InRootMotionTransform, const UPoseSearchDatabase* InDatabase, EDebugDrawFlags InFlags)
-: AnimInstanceProxy(InAnimInstanceProxy)
-, World(nullptr)
-, Mesh(nullptr)
-, RootMotionTransform(&InRootMotionTransform)
+FDebugDrawParams::FDebugDrawParams(TArrayView<FAnimInstanceProxy*> InAnimInstanceProxies, TConstArrayView<const IPoseHistory*> InPoseHistories, const FRoleToIndex& InRoleToIndex, const UPoseSearchDatabase* InDatabase, EDebugDrawFlags InFlags)
+: AnimInstanceProxies(InAnimInstanceProxies)
+, Meshes()
+, PoseHistories(InPoseHistories)
+, RoleToIndex(InRoleToIndex)
 , Database(InDatabase)
 , Flags(InFlags)
 {
+	check(RoleToIndex.Num() == InAnimInstanceProxies.Num() && RoleToIndex.Num() == PoseHistories.Num());
+	check(IsValid(RoleToIndex));
 }
 
-FDebugDrawParams::FDebugDrawParams(const UWorld* InWorld, const USkinnedMeshComponent* InMesh, const FTransform& InRootMotionTransform, const UPoseSearchDatabase* InDatabase, EDebugDrawFlags InFlags)
-: AnimInstanceProxy(nullptr)
-, World(InWorld)
-, Mesh(InMesh)
-, RootMotionTransform(&InRootMotionTransform)
+FDebugDrawParams::FDebugDrawParams(TArrayView<const USkinnedMeshComponent*> InMeshes, TConstArrayView<const IPoseHistory*> InPoseHistories, const FRoleToIndex& InRoleToIndex, const UPoseSearchDatabase* InDatabase, EDebugDrawFlags InFlags)
+: AnimInstanceProxies()
+, Meshes(InMeshes)
+, PoseHistories(InPoseHistories)
+, RoleToIndex(InRoleToIndex)
 , Database(InDatabase)
 , Flags(InFlags)
 {
+	check(RoleToIndex.Num() == Meshes.Num() && RoleToIndex.Num() == PoseHistories.Num());
+	check(IsValid(RoleToIndex));
 }
 
 bool FDebugDrawParams::CanDraw() const
 {
-	return (AnimInstanceProxy || World) && Database && Database->Schema && Database->Schema->IsValid();
+	return (!AnimInstanceProxies.IsEmpty() || !Meshes.IsEmpty()) && !RoleToIndex.IsEmpty() && Database && Database->Schema;
 }
 
 const FSearchIndex* FDebugDrawParams::GetSearchIndex() const
@@ -54,12 +58,13 @@ const UPoseSearchSchema* FDebugDrawParams::GetSchema() const
 FVector FDebugDrawParams::ExtractPosition(TConstArrayView<float> PoseVector, const UPoseSearchFeatureChannel_Position* Position) const
 {
 	check(Position);
+	check(Position->SampleRole == Position->OriginRole);
 	const FVector BonePosition = FFeatureVectorHelper::DecodeVector(PoseVector, Position->GetChannelDataOffset(), Position->ComponentStripping);
-	const FVector WorldBonePosition = GetRootTransform().TransformPosition(BonePosition);
+	const FVector WorldBonePosition = GetRootTransform(Position->SampleRole).TransformPosition(BonePosition);
 	return WorldBonePosition;
 }
 
-FVector FDebugDrawParams::ExtractPosition(TConstArrayView<float> PoseVector, float SampleTimeOffset, int8 SchemaBoneIdx, EPermutationTimeType PermutationTimeType, int32 SamplingAttributeId) const
+FVector FDebugDrawParams::ExtractPosition(TConstArrayView<float> PoseVector, float SampleTimeOffset, int8 SchemaBoneIdx, const FRole& Role, EPermutationTimeType PermutationTimeType, int32 SamplingAttributeId) const
 {
 	// we don't wanna ask for a SchemaOriginBoneIdx in the future or past
 	check(PermutationTimeType != EPermutationTimeType::UsePermutationTime);
@@ -68,15 +73,18 @@ FVector FDebugDrawParams::ExtractPosition(TConstArrayView<float> PoseVector, flo
 		// looking for a UPoseSearchFeatureChannel_Position that matches the TimeOffset and SchemaBoneIdx,
 		// with SchemaOriginBoneIdx to be the root bone and the appropriate PermutationTimeType 
 		if (const UPoseSearchFeatureChannel_Position* FoundPosition = static_cast<const UPoseSearchFeatureChannel_Position*>(
-			Schema->FindChannel([SampleTimeOffset, SchemaBoneIdx, PermutationTimeType, SamplingAttributeId, Schema](const UPoseSearchFeatureChannel* Channel) -> const UPoseSearchFeatureChannel_Position*
+			Schema->FindChannel([SampleTimeOffset, SchemaBoneIdx, &Role, PermutationTimeType, SamplingAttributeId](const UPoseSearchFeatureChannel* Channel) -> const UPoseSearchFeatureChannel_Position*
 				{
 					if (const UPoseSearchFeatureChannel_Position* Position = Cast<UPoseSearchFeatureChannel_Position>(Channel))
 					{
 						if (Position->SchemaBoneIdx == SchemaBoneIdx &&
 							Position->SampleTimeOffset == SampleTimeOffset &&
+							Position->OriginTimeOffset == 0.f &&
 							Position->PermutationTimeType == PermutationTimeType &&
 							Position->SamplingAttributeId == SamplingAttributeId &&
-							Position->SchemaOriginBoneIdx == RootSchemaBoneIdx)
+							Position->SchemaOriginBoneIdx == RootSchemaBoneIdx &&
+							Position->SampleRole == Role &&
+							Position->OriginRole == Role)
 						{
 							return Position;
 						}
@@ -87,15 +95,32 @@ FVector FDebugDrawParams::ExtractPosition(TConstArrayView<float> PoseVector, flo
 			return ExtractPosition(PoseVector, FoundPosition);
 		}
 
-		if (Mesh && SchemaBoneIdx > RootSchemaBoneIdx)
+		if (const int32* RoleIndex = RoleToIndex.Find(Role))
 		{
-			return Mesh->GetSocketTransform(Schema->BoneReferences[SchemaBoneIdx].BoneName).GetTranslation();
+			if (const IPoseHistory* PoseHistory = PoseHistories[*RoleIndex])
+			{
+				if (const USkeleton* Skeleton = Schema->GetSkeleton(Role))
+				{
+					const FBoneIndexType BoneIndexType = Schema->GetBoneReferences(Role)[SchemaBoneIdx].BoneIndex;
+
+					FTransform WorldBoneTransform;
+					if (PoseHistory->GetTransformAtTime(SampleTimeOffset, WorldBoneTransform, Skeleton, BoneIndexType, WorldSpaceIndexType))
+					{
+						return WorldBoneTransform.GetTranslation();
+					}
+				}
+			}
+
+			if (SchemaBoneIdx > RootSchemaBoneIdx && !Meshes.IsEmpty())
+			{
+				return Meshes[*RoleIndex]->GetSocketTransform(Schema->GetBoneReferences(Role)[SchemaBoneIdx].BoneName).GetTranslation();
+			}
 		}
 	}
-	return GetRootTransform().GetTranslation();
+	return GetRootTransform(Role).GetTranslation();
 }
 
-FQuat FDebugDrawParams::ExtractRotation(TConstArrayView<float> PoseVector, float SampleTimeOffset, int8 SchemaBoneIdx) const
+FQuat FDebugDrawParams::ExtractRotation(TConstArrayView<float> PoseVector, float SampleTimeOffset, int8 SchemaBoneIdx, const FRole& Role, EPermutationTimeType PermutationTimeType, int32 SamplingAttributeId) const
 {
 	if (const UPoseSearchSchema* Schema = GetSchema())
 	{
@@ -108,13 +133,18 @@ FQuat FDebugDrawParams::ExtractRotation(TConstArrayView<float> PoseVector, float
 			// the features data associated to this channel would be a heading vector in GetRootTransform space (since OriginTimeOffset is zero)), 
 			// so by finding at least two with differnt axis we'll be able to compose a delta rotation from OriginTimeOffset (zero) to SampleTimeOffset
 			FoundHeading[HeadingAxis] = static_cast<const UPoseSearchFeatureChannel_Heading*>(
-				Schema->FindChannel([SampleTimeOffset, SchemaBoneIdx, HeadingAxis](const UPoseSearchFeatureChannel* Channel) -> const UPoseSearchFeatureChannel_Heading*
+				Schema->FindChannel([SampleTimeOffset, SchemaBoneIdx, &Role, PermutationTimeType, SamplingAttributeId, HeadingAxis](const UPoseSearchFeatureChannel* Channel) -> const UPoseSearchFeatureChannel_Heading*
 					{
 						if (const UPoseSearchFeatureChannel_Heading* Heading = Cast<UPoseSearchFeatureChannel_Heading>(Channel))
 						{
-							if (Heading->OriginTimeOffset == 0.f &&
-								Heading->SchemaBoneIdx == SchemaBoneIdx &&
+							if (Heading->SchemaBoneIdx == SchemaBoneIdx &&
 								Heading->SampleTimeOffset == SampleTimeOffset &&
+								Heading->OriginTimeOffset == 0.f &&
+								Heading->PermutationTimeType == PermutationTimeType &&
+								Heading->SamplingAttributeId == SamplingAttributeId &&
+								Heading->SchemaOriginBoneIdx == RootSchemaBoneIdx &&
+								Heading->SampleRole == Role &&
+								Heading->OriginRole == Role &&
 								int32(Heading->HeadingAxis) == HeadingAxis)
 							{
 								return Heading;
@@ -185,37 +215,60 @@ FQuat FDebugDrawParams::ExtractRotation(TConstArrayView<float> PoseVector, float
 				// world rotation associated to the time zero, we can calcualte the world rotation at time SampleTimeOffset
 				const FMatrix RotMatrix(DecodedHeading[int32(EHeadingAxis::X)], DecodedHeading[int32(EHeadingAxis::Y)], DecodedHeading[int32(EHeadingAxis::Z)], FVector::ZeroVector);
 				const FQuat RotQuat(RotMatrix);
-				const FQuat RotQuatWorld = RotQuat * GetRootTransform().GetRotation();
+				const FQuat RotQuatWorld = RotQuat * GetRootTransform(Role).GetRotation();
 				return RotQuatWorld;
 			}
 		}
 
-		if (Mesh && SchemaBoneIdx > RootSchemaBoneIdx)
+		if (const int32* RoleIndex = RoleToIndex.Find(Role))
 		{
-			return Mesh->GetSocketTransform(Schema->BoneReferences[SchemaBoneIdx].BoneName).GetRotation();
+			if (const IPoseHistory* PoseHistory = PoseHistories[*RoleIndex])
+			{
+				if (const USkeleton* Skeleton = Schema->GetSkeleton(Role))
+				{
+					const FBoneIndexType BoneIndexType = Schema->GetBoneReferences(Role)[SchemaBoneIdx].BoneIndex;
+
+					FTransform WorldBoneTransform;
+					if (PoseHistory->GetTransformAtTime(SampleTimeOffset, WorldBoneTransform, Skeleton, BoneIndexType, WorldSpaceIndexType))
+					{
+						return WorldBoneTransform.GetRotation();
+					}
+				}
+			}
+
+			if (SchemaBoneIdx > RootSchemaBoneIdx && !Meshes.IsEmpty())
+			{
+				return Meshes[*RoleIndex]->GetSocketTransform(Schema->GetBoneReferences(Role)[SchemaBoneIdx].BoneName).GetRotation();
+			}
 		}
 	}
 
-	return GetRootTransform().GetRotation();
+	return GetRootTransform(Role).GetRotation();
 }
 
-const FTransform& FDebugDrawParams::GetRootTransform() const
+FTransform FDebugDrawParams::GetRootTransform(const FRole& Role) const
 {
-	check(RootMotionTransform);
-	return *RootMotionTransform;
+	if (const int32* RoleIndex = RoleToIndex.Find(Role))
+	{
+		// @todo: add and use SampleTimeOffset to get the root transform at a specific time now using the trajectory
+		return PoseHistories[*RoleIndex]->GetTrajectory().GetSampleAtTime(0.f).GetTransform();
+	}
+	return FTransform::Identity;
 }
 
 void FDebugDrawParams::DrawLine(const FVector& LineStart, const FVector& LineEnd, const FColor& Color, float Thickness) const
 {
 	if (Color.A > 0)
 	{
-		if (AnimInstanceProxy)
+		if (!AnimInstanceProxies.IsEmpty())
 		{
-			AnimInstanceProxy->AnimDrawDebugLine(LineStart, LineEnd, Color, false, 0.f, Thickness, SDPG_Foreground);
+			// any AnimInstanceProxy is fine to draw
+			AnimInstanceProxies[0]->AnimDrawDebugLine(LineStart, LineEnd, Color, false, 0.f, Thickness, SDPG_Foreground);
 		}
-		else if (World)
+		else if (!Meshes.IsEmpty())
 		{
-			DrawDebugLine(World, LineStart, LineEnd, Color, false, 0.f, SDPG_Foreground, Thickness);
+			// any Mesh is fine to draw
+			DrawDebugLine(Meshes[0]->GetWorld(), LineStart, LineEnd, Color, false, 0.f, SDPG_Foreground, Thickness);
 		}
 	}
 }
@@ -224,13 +277,15 @@ void FDebugDrawParams::DrawPoint(const FVector& Position, const FColor& Color, f
 {
 	if (Color.A > 0)
 	{
-		if (AnimInstanceProxy)
+		if (!AnimInstanceProxies.IsEmpty())
 		{
-			AnimInstanceProxy->AnimDrawDebugPoint(Position, Thickness, Color, false, 0.f, SDPG_Foreground);
+			// any AnimInstanceProxy is fine to draw
+			AnimInstanceProxies[0]->AnimDrawDebugPoint(Position, Thickness, Color, false, 0.f, SDPG_Foreground);
 		}
-		else if (World)
+		else if (!Meshes.IsEmpty())
 		{
-			DrawDebugPoint(World, Position, Thickness, Color, false, 0.f, SDPG_Foreground);
+			// any Mesh is fine to draw
+			DrawDebugPoint(Meshes[0]->GetWorld(), Position, Thickness, Color, false, 0.f, SDPG_Foreground);
 		}
 	}
 }
@@ -239,14 +294,16 @@ void FDebugDrawParams::DrawCircle(const FMatrix& TransformMatrix, float Radius,
 {
 	if (Color.A > 0)
 	{
-		if (AnimInstanceProxy)
+		if (!AnimInstanceProxies.IsEmpty())
 		{
-			AnimInstanceProxy->AnimDrawDebugCircle(TransformMatrix.GetOrigin(), Radius, Segments, Color, TransformMatrix.GetScaledAxis(EAxis::X), false, 0.f, SDPG_Foreground, Thickness);
+			// any AnimInstanceProxy is fine to draw
+			AnimInstanceProxies[0]->AnimDrawDebugCircle(TransformMatrix.GetOrigin(), Radius, Segments, Color, TransformMatrix.GetScaledAxis(EAxis::X), false, 0.f, SDPG_Foreground, Thickness);
 		}
-		else if (World)
+		else if (!Meshes.IsEmpty())
 		{
+			// any Mesh is fine to draw
 			// @todo: use the DrawDebugCircle API with the up vector to communize with the AnimInstanceProxy call
-			DrawDebugCircle(World, TransformMatrix, Radius, Segments, Color, false, 0.f, SDPG_Foreground, Thickness);
+			DrawDebugCircle(Meshes[0]->GetWorld(), TransformMatrix, Radius, Segments, Color, false, 0.f, SDPG_Foreground, Thickness);
 		}
 	}
 }
@@ -353,29 +410,39 @@ void FDebugDrawParams::DrawFeatureVector(int32 PoseIdx)
 // FCachedQuery
 FCachedQuery::FCachedQuery(const UPoseSearchSchema* InSchema)
 {
-	check(InSchema && InSchema->IsValid());
+	check(InSchema);
 	Schema = InSchema;
 	Values.SetNumZeroed(Schema->SchemaCardinality);
 }
 
 //////////////////////////////////////////////////////////////////////////
 // FSearchContext
-FSearchContext::FSearchContext(const UAnimInstance* InAnimInstance, const IPoseHistory* InHistory, TConstArrayView<const UObject*> InAssetsToConsider,
-		float InDesiredPermutationTimeOffset, const FPoseIndicesHistory* InPoseIndicesHistory,
-		const FSearchResult& InCurrentResult, const FFloatInterval& InPoseJumpThresholdTime, bool bInUseCachedChannelData)
-: AnimInstance(InAnimInstance)
-, History(InHistory)
-, AssetsToConsider(InAssetsToConsider)
+FSearchContext::FSearchContext(float InDesiredPermutationTimeOffset, const FPoseIndicesHistory* InPoseIndicesHistory,
+	const FSearchResult& InCurrentResult, const FFloatInterval& InPoseJumpThresholdTime, bool bInUseCachedChannelData)
+: AnimInstances()
+, PoseHistories()
+, RoleToIndex()
+, AssetsToConsider()
 , DesiredPermutationTimeOffset(InDesiredPermutationTimeOffset)
 , PoseIndicesHistory(InPoseIndicesHistory)
 , CurrentResult(InCurrentResult)
 , PoseJumpThresholdTime(InPoseJumpThresholdTime)
 , bUseCachedChannelData(bInUseCachedChannelData)
 {
-	check(AnimInstance);
 	UpdateCurrentResultPoseVector();
 }
 
+void FSearchContext::AddRole(const FRole& Role, const UAnimInstance* AnimInstance, const IPoseHistory* PoseHistory)
+{
+	check(RoleToIndex.Num() == AnimInstances.Num() && RoleToIndex.Num() == PoseHistories.Num());
+
+	AnimInstances.Add(AnimInstance);
+	PoseHistories.Add(PoseHistory);
+	RoleToIndex.Add(Role) = RoleToIndex.Num();
+
+	check(IsValid(RoleToIndex));
+}
+
 void FSearchContext::UpdateCurrentResultPoseVector()
 {
 	if (CurrentResult.IsValid())
@@ -394,7 +461,7 @@ void FSearchContext::UpdateCurrentResultPoseVector()
 	}
 }
 
-FQuat FSearchContext::GetSampleRotation(float SampleTimeOffset, float OriginTimeOffset, int8 SchemaSampleBoneIdx, int8 SchemaOriginBoneIdx, EPermutationTimeType PermutationTimeType, const FQuat* SampleBoneRotationWorldOverride)
+FQuat FSearchContext::GetSampleRotation(float SampleTimeOffset, float OriginTimeOffset, int8 SchemaSampleBoneIdx, int8 SchemaOriginBoneIdx, const FRole& SampleRole, const FRole& OriginRole, EPermutationTimeType PermutationTimeType, const FQuat* SampleBoneRotationWorldOverride)
 {
 	float PermutationSampleTimeOffset = 0.f;
 	float PermutationOriginTimeOffset = 0.f;
@@ -403,10 +470,10 @@ FQuat FSearchContext::GetSampleRotation(float SampleTimeOffset, float OriginTime
 	const float SampleTime = SampleTimeOffset + PermutationSampleTimeOffset;
 	const float OriginTime = OriginTimeOffset + PermutationOriginTimeOffset;
 
-	return GetSampleRotationInternal(SampleTime, OriginTime, SchemaSampleBoneIdx, SchemaOriginBoneIdx, SampleBoneRotationWorldOverride);
+	return GetSampleRotationInternal(SampleTime, OriginTime, SchemaSampleBoneIdx, SchemaOriginBoneIdx, SampleRole, OriginRole, SampleBoneRotationWorldOverride);
 }
 
-FVector FSearchContext::GetSamplePosition(float SampleTimeOffset, float OriginTimeOffset, int8 SchemaSampleBoneIdx, int8 SchemaOriginBoneIdx, EPermutationTimeType PermutationTimeType, const FVector* SampleBonePositionWorldOverride)
+FVector FSearchContext::GetSamplePosition(float SampleTimeOffset, float OriginTimeOffset, int8 SchemaSampleBoneIdx, int8 SchemaOriginBoneIdx, const FRole& SampleRole, const FRole& OriginRole, EPermutationTimeType PermutationTimeType, const FVector* SampleBonePositionWorldOverride)
 {
 	float PermutationSampleTimeOffset = 0.f;
 	float PermutationOriginTimeOffset = 0.f;
@@ -414,10 +481,10 @@ FVector FSearchContext::GetSamplePosition(float SampleTimeOffset, float OriginTi
 
 	const float SampleTime = SampleTimeOffset + PermutationSampleTimeOffset;
 	const float OriginTime = OriginTimeOffset + PermutationOriginTimeOffset;
-	return GetSamplePositionInternal(SampleTime, OriginTime, SchemaSampleBoneIdx, SchemaOriginBoneIdx, SampleBonePositionWorldOverride);
+	return GetSamplePositionInternal(SampleTime, OriginTime, SchemaSampleBoneIdx, SchemaOriginBoneIdx, SampleRole, OriginRole, SampleBonePositionWorldOverride);
 }
 
-FVector FSearchContext::GetSampleVelocity(float SampleTimeOffset, float OriginTimeOffset, int8 SchemaSampleBoneIdx, int8 SchemaOriginBoneIdx, bool bUseCharacterSpaceVelocities, EPermutationTimeType PermutationTimeType, const FVector* SampleBoneVelocityWorldOverride)
+FVector FSearchContext::GetSampleVelocity(float SampleTimeOffset, float OriginTimeOffset, int8 SchemaSampleBoneIdx, int8 SchemaOriginBoneIdx, const FRole& SampleRole, const FRole& OriginRole, bool bUseCharacterSpaceVelocities, EPermutationTimeType PermutationTimeType, const FVector* SampleBoneVelocityWorldOverride)
 {
 	using namespace UE::PoseSearch;
 
@@ -430,26 +497,27 @@ FVector FSearchContext::GetSampleVelocity(float SampleTimeOffset, float OriginTi
 
 	if (SampleBoneVelocityWorldOverride)
 	{
-		const FTransform RootBoneTransform = GetWorldBoneTransformAtTime(OriginTime, RootSchemaBoneIdx);
+		const FTransform RootBoneTransform = GetWorldBoneTransformAtTime(OriginTime, OriginRole, RootSchemaBoneIdx);
 		return RootBoneTransform.InverseTransformVector(*SampleBoneVelocityWorldOverride);
 	}
 
 	// calculating the local Position for the bone indexed by SchemaSampleBoneIdx
-	const FVector PreviousTranslation = GetSamplePositionInternal(SampleTime - FiniteDelta, bUseCharacterSpaceVelocities ? OriginTime - FiniteDelta : OriginTime, SchemaSampleBoneIdx, SchemaOriginBoneIdx);
-	const FVector CurrentTranslation = GetSamplePositionInternal(SampleTime, OriginTime, SchemaSampleBoneIdx, SchemaOriginBoneIdx);
+	const FVector PreviousTranslation = GetSamplePositionInternal(SampleTime - FiniteDelta, bUseCharacterSpaceVelocities ? OriginTime - FiniteDelta : OriginTime, SchemaSampleBoneIdx, SchemaOriginBoneIdx, SampleRole, OriginRole);
+	const FVector CurrentTranslation = GetSamplePositionInternal(SampleTime, OriginTime, SchemaSampleBoneIdx, SchemaOriginBoneIdx, SampleRole, OriginRole);
 
 	const FVector LinearVelocity = (CurrentTranslation - PreviousTranslation) / FiniteDelta;
 	return LinearVelocity;
 }
 
-FTransform FSearchContext::GetWorldRootBoneTransformAtTime(float SampleTime) const
+FTransform FSearchContext::GetWorldRootBoneTransformAtTime(float SampleTime, const FRole& SampleRole) const
 {
 	check(!CachedQueries.IsEmpty());
 	const UPoseSearchSchema* Schema = CachedQueries.Last().GetSchema();
 	check(Schema);
 
+	const IPoseHistory* PoseHistory = GetPoseHistory(SampleRole);
 	#if WITH_EDITOR
-	if (!History)
+	if (!PoseHistory)
 	{
 		UE_LOG(LogPoseSearch, Error, TEXT("FSearchContext::GetWorldRootBoneTransformAtTime - Couldn't search for world space root boneTransform by %s, because no IPoseHistory has been found!"), *Schema->GetName());
 	}
@@ -457,12 +525,13 @@ FTransform FSearchContext::GetWorldRootBoneTransformAtTime(float SampleTime) con
 	#endif // WITH_EDITOR
 	{
 		FTransform WorldRootBoneTransform;
-		if (ensure(History) && History->GetTransformAtTime(SampleTime, WorldRootBoneTransform, Schema->Skeleton, RootBoneIndexType, WorldSpaceIndexType))
+		if (ensure(PoseHistory) && PoseHistory->GetTransformAtTime(SampleTime, WorldRootBoneTransform, Schema->GetSkeleton(SampleRole), RootBoneIndexType, WorldSpaceIndexType))
 		{
 			return WorldRootBoneTransform;
 		}
 	}
 
+	const UAnimInstance* AnimInstance = GetAnimInstance(SampleRole);
 	if (AnimInstance && AnimInstance->CurrentSkeleton)
 	{
 		const FTransform& RootBoneTransform = AnimInstance->CurrentSkeleton->GetReferenceSkeleton().GetRefBonePose()[RootSchemaBoneIdx];
@@ -473,54 +542,79 @@ FTransform FSearchContext::GetWorldRootBoneTransformAtTime(float SampleTime) con
 	return FTransform::Identity;
 }
 
-FTransform FSearchContext::GetWorldBoneTransformAtTime(float SampleTime, int8 SchemaBoneIdx)
+bool FSearchContext::ArePoseHistoriesValid() const
+{
+	for (const IPoseHistory* PoseHistory : PoseHistories)
+	{
+		if (!PoseHistory)
+		{
+			return false;
+		}
+	}
+	return true;
+}
+
+FTransform FSearchContext::GetWorldBoneTransformAtTime(float SampleTime, const FRole& SampleRole, int8 SchemaBoneIdx)
 {
 	// CachedQueries.Last is the query we're building 
 	check(!CachedQueries.IsEmpty());
 	const UPoseSearchSchema* Schema = CachedQueries.Last().GetSchema();
 	check(Schema);
 
-	const FBoneIndexType BoneIndexType = Schema->GetBoneIndexType(SchemaBoneIdx);
-	if (const FCachedTransform<FTransform>* CachedTransform = CachedTransforms.Find(SampleTime, BoneIndexType))
+	TConstArrayView<FBoneReference> BoneReferences = Schema->GetBoneReferences(SampleRole);
+	check(BoneReferences[SchemaBoneIdx].HasValidSetup());
+	const FBoneIndexType BoneIndexType = BoneReferences[SchemaBoneIdx].BoneIndex;
+
+	const uint32 SampleTimeHash = GetTypeHash(SampleTime);
+	const uint32 SampleRoleHash = GetTypeHash(SampleRole);
+	const uint32 SampleTimeAndRoleHash = HashCombineFast(SampleTimeHash, SampleRoleHash);
+	const uint32 BoneIndexTypeHash = GetTypeHash(BoneIndexType);
+	const uint32 BoneCachedTransformKey = HashCombineFast(SampleTimeAndRoleHash, BoneIndexTypeHash);
+
+	if (const FTransform* CachedTransform = CachedTransforms.Find(BoneCachedTransformKey))
 	{
-		return CachedTransform->Transform;
+		return *CachedTransform;
 	}
 
 	FTransform WorldBoneTransform;
 	if (BoneIndexType == RootBoneIndexType)
 	{
 		// we already tried querying the CachedTransforms so, let's search in Trajectory
-		WorldBoneTransform = GetWorldRootBoneTransformAtTime(SampleTime);
+		WorldBoneTransform = GetWorldRootBoneTransformAtTime(SampleTime, SampleRole);
 	}
 	else // if (BoneIndexType != RootBoneIndexType)
 	{
 		// searching for RootBoneIndexType in CachedTransforms
-		if (const FCachedTransform<FTransform>* CachedTransform = CachedTransforms.Find(SampleTime, RootBoneIndexType))
+		static const uint32 RootBoneIndexTypeHash = GetTypeHash(RootBoneIndexType); // Note: static const, since RootBoneIndexType is a constant
+		const uint32 RootBoneCachedTransformKey = HashCombineFast(SampleTimeAndRoleHash, RootBoneIndexTypeHash);
+		if (const FTransform* CachedTransform = CachedTransforms.Find(RootBoneCachedTransformKey))
 		{
-			WorldBoneTransform = CachedTransform->Transform;
+			WorldBoneTransform = *CachedTransform;
 		}
 		else
 		{
-			WorldBoneTransform = GetWorldRootBoneTransformAtTime(SampleTime);
+			WorldBoneTransform = GetWorldRootBoneTransformAtTime(SampleTime, SampleRole);
 		}
 
 		// collecting the local bone transforms from the IPoseHistory
+		const IPoseHistory* PoseHistory = GetPoseHistory(SampleRole);
 		#if WITH_EDITOR
-		if (!History)
+		if (!PoseHistory)
 		{
 			UE_LOG(LogPoseSearch, Error, TEXT("FSearchContext::GetWorldBoneTransformAtTime - Couldn't search for bones requested by %s, because no IPoseHistory has been found!"), *Schema->GetName());
 		}
 		else
 		#endif // WITH_EDITOR
 		{
-			check(History);
+			check(PoseHistory);
 
+			const USkeleton* Skeleton = Schema->GetSkeleton(SampleRole);
 			FTransform LocalBoneTransform;
-			if (!History->GetTransformAtTime(SampleTime, LocalBoneTransform, Schema->Skeleton, BoneIndexType, RootBoneIndexType))
+			if (!PoseHistory->GetTransformAtTime(SampleTime, LocalBoneTransform, Skeleton, BoneIndexType, RootBoneIndexType))
 			{
-				if (const USkeleton* Skeleton = Schema->Skeleton)
+				if (Skeleton)
 				{
-					if (!History->IsEmpty())
+					if (!PoseHistory->IsEmpty())
 					{
 						UE_LOG(LogPoseSearch, Warning, TEXT("FSearchContext::GetWorldBoneTransformAtTime - Couldn't find BoneIndexType %d (%s) requested by %s"), BoneIndexType, *Skeleton->GetReferenceSkeleton().GetBoneName(BoneIndexType).ToString(), *Schema->GetName());
 					}
@@ -535,55 +629,55 @@ FTransform FSearchContext::GetWorldBoneTransformAtTime(float SampleTime, int8 Sc
 		}
 	}
 
-	CachedTransforms.Add(SampleTime, BoneIndexType, WorldBoneTransform);
+	CachedTransforms.Add(BoneCachedTransformKey) = WorldBoneTransform;
 	return WorldBoneTransform;
 }
 
-FVector FSearchContext::GetSamplePositionInternal(float SampleTime, float OriginTime, int8 SchemaSampleBoneIdx, int8 SchemaOriginBoneIdx, const FVector* SampleBonePositionWorldOverride)
+FVector FSearchContext::GetSamplePositionInternal(float SampleTime, float OriginTime, int8 SchemaSampleBoneIdx, int8 SchemaOriginBoneIdx, const FRole& SampleRole, const FRole& OriginRole, const FVector* SampleBonePositionWorldOverride)
 {
 	if (SampleBonePositionWorldOverride)
 	{
-		const FTransform RootBoneTransform = GetWorldBoneTransformAtTime(OriginTime, RootSchemaBoneIdx);
+		const FTransform RootBoneTransform = GetWorldBoneTransformAtTime(OriginTime, OriginRole, RootSchemaBoneIdx);
 		if (SchemaOriginBoneIdx == RootSchemaBoneIdx)
 		{
 			return RootBoneTransform.InverseTransformPosition(*SampleBonePositionWorldOverride);
 		}
 
 		// @todo: validate this still works for when root bone is not Identity
-		const FTransform OriginBoneTransform = GetWorldBoneTransformAtTime(OriginTime, SchemaOriginBoneIdx);
+		const FTransform OriginBoneTransform = GetWorldBoneTransformAtTime(OriginTime, OriginRole, SchemaOriginBoneIdx);
 		const FVector DeltaBoneTranslation = *SampleBonePositionWorldOverride - OriginBoneTransform.GetTranslation();
 		return RootBoneTransform.InverseTransformVector(DeltaBoneTranslation);
 	}
 
-	const FTransform RootBoneTransform = GetWorldBoneTransformAtTime(OriginTime, RootSchemaBoneIdx);
-	const FTransform SampleBoneTransform = GetWorldBoneTransformAtTime(SampleTime, SchemaSampleBoneIdx);
+	const FTransform RootBoneTransform = GetWorldBoneTransformAtTime(OriginTime, OriginRole, RootSchemaBoneIdx);
+	const FTransform SampleBoneTransform = GetWorldBoneTransformAtTime(SampleTime, SampleRole, SchemaSampleBoneIdx);
 	if (SchemaOriginBoneIdx == RootSchemaBoneIdx)
 	{
 		return RootBoneTransform.InverseTransformPosition(SampleBoneTransform.GetTranslation());
 	}
 
-	const FTransform OriginBoneTransform = GetWorldBoneTransformAtTime(OriginTime, SchemaOriginBoneIdx);
+	const FTransform OriginBoneTransform = GetWorldBoneTransformAtTime(OriginTime, OriginRole, SchemaOriginBoneIdx);
 	const FVector DeltaBoneTranslation = SampleBoneTransform.GetTranslation() - OriginBoneTransform.GetTranslation();
 	return RootBoneTransform.InverseTransformVector(DeltaBoneTranslation);
 }
 
-FQuat FSearchContext::GetSampleRotationInternal(float SampleTime, float OriginTime, int8 SchemaSampleBoneIdx, int8 SchemaOriginBoneIdx, const FQuat* SampleBoneRotationWorldOverride)
+FQuat FSearchContext::GetSampleRotationInternal(float SampleTime, float OriginTime, int8 SchemaSampleBoneIdx, int8 SchemaOriginBoneIdx, const FRole& SampleRole, const FRole& OriginRole, const FQuat* SampleBoneRotationWorldOverride)
 {
 	if (SampleBoneRotationWorldOverride)
 	{
-		const FTransform RootBoneTransform = GetWorldBoneTransformAtTime(OriginTime, RootSchemaBoneIdx);
+		const FTransform RootBoneTransform = GetWorldBoneTransformAtTime(OriginTime, OriginRole, RootSchemaBoneIdx);
 		if (SchemaOriginBoneIdx == RootSchemaBoneIdx)
 		{
 			return RootBoneTransform.InverseTransformRotation(*SampleBoneRotationWorldOverride);
 		}
 
-		const FTransform OriginBoneTransform = GetWorldBoneTransformAtTime(OriginTime, SchemaOriginBoneIdx);
+		const FTransform OriginBoneTransform = GetWorldBoneTransformAtTime(OriginTime, OriginRole, SchemaOriginBoneIdx);
 		const FQuat DeltaBoneRotation = OriginBoneTransform.InverseTransformRotation(*SampleBoneRotationWorldOverride);
 		return RootBoneTransform.InverseTransformRotation(DeltaBoneRotation);
 	}
 
-	const FTransform RootBoneTransform = GetWorldBoneTransformAtTime(OriginTime, RootSchemaBoneIdx);
-	const FTransform SampleBoneTransform = GetWorldBoneTransformAtTime(SampleTime, SchemaSampleBoneIdx);
+	const FTransform RootBoneTransform = GetWorldBoneTransformAtTime(OriginTime, OriginRole, RootSchemaBoneIdx);
+	const FTransform SampleBoneTransform = GetWorldBoneTransformAtTime(SampleTime, SampleRole, SchemaSampleBoneIdx);
 	return RootBoneTransform.InverseTransformRotation(SampleBoneTransform.GetRotation());
 }
 
@@ -645,7 +739,7 @@ TConstArrayView<float> FSearchContext::GetOrBuildQuery(const UPoseSearchSchema*
 {
 	QUICK_SCOPE_CYCLE_COUNTER(STAT_PoseSearch_GetOrBuildQuery);
 
-	check(Schema && Schema->IsValid());
+	check(Schema);
 	if (const FCachedQuery* FoundCachedQuery = CachedQueries.FindByPredicate(
 		[Schema](const FCachedQuery& CachedQuery)
 		{
diff --git a/Engine/Plugins/Experimental/Animation/PoseSearch/Source/Runtime/Private/PoseSearchDatabase.cpp b/Engine/Plugins/Experimental/Animation/PoseSearch/Source/Runtime/Private/PoseSearchDatabase.cpp
--- a/Engine/Plugins/Experimental/Animation/PoseSearch/Source/Runtime/Private/PoseSearchDatabase.cpp
+++ b/Engine/Plugins/Experimental/Animation/PoseSearch/Source/Runtime/Private/PoseSearchDatabase.cpp
@@ -13,6 +13,7 @@
 #include "PoseSearch/PoseSearchDerivedData.h"
 #include "PoseSearch/PoseSearchFeatureChannel.h"
 #include "PoseSearch/PoseSearchHistory.h"
+#include "PoseSearch/PoseSearchMultiSequence.h"
 #include "PoseSearch/PoseSearchSchema.h"
 #include "Serialization/ArchiveCountMem.h"
 #include "UObject/ObjectSaveContext.h"
@@ -40,7 +41,7 @@ static TAutoConsoleVariable<bool> CVarMotionMatchCompareAgainstBruteForce(TEXT("
 static TAutoConsoleVariable<bool> CVarMotionMatchValidateKNNSearch(TEXT("a.MotionMatch.ValidateKNNSearch"), false, TEXT("Validate KNN search"));
 #endif
 
-typedef TArray<int32, TInlineAllocator<256>> FSelectableAssetIdx;
+typedef TArray<int32, TInlineAllocator<256, TMemStackAllocator<>>> FSelectableAssetIdx;
 static void PopulateSelectableAssetIdx(FSelectableAssetIdx& SelectableAssetIdx, TConstArrayView<const UObject*> AssetsToConsider, const UPoseSearchDatabase* Database)
 {
 	check(Database);
@@ -62,7 +63,7 @@ static void PopulateSelectableAssetIdx(FSelectableAssetIdx& SelectableAssetIdx,
 	}
 }
 
-typedef TArray<int32, TInlineAllocator<256>> FNonSelectableIdx;
+typedef TArray<int32, TInlineAllocator<256, TMemStackAllocator<>>> FNonSelectableIdx;
 static void PopulateNonSelectableIdx(FNonSelectableIdx& NonSelectableIdx, FSearchContext& SearchContext, const UPoseSearchDatabase* Database
 #if UE_POSE_SEARCH_TRACE_ENABLED
 	, TConstArrayView<float> QueryValues
@@ -301,7 +302,7 @@ private:
 	FSelectableAssetIdxFilter SelectableAssetIdxFilter;
 	FBlockTransitionFilter BlockTransitionFilter;
 
-	TArray<const IPoseSearchFilter*, TInlineAllocator<64>> Filters;
+	TArray<const IPoseSearchFilter*, TInlineAllocator<64, TMemStackAllocator<>>> Filters;
 };
 
 } // namespace UE::PoseSearch
@@ -320,6 +321,29 @@ float FPoseSearchDatabaseAnimationAssetBase::GetPlayLength() const
 	return 0;
 }
 
+#if WITH_EDITOR
+int32 FPoseSearchDatabaseAnimationAssetBase::GetFrameAtTime(float Time) const
+{
+	if (const UAnimSequenceBase* SequenceBase = Cast<UAnimSequenceBase>(GetAnimationAsset()))
+	{
+		return SequenceBase->GetFrameAtTime(Time);
+	}
+	return 0.f;
+}
+#endif // WITH_EDITOR
+
+UAnimationAsset* FPoseSearchDatabaseAnimationAssetBase::GetAnimationAssetForRole(const UE::PoseSearch::FRole& Role) const
+{
+	check(GetNumRoles() == 1);
+	return CastChecked<UAnimationAsset>(GetAnimationAsset());
+}
+
+const FTransform& FPoseSearchDatabaseAnimationAssetBase::GetRootTransformOriginForRole(const UE::PoseSearch::FRole& Role) const
+{
+	check(GetNumRoles() == 1);
+	return FTransform::Identity;
+}
+
 #if WITH_EDITORONLY_DATA
 int64 FPoseSearchDatabaseAnimationAssetBase::GetEditorMemSize() const
 {
@@ -377,6 +401,14 @@ UObject* FPoseSearchDatabaseBlendSpace::GetAnimationAsset() const
 	return BlendSpace.Get();
 }
 
+#if WITH_EDITOR
+int32 FPoseSearchDatabaseBlendSpace::GetFrameAtTime(float Time) const
+{
+	// @todo: implement GetFrameAtTime logic for BlendSpace
+	return 0;
+}
+#endif // WITH_EDITOR
+
 #if WITH_EDITORONLY_DATA
 UClass* FPoseSearchDatabaseBlendSpace::GetAnimationAssetStaticClass() const
 {
@@ -545,6 +577,70 @@ bool FPoseSearchDatabaseAnimMontage::IsRootMotionEnabled() const
 }
 #endif // WITH_EDITORONLY_DATA
 
+//////////////////////////////////////////////////////////////////////////
+// FPoseSearchDatabaseMultiSequence
+UObject* FPoseSearchDatabaseMultiSequence::GetAnimationAsset() const
+{
+	return MultiSequence.Get();
+}
+
+float FPoseSearchDatabaseMultiSequence::GetPlayLength() const
+{
+	return MultiSequence ? MultiSequence->GetPlayLength() : 0.f;
+}
+
+#if WITH_EDITOR
+int32 FPoseSearchDatabaseMultiSequence::GetFrameAtTime(float Time) const
+{
+	return MultiSequence ? MultiSequence->GetFrameAtTime(Time) : 0;
+}
+#endif // WITH_EDITOR
+
+int32 FPoseSearchDatabaseMultiSequence::GetNumRoles() const
+{
+	return MultiSequence ? MultiSequence->GetNumRoles() : 0;
+}
+
+UE::PoseSearch::FRole FPoseSearchDatabaseMultiSequence::GetRole(int32 RoleIndex) const
+{
+	return MultiSequence ? MultiSequence->GetRole(RoleIndex) : UE::PoseSearch::DefaultRole;
+}
+
+UAnimationAsset* FPoseSearchDatabaseMultiSequence::GetAnimationAssetForRole(const UE::PoseSearch::FRole& Role) const
+{
+	return MultiSequence ? MultiSequence->GetSequence(Role) : nullptr;
+}
+
+const FTransform& FPoseSearchDatabaseMultiSequence::GetRootTransformOriginForRole(const UE::PoseSearch::FRole& Role) const
+{
+	return MultiSequence ? MultiSequence->GetOrigin(Role) : FTransform::Identity;
+}
+
+#if WITH_EDITORONLY_DATA
+UClass* FPoseSearchDatabaseMultiSequence::GetAnimationAssetStaticClass() const
+{
+	return UPoseSearchMultiSequence::StaticClass();
+}
+
+bool FPoseSearchDatabaseMultiSequence::IsLooping() const
+{
+	return MultiSequence &&
+		MultiSequence->IsLooping() &&
+		SamplingRange.Min == 0.f &&
+		SamplingRange.Max == 0.f;
+}
+
+const FString FPoseSearchDatabaseMultiSequence::GetName() const
+{
+	return MultiSequence ? MultiSequence->GetName() : FString();
+}
+
+bool FPoseSearchDatabaseMultiSequence::IsRootMotionEnabled() const
+{
+	return MultiSequence ? MultiSequence->HasRootMotion() : false;
+}
+#endif // WITH_EDITORONLY_DATA
+
 //////////////////////////////////////////////////////////////////////////
 // UPoseSearchDatabase
 UPoseSearchDatabase::~UPoseSearchDatabase()
@@ -560,7 +656,7 @@ void UPoseSearchDatabase::SetSearchIndex(const UE::PoseSearch::FSearchIndex& Sea
 const UE::PoseSearch::FSearchIndex& UPoseSearchDatabase::GetSearchIndex() const
 {
 	// making sure the search index is consistent. if it fails the calling code hasn't been protected by FAsyncPoseSearchDatabasesManagement::RequestAsyncBuildIndex
-	check(Schema && Schema->IsValid() && !SearchIndexPrivate.IsEmpty() && SearchIndexPrivate.GetNumDimensions() == Schema->SchemaCardinality);
+	check(Schema && !SearchIndexPrivate.IsEmpty() && SearchIndexPrivate.GetNumDimensions() == Schema->SchemaCardinality);
 	return SearchIndexPrivate;
 }
 
@@ -569,6 +665,16 @@ int32 UPoseSearchDatabase::GetPoseIndexFromTime(float Time, const UE::PoseSearch
 	return SearchIndexAsset.GetPoseIndexFromTime(Time, Schema->SampleRate);
 }
 
+void UPoseSearchDatabase::AddAnimationAsset(FInstancedStruct AnimationAsset)
+{
+	AnimationAssets.Add(AnimationAsset);
+}
+
+void UPoseSearchDatabase::RemoveAnimationAssetAt(int32 AnimationAssetIndex)
+{
+	AnimationAssets.RemoveAt(AnimationAssetIndex);
+}
+
 const FInstancedStruct& UPoseSearchDatabase::GetAnimationAssetStruct(int32 AnimationAssetIndex) const
 {
 	check(AnimationAssets.IsValidIndex(AnimationAssetIndex));
@@ -705,18 +811,17 @@ void UPoseSearchDatabase::SynchronizeWithExternalDependencies()
 
 void UPoseSearchDatabase::SynchronizeWithExternalDependencies(TConstArrayView<UAnimSequenceBase*> SequencesBase)
 {
-	// @todo: improve the logic to identify if the asset is actually modified (this method could potentially remove and readd exatly the same AnimationAssets)
-	bool bModified = false;
+	// @todo: use a TSet to speed up this method
+	TArray<FInstancedStruct> AnimationAssetsCopy = AnimationAssets;
 
-	// removing all the SequencesBase references from the database
-	for (int32 AnimationAssetIndex = AnimationAssets.Num() - 1; AnimationAssetIndex >= 0; --AnimationAssetIndex)
+	// removing all the SequencesBase references from AnimationAssetsCopy
+	for (int32 AnimationAssetIndex = AnimationAssetsCopy.Num() - 1; AnimationAssetIndex >= 0; --AnimationAssetIndex)
 	{
-		if (const FPoseSearchDatabaseAnimationAssetBase* AnimationAssetBase = AnimationAssets[AnimationAssetIndex].GetPtr<FPoseSearchDatabaseAnimationAssetBase>())
+		if (const FPoseSearchDatabaseAnimationAssetBase* AnimationAssetBase = AnimationAssetsCopy[AnimationAssetIndex].GetPtr<FPoseSearchDatabaseAnimationAssetBase>())
 		{
 			if (AnimationAssetBase->bSynchronizeWithExternalDependency && SequencesBase.Contains(AnimationAssetBase->GetAnimationAsset()))
 			{
-				AnimationAssets.RemoveAt(AnimationAssetIndex);
-				bModified = true;
+				AnimationAssetsCopy.RemoveAt(AnimationAssetIndex);
 			}
 		}
 	}
@@ -738,8 +843,7 @@ void UPoseSearchDatabase::SynchronizeWithExternalDependencies(TConstArrayView<UA
 							DatabaseSequence.Sequence = Sequence;
 							DatabaseSequence.SamplingRange = FFloatInterval(NotifyEvent.GetTime(), NotifyEvent.GetTime() + NotifyEvent.GetDuration());
 							DatabaseSequence.bSynchronizeWithExternalDependency = true;
-							AnimationAssets.Add(FInstancedStruct::Make(DatabaseSequence));
-							bModified = true;
+							AnimationAssetsCopy.Add(FInstancedStruct::Make(DatabaseSequence));
 						}
 						else if (UAnimComposite* AnimComposite = Cast<UAnimComposite>(SequenceBase))
 						{
@@ -747,8 +851,7 @@ void UPoseSearchDatabase::SynchronizeWithExternalDependencies(TConstArrayView<UA
 							DatabaseAnimComposite.AnimComposite = AnimComposite;
 							DatabaseAnimComposite.SamplingRange = FFloatInterval(NotifyEvent.GetTime(), NotifyEvent.GetTime() + NotifyEvent.GetDuration());
 							DatabaseAnimComposite.bSynchronizeWithExternalDependency = true;
-							AnimationAssets.Add(FInstancedStruct::Make(DatabaseAnimComposite));
-							bModified = true;
+							AnimationAssetsCopy.Add(FInstancedStruct::Make(DatabaseAnimComposite));
 						}
 						else if (UAnimMontage* AnimMontage = Cast<UAnimMontage>(SequenceBase))
 						{
@@ -756,8 +859,7 @@ void UPoseSearchDatabase::SynchronizeWithExternalDependencies(TConstArrayView<UA
 							DatabaseAnimMontage.AnimMontage = AnimMontage;
 							DatabaseAnimMontage.SamplingRange = FFloatInterval(NotifyEvent.GetTime(), NotifyEvent.GetTime() + NotifyEvent.GetDuration());
 							DatabaseAnimMontage.bSynchronizeWithExternalDependency = true;
-							AnimationAssets.Add(FInstancedStruct::Make(DatabaseAnimMontage));
-							bModified = true;
+							AnimationAssetsCopy.Add(FInstancedStruct::Make(DatabaseAnimMontage));
 						}
 					}
 				}
@@ -765,6 +867,22 @@ void UPoseSearchDatabase::SynchronizeWithExternalDependencies(TConstArrayView<UA
 		}
 	}
 
+	// updating AnimationAssets from AnimationAssetsCopy
+	bool bModified = false;
+	for (int32 AnimationAssetIndex = AnimationAssets.Num() - 1; AnimationAssetIndex >= 0; --AnimationAssetIndex)
+	{
+		const int32 CopyIndex = AnimationAssetsCopy.Find(AnimationAssets[AnimationAssetIndex]);
+		if (CopyIndex >= 0)
+		{
+			AnimationAssetsCopy.RemoveAt(CopyIndex);
+		}
+		else
+		{
+			AnimationAssets.RemoveAt(AnimationAssetIndex);
+			bModified = true;
+		}
+	}
+
 	if (bModified)
 	{
 		Modify();
@@ -772,6 +890,21 @@ void UPoseSearchDatabase::SynchronizeWithExternalDependencies(TConstArrayView<UA
 	}
 }
 
+bool UPoseSearchDatabase::Contains(const UObject* Object) const
+{
+	for (const FInstancedStruct& AnimationAsset : AnimationAssets)
+	{
+		if (const FPoseSearchDatabaseAnimationAssetBase* AnimationAssetBase = AnimationAsset.GetPtr<FPoseSearchDatabaseAnimationAssetBase>())
+		{
+			if (AnimationAssetBase->GetAnimationAsset() == Object)
+			{
+				return true;
+			}
+		}
+	}
+	return false;
+}
+
 void UPoseSearchDatabase::BeginCacheForCookedPlatformData(const ITargetPlatform* TargetPlatform)
 {
 	using namespace UE::PoseSearch;
@@ -787,6 +920,20 @@ bool UPoseSearchDatabase::IsCachedCookedPlatformDataLoaded(const ITargetPlatform
 }
 #endif // WITH_EDITOR
 
+#if WITH_EDITOR && ENABLE_ANIM_DEBUG
+void UPoseSearchDatabase::TestSynchronizeWithExternalDependencies()
+{
+	TArray<FInstancedStruct> AnimationAssetsCopy = AnimationAssets;
+	SynchronizeWithExternalDependencies();
+
+	if (AnimationAssetsCopy != AnimationAssets)
+	{
+		UE_LOG(LogPoseSearch, Error, TEXT("TestSynchronizeWithExternalDependencies failed"));
+		AnimationAssets = AnimationAssetsCopy;
+	}
+}
+#endif // WITH_EDITOR && ENABLE_ANIM_DEBUG
+
 void UPoseSearchDatabase::PreSaveRoot(FObjectPreSaveRootContext ObjectSaveContext)
 {
 #if WITH_EDITOR
@@ -981,7 +1128,9 @@ UE::PoseSearch::FSearchResult UPoseSearchDatabase::SearchContinuingPose(UE::Pose
 	const FPoseSearchDatabaseAnimationAssetBase* DatabaseAnimationAssetBase = GetAnimationAssetStruct(SearchIndexAsset).GetPtr<FPoseSearchDatabaseAnimationAssetBase>();
 	check(DatabaseAnimationAssetBase);
 	const UAnimationAsset* AnimationAsset = CastChecked<UAnimationAsset>(DatabaseAnimationAssetBase->GetAnimationAsset());
-	const FAnimationAssetSampler SequenceBaseSampler(AnimationAsset, SearchIndexAsset.GetBlendParameters());
+	
+	// sampler used only to extract the notify states. RootTransformOrigin can be set as Identity, since will not be relevant
+	const FAnimationAssetSampler SequenceBaseSampler(AnimationAsset, FTransform::Identity, SearchIndexAsset.GetBlendParameters());
 	const float SampleTime = GetRealAssetTime(PoseIdx);
 
 	float UpdatedContinuingPoseCostBias = ContinuingPoseCostBias;
