motion matching - multicharacter WIP

diff --git a/Engine/Plugins/Experimental/Animation/PoseSearch/Source/Runtime/Private/PoseSearchDerivedData.cpp b/Engine/Plugins/Experimental/Animation/PoseSearch/Source/Runtime/Private/PoseSearchDerivedData.cpp
--- a/Engine/Plugins/Experimental/Animation/PoseSearch/Source/Runtime/Private/PoseSearchDerivedData.cpp
+++ b/Engine/Plugins/Experimental/Animation/PoseSearch/Source/Runtime/Private/PoseSearchDerivedData.cpp
@@ -19,6 +19,7 @@
 #include "PoseSearch/PoseSearchDefines.h"
 #include "PoseSearch/PoseSearchDerivedDataKey.h"
 #include "PoseSearch/PoseSearchFeatureChannel.h"
+#include "PoseSearch/PoseSearchMultiSequence.h"
 #include "PoseSearch/PoseSearchNormalizationSet.h"
 #include "PoseSearch/PoseSearchSchema.h"
 #include "PoseSearchEigenHelper.h"
@@ -71,6 +72,9 @@ enum EMotionMatchTestFlags
 
 	// validating the data we gave to DDC is stored correctly
 	ValidateDDC = 1 << 10,
+
+	// validating SynchronizeWithExternalDependencies doesn't alter the database AnimationAssets order
+	ValidateSynchronizeWithExternalDependenciesDeterminism = 1 << 11,
 };
 static TAutoConsoleVariable<int32> CVarMotionMatchTestFlags(TEXT("a.MotionMatch.TestFlags"), EMotionMatchTestFlags::None, TEXT("Test Motion Matching using EMotionMatchTestFlags"));
 static TAutoConsoleVariable<int32> CVarMotionMatchTestNumIterations(TEXT("a.MotionMatch.TestNumIterations"), 10, TEXT("Test Motion Matching Num Iterations"));
@@ -230,7 +234,6 @@ public:
 		int32 ThisSchemaIndex = 0;
 		check(SearchIndexBases.Num() == Schemas.Num() && Schemas.Num() > ThisSchemaIndex);
 		const UPoseSearchSchema* ThisSchema = Schemas[ThisSchemaIndex];
-		check(ThisSchema->IsValid());
 		const int32 NumDimensions = ThisSchema->SchemaCardinality;
 
 		TArray<float> MeanDeviations;
@@ -320,7 +323,7 @@ static void InitSearchIndexAssets(FSearchIndexBase& SearchIndex, const UPoseSear
 	TArray<FBlendSampleData> BlendSamples;
 
 	int32 TotalPoses = 0;
-	for (int32 AnimationAssetIndex = 0; AnimationAssetIndex < Database->AnimationAssets.Num(); ++AnimationAssetIndex)
+	for (int32 AnimationAssetIndex = 0; AnimationAssetIndex < Database->GetAnimationAssets().Num(); ++AnimationAssetIndex)
 	{
 		const FInstancedStruct& DatabaseAssetStruct = Database->GetAnimationAssetStruct(AnimationAssetIndex);
 		if (const FPoseSearchDatabaseAnimationAssetBase* DatabaseAsset = DatabaseAssetStruct.GetPtr<FPoseSearchDatabaseAnimationAssetBase>())
@@ -381,6 +384,36 @@ static void InitSearchIndexAssets(FSearchIndexBase& SearchIndex, const UPoseSear
 					}
 				}
 			}
+			else if (const FPoseSearchDatabaseMultiSequence* DatabaseMultiSequence = DatabaseAssetStruct.GetPtr<FPoseSearchDatabaseMultiSequence>())
+			{
+				// @todo: support FindValidSequenceIntervals(SequenceBase, DatabaseAsset->GetSamplingRange(), bIsLooping, Database->ExcludeFromDatabaseParameters, ValidRanges);
+				const float PlayLength = DatabaseMultiSequence->GetPlayLength();
+
+				for (int32 PermutationIdx = 0; PermutationIdx < Database->Schema->NumberOfPermutations; ++PermutationIdx)
+				{
+					if (bAddUnmirrored)
+					{
+						const FSearchIndexAsset PoseSearchIndexAsset(AnimationAssetIndex, TotalPoses, false, bIsLooping, 
+							bDisableReselection, FFloatInterval(0.f, PlayLength), SchemaSampleRate, PermutationIdx);
+						if (PoseSearchIndexAsset.GetNumPoses() > 0)
+						{
+							SearchIndex.Assets.Add(PoseSearchIndexAsset);
+							TotalPoses += PoseSearchIndexAsset.GetNumPoses();
+						}
+					}
+
+					if (bAddMirrored)
+					{
+						const FSearchIndexAsset PoseSearchIndexAsset(AnimationAssetIndex, TotalPoses, true, bIsLooping,
+							bDisableReselection, FFloatInterval(0.f, PlayLength), SchemaSampleRate, PermutationIdx);
+						if (PoseSearchIndexAsset.GetNumPoses() > 0)
+						{
+							SearchIndex.Assets.Add(PoseSearchIndexAsset);
+							TotalPoses += PoseSearchIndexAsset.GetNumPoses();
+						}
+					}
+				}
+			}
 			// support for FPoseSearchDatabaseSequence, FPoseSearchDatabaseAnimComposite, FPoseSearchDatabaseAnimMontage
 			else if (const UAnimSequenceBase* SequenceBase = Cast<UAnimSequenceBase>(DatabaseAsset->GetAnimationAsset()))
 			{
@@ -778,6 +811,36 @@ static void PreprocessSearchIndexVPTree(FSearchIndex& SearchIndex, const UPoseSe
 	}
 }
 
+// @todo: this struct could be replaced by TTuple<const UAnimationAsset*, FTransform, FVector> if FTransform implements operator==
+struct FSamplerMapKey
+{
+	FSamplerMapKey(const UAnimationAsset* InAnimationAsset, const FTransform& InRootTransformOrigin, const FVector& InBlendParameters = FVector::ZeroVector)
+		: AnimationAsset(InAnimationAsset)
+		, RootTransformOrigin(InRootTransformOrigin)
+		, BlendParameters(InBlendParameters)
+	{
+	}
+
+	bool operator==(const FSamplerMapKey& Other) const
+	{
+		return AnimationAsset == Other.AnimationAsset &&
+			RootTransformOrigin.Equals(Other.RootTransformOrigin, 0.f) &&
+			BlendParameters == Other.BlendParameters;
+	}
+
+	friend FORCEINLINE uint32 GetTypeHash(const FSamplerMapKey& SamplerMapKey)
+	{
+		const uint32 AnimationAssetHash = GetTypeHash(SamplerMapKey.AnimationAsset);
+		const uint32 RootTransformOriginHash = GetTypeHash(SamplerMapKey.RootTransformOrigin);
+		const uint32 BlendParametersHash = GetTypeHash(SamplerMapKey.BlendParameters);
+		return HashCombineFast(HashCombineFast(AnimationAssetHash, RootTransformOriginHash), BlendParametersHash);
+	}
+
+	const UAnimationAsset* AnimationAsset = nullptr;
+	FTransform RootTransformOrigin = FTransform::Identity;
+	FVector BlendParameters = FVector::ZeroVector;
+};
+
 static bool IndexDatabase(FSearchIndexBase& SearchIndexBase, const UPoseSearchDatabase& Database, UE::DerivedData::FRequestOwner& Owner)
 {
 	const UPoseSearchSchema* Schema = Database.Schema;
@@ -787,11 +850,12 @@ static bool IndexDatabase(FSearchIndexBase& SearchIndexBase, const UPoseSearchDa
 	TArray<FAnimationAssetSampler> Samplers;
 	Samplers.Reserve(256);
 	
-	TMap<TPair<const UObject*, FVector>, int32> SamplerMap;
+	TMap<FSamplerMapKey, int32> SamplerMap;
 	SamplerMap.Reserve(256);
 
-	for (const FInstancedStruct& DatabaseAssetStruct : Database.AnimationAssets)
+	for (int32 AnimationAssetIndex = 0; AnimationAssetIndex < Database.GetAnimationAssets().Num(); ++AnimationAssetIndex)
 	{
+		const FInstancedStruct& DatabaseAssetStruct = Database.GetAnimationAssetStruct(AnimationAssetIndex);
 		if (const FPoseSearchDatabaseBlendSpace* DatabaseBlendSpace = DatabaseAssetStruct.GetPtr<FPoseSearchDatabaseBlendSpace>())
 		{
 			if (DatabaseBlendSpace->BlendSpace)
@@ -805,33 +869,40 @@ static bool IndexDatabase(FSearchIndexBase& SearchIndexBase, const UPoseSearchDa
 					{
 						const FVector BlendParameters = DatabaseBlendSpace->BlendParameterForSampleRanges(HorizontalIndex, VerticalIndex);
 
-						if (!SamplerMap.Contains({ DatabaseBlendSpace->BlendSpace, BlendParameters }))
+						check(DatabaseBlendSpace->GetNumRoles() == 1);
+						const FTransform& RootTransformOrigin = DatabaseBlendSpace->GetRootTransformOriginForRole(DatabaseBlendSpace->GetRole(0));
+						const FSamplerMapKey SamplerMapKey(DatabaseBlendSpace->BlendSpace, RootTransformOrigin, BlendParameters);
+						if (!SamplerMap.Contains(SamplerMapKey))
 						{
-							SamplerMap.Add({ DatabaseBlendSpace->BlendSpace, BlendParameters }, Samplers.Num());
-							Samplers.Emplace(DatabaseBlendSpace->BlendSpace, BlendParameters);
+							SamplerMap.Add(SamplerMapKey, Samplers.Num());
+							Samplers.Emplace(DatabaseBlendSpace->BlendSpace, RootTransformOrigin, BlendParameters);
 						}
 					}
 				}
 			}
 		}
-		else if (const FPoseSearchDatabaseAnimationAssetBase* DatabaseAnimationAssetBase = DatabaseAssetStruct.GetPtr<FPoseSearchDatabaseAnimationAssetBase>())
+		if (const FPoseSearchDatabaseAnimationAssetBase* DatabaseAnimationAssetBase = DatabaseAssetStruct.GetPtr<FPoseSearchDatabaseAnimationAssetBase>())
 		{
 			if (const UObject* AnimationAssetObject = DatabaseAnimationAssetBase->GetAnimationAsset())
 			{
-				const UAnimationAsset* AnimationAsset = CastChecked<UAnimationAsset>(AnimationAssetObject);
-				if (!SamplerMap.Contains({ AnimationAsset, FVector::ZeroVector }))
+				const int32 NumRoles = DatabaseAnimationAssetBase->GetNumRoles();
+				for (int32 RoleIndex = 0; RoleIndex < NumRoles; ++RoleIndex)
 				{
-					SamplerMap.Add({ AnimationAsset, FVector::ZeroVector }, Samplers.Num());
-					Samplers.Emplace(AnimationAsset);
+					const FRole& Role = DatabaseAnimationAssetBase->GetRole(RoleIndex);
+					const UAnimationAsset* AnimationAsset = DatabaseAnimationAssetBase->GetAnimationAssetForRole(Role);
+					const FTransform& RootTransformOrigin = DatabaseAnimationAssetBase->GetRootTransformOriginForRole(Role);
+					const FSamplerMapKey SamplerMapKey(AnimationAsset, RootTransformOrigin);
+					if (!SamplerMap.Contains(SamplerMapKey))
+					{
+						SamplerMap.Add(SamplerMapKey, Samplers.Num());
+						Samplers.Emplace(AnimationAsset, RootTransformOrigin);
+					}
 				}
 			}
 		}
 	}
 
-	FBoneContainer BoneContainer;
-	BoneContainer.InitializeTo(Schema->BoneIndicesWithParents, UE::Anim::FCurveFilterSettings(UE::Anim::ECurveFilterMode::DisallowAll), *Schema->Skeleton);
-	ParallelFor(Samplers.Num(), [&Samplers, &BoneContainer](int32 SamplerIdx) { Samplers[SamplerIdx].Process(BoneContainer); }, ParallelForFlags);
-
+	ParallelFor(Samplers.Num(), [&Samplers](int32 SamplerIdx) { Samplers[SamplerIdx].Process(); }, ParallelForFlags);
 	if (Owner.IsCanceled())
 	{
 		return false;
@@ -840,7 +911,23 @@ static bool IndexDatabase(FSearchIndexBase& SearchIndexBase, const UPoseSearchDa
 	// prepare indexers
 	TArray<FAssetIndexer> Indexers;
 	Indexers.Reserve(SearchIndexBase.Assets.Num());
-	FAssetSamplingContext SamplingContext(Database, BoneContainer);
+
+	TMap<FRole, FBoneContainer> RoledBoneContainers;
+	Schema->InitBoneContainersFromRoledSkeleton(RoledBoneContainers);
+
+	TMap<FRole, FMirrorDataCache> RoledMirrorDataCaches;
+	RoledMirrorDataCaches.Reserve(RoledBoneContainers.Num());
+	for (const TPair<FRole, FBoneContainer>& RoledBoneContainerPair : RoledBoneContainers)
+	{
+		const FRole& Role = RoledBoneContainerPair.Key;
+		RoledMirrorDataCaches.Add(Role).Init(Database.Schema ? Database.Schema->GetMirrorDataTable(Role) : nullptr, RoledBoneContainers[Role]);
+	}
+
+	const FAssetSamplingContext SamplingContext(Database);
+
+	FAnimationAssetSamplers TempAssetSamplers;
+	TArray<FBoneContainer> TempBoneContainers;
+	FRoleToIndex TempRoleToIndex;
 
 	int32 TotalPoses = 0;
 	for (int32 AssetIdx = 0; AssetIdx != SearchIndexBase.Assets.Num(); ++AssetIdx)
@@ -848,12 +935,39 @@ static bool IndexDatabase(FSearchIndexBase& SearchIndexBase, const UPoseSearchDa
 		FSearchIndexAsset& SearchIndexAsset = SearchIndexBase.Assets[AssetIdx];
 		check(SearchIndexAsset.GetFirstPoseIdx() == TotalPoses);
 
-		const FPoseSearchDatabaseAnimationAssetBase* DatabaseAnimationAssetBase = Database.GetAnimationAssetStruct(SearchIndexAsset).GetPtr<FPoseSearchDatabaseAnimationAssetBase>();
-		check(DatabaseAnimationAssetBase && DatabaseAnimationAssetBase->GetAnimationAsset());
-		const FAnimationAssetSampler& AssetSampler = Samplers[SamplerMap[{ DatabaseAnimationAssetBase->GetAnimationAsset(), SearchIndexAsset.GetBlendParameters() }]];
+		const FInstancedStruct& DatabaseAssetStruct = Database.GetAnimationAssetStruct(SearchIndexAsset);
+		const FPoseSearchDatabaseAnimationAssetBase* DatabaseAnimationAssetBase = DatabaseAssetStruct.GetPtr<FPoseSearchDatabaseAnimationAssetBase>();
+		check(DatabaseAnimationAssetBase);
+
+		TempAssetSamplers.Reset();
+		TempBoneContainers.Reset();
+		TempRoleToIndex.Reset();
+
+		const int32 NumRoles = DatabaseAnimationAssetBase->GetNumRoles();
+		for (int32 RoleIndex = 0; RoleIndex < NumRoles; ++RoleIndex)
+		{
+			const FRole& Role = DatabaseAnimationAssetBase->GetRole(RoleIndex);
+			if (const FMirrorDataCache* MirrorDataCache = RoledMirrorDataCaches.Find(Role))
+			{
+				UAnimationAsset* AnimationAsset = DatabaseAnimationAssetBase->GetAnimationAssetForRole(Role);
+				const FTransform& RootTransformOrigin = DatabaseAnimationAssetBase->GetRootTransformOriginForRole(Role);
+				const FVector BlendParameters = SearchIndexAsset.GetBlendParameters();
+				const int32 SamplerIndex = SamplerMap[{ AnimationAsset, RootTransformOrigin, BlendParameters }];
+				TempAssetSamplers.AnimationAssetSamplers.Emplace(&Samplers[SamplerIndex]);
+				TempAssetSamplers.MirrorDataCaches.Emplace(MirrorDataCache);
+				TempBoneContainers.Emplace(RoledBoneContainers[Role]);
+				TempRoleToIndex.Add(Role) = RoleIndex;
+			}
+			else
+			{
+				// @todo: HANDLE ERRORS / WARNINGS....
+				UE_LOG(LogPoseSearch, Warning, TEXT("Role '%s' from asset '%s' contained in database '%s' cannot be found in Schema::Skeletons '%s'"), *Role.ToString(), *DatabaseAnimationAssetBase->GetName(), *Database.GetName(), *Database.Schema->GetName());
+				return false;
+			}
+		}
 
 		const FFloatInterval ExtrapolationTimeInterval = SearchIndexAsset.GetExtrapolationTimeInterval(Schema->SampleRate, Database.AdditionalExtrapolationTime);
-		Indexers.Emplace(BoneContainer, SearchIndexAsset, SamplingContext, *Schema, AssetSampler, ExtrapolationTimeInterval);
+		Indexers.Emplace(TempBoneContainers, SearchIndexAsset, SamplingContext, *Schema, TempAssetSamplers, TempRoleToIndex, ExtrapolationTimeInterval);
 		TotalPoses += SearchIndexAsset.GetNumPoses();
 	}
 
@@ -1111,6 +1225,10 @@ struct FPoseSearchDatabaseAsyncCacheTask
 	~FPoseSearchDatabaseAsyncCacheTask();
 	EState GetState() const { return EState(ThreadSafeState.GetValue()); }
 
+#if ENABLE_ANIM_DEBUG
+	void TestSynchronizeWithExternalDependencies();
+#endif //ENABLE_ANIM_DEBUG
+
 private:
 	FPoseSearchDatabaseAsyncCacheTask(const FPoseSearchDatabaseAsyncCacheTask& Other) = delete;
 	FPoseSearchDatabaseAsyncCacheTask(FPoseSearchDatabaseAsyncCacheTask&& Other) = delete;
@@ -1169,6 +1287,16 @@ FPoseSearchDatabaseAsyncCacheTask::~FPoseSearchDatabaseAsyncCacheTask()
 	DatabaseDependencies.Reset();
 }
 
+#if ENABLE_ANIM_DEBUG
+void FPoseSearchDatabaseAsyncCacheTask::TestSynchronizeWithExternalDependencies()
+{
+	if (GetState() == EState::Ended && Database != nullptr)
+	{
+		Database->TestSynchronizeWithExternalDependencies();
+	}
+}
+#endif //ENABLE_ANIM_DEBUG
+
 void FPoseSearchDatabaseAsyncCacheTask::StartNewRequestIfNeeded(bool bPerformConditionalPostLoadIfRequired)
 {
 	using namespace UE::DerivedData;
@@ -1301,7 +1429,7 @@ void FPoseSearchDatabaseAsyncCacheTask::Wait(FCriticalSection& OuterMutex)
 	{
 		Database->SetSearchIndex(SearchIndex); // @todo: implement FSearchIndex move ctor and assignment operator and use a MoveTemp(SearchIndex) here
 
-		check(Database->Schema && Database->Schema->IsValid() && !SearchIndex.IsEmpty() && SearchIndex.GetNumDimensions() == Database->Schema->SchemaCardinality);
+		check(Database->Schema && !SearchIndex.IsEmpty() && SearchIndex.GetNumDimensions() == Database->Schema->SchemaCardinality);
 
 		SetState(EState::Ended);
 		bBroadcastOnDerivedDataRebuild = true;
@@ -1353,7 +1481,7 @@ void FPoseSearchDatabaseAsyncCacheTask::OnGetComplete(UE::DerivedData::FCacheGet
 		FMemoryReaderView Reader(RawData);
 		Reader << SearchIndex;
 
-		check(Database != nullptr && Database->Schema && Database->Schema->IsValid());
+		check(Database != nullptr && Database->Schema);
 		// cache can be corrupted in case the version of the derived data cache has not being updated while 
 		// developing channels that changes their cardinality without impacting any asset properties
 		// so to account for this, we just reindex the database and update the associated DDC 
@@ -1431,7 +1559,7 @@ void FPoseSearchDatabaseAsyncCacheTask::OnGetComplete(UE::DerivedData::FCacheGet
 					Schemas[IndexBaseIdx] = IndexBaseDatabase->Schema;
 
 					// early out for invalid indexing conditions
-					if (!IndexBaseDatabase->Schema || !IndexBaseDatabase->Schema->IsValid() || IndexBaseDatabase->Schema->SchemaCardinality <= 0)
+					if (!IndexBaseDatabase->Schema || IndexBaseDatabase->Schema->SchemaCardinality <= 0)
 					{
 						if (IndexBaseIdx == 0)
 						{
@@ -1448,17 +1576,17 @@ void FPoseSearchDatabaseAsyncCacheTask::OnGetComplete(UE::DerivedData::FCacheGet
 						return;
 					}
 
-					// validating that the missing MirrorDataTable is not necessary 
-					if (!IndexBaseDatabase->Schema->MirrorDataTable)
+					// validating that the missing MirrorDataTable(s) are not necessary 
+					if (!IndexBaseDatabase->Schema->AllRoledSkeletonHaveMirrorDataTable())
 					{
-						for (int32 AnimationAssetIndex = 0; AnimationAssetIndex < IndexBaseDatabase->AnimationAssets.Num(); ++AnimationAssetIndex)
+						for (int32 AnimationAssetIndex = 0; AnimationAssetIndex < IndexBaseDatabase->GetAnimationAssets().Num(); ++AnimationAssetIndex)
 						{
 							if (const FPoseSearchDatabaseAnimationAssetBase* DatabaseAsset = IndexBaseDatabase->GetAnimationAssetBase(AnimationAssetIndex))
 							{
 								if (DatabaseAsset->GetMirrorOption() == EPoseSearchMirrorOption::MirroredOnly || DatabaseAsset->GetMirrorOption() == EPoseSearchMirrorOption::UnmirroredAndMirrored)
 								{
 									// want to sample a mirrored asset
-									UE_LOG(LogPoseSearch, Error, TEXT("%s - %s BuildIndex Failed because '%s' requires a MirrorDataTable to sample mirrored animation assets"), *LexToString(FullIndexKey.Hash), *IndexBaseDatabases[0]->GetName(), *IndexBaseDatabase->Schema->GetName());
+									UE_LOG(LogPoseSearch, Error, TEXT("%s - %s BuildIndex Failed because '%s' requires MirrorDataTable(s) to sample mirrored animation assets"), *LexToString(FullIndexKey.Hash), *IndexBaseDatabases[0]->GetName(), *IndexBaseDatabase->Schema->GetName());
 									SearchIndex.Reset();
 #if ENABLE_ANIM_DEBUG
 									SearchIndexCompare.Reset();
@@ -1993,6 +2121,15 @@ void FAsyncPoseSearchDatabasesManagement::Tick(float DeltaTime)
 			Tasks.Reset();
 		}
 	}
+
+	if (AnyTestFlags(EMotionMatchTestFlags::ValidateSynchronizeWithExternalDependenciesDeterminism))
+	{
+		for (int32 TaskIndex = 0; TaskIndex < Tasks.Num(); ++TaskIndex)
+		{
+			Tasks[TaskIndex]->TestSynchronizeWithExternalDependencies();
+		}
+	}
+
 #endif // ENABLE_ANIM_DEBUG
 	
 	// iterating backwards because of the possible RemoveAtSwap 
diff --git a/Engine/Plugins/Experimental/Animation/PoseSearch/Source/Runtime/Private/PoseSearchFeatureChannel.cpp b/Engine/Plugins/Experimental/Animation/PoseSearch/Source/Runtime/Private/PoseSearchFeatureChannel.cpp
--- a/Engine/Plugins/Experimental/Animation/PoseSearch/Source/Runtime/Private/PoseSearchFeatureChannel.cpp
+++ b/Engine/Plugins/Experimental/Animation/PoseSearch/Source/Runtime/Private/PoseSearchFeatureChannel.cpp
@@ -191,7 +191,7 @@ bool UPoseSearchFeatureChannel::CanBeNormalizedWith(const UPoseSearchFeatureChan
 		return false;
 	}
 
-	if (GetSchema()->Skeleton != Other->GetSchema()->Skeleton)
+	if (!GetSchema()->AreSkeletonsCompatible(Other->GetSchema()))
 	{
 		return false;
 	}
@@ -218,6 +218,16 @@ const UPoseSearchSchema* UPoseSearchFeatureChannel::GetSchema() const
 	}
 	return nullptr;
 }
+
+const UE::PoseSearch::FRole UPoseSearchFeatureChannel::GetDefaultRole() const
+{
+	if (const UPoseSearchSchema* Schema = GetSchema())
+	{
+		Schema->GetDefaultRole();
+	}
+	return UE::PoseSearch::DefaultRole;
+}
+
 #endif // WITH_EDITOR
 
 USkeleton* UPoseSearchFeatureChannel::GetSkeleton(bool& bInvalidSkeletonIsError, const IPropertyHandle* PropertyHandle)
@@ -227,7 +237,7 @@ USkeleton* UPoseSearchFeatureChannel::GetSkeleton(bool& bInvalidSkeletonIsError,
 	// blueprint generated classes don't have a schema, until they're instanced by the schema
 	if (const UPoseSearchSchema* Schema = GetSchema())
 	{
-		return Schema->Skeleton;
+		return Schema->GetSkeleton(GetDefaultRole());
 	}
 #else
 	checkNoEntry();
diff --git a/Engine/Plugins/Experimental/Animation/PoseSearch/Source/Runtime/Private/PoseSearchFeatureChannel_FilterCrashingLegs.cpp b/Engine/Plugins/Experimental/Animation/PoseSearch/Source/Runtime/Private/PoseSearchFeatureChannel_FilterCrashingLegs.cpp
--- a/Engine/Plugins/Experimental/Animation/PoseSearch/Source/Runtime/Private/PoseSearchFeatureChannel_FilterCrashingLegs.cpp
+++ b/Engine/Plugins/Experimental/Animation/PoseSearch/Source/Runtime/Private/PoseSearchFeatureChannel_FilterCrashingLegs.cpp
@@ -27,26 +27,28 @@ static float ComputeCrashingLegsValue(const FVector& RightThighPos, const FVecto
 	return CrashingLegsValue;
 }
 
-void UPoseSearchFeatureChannel_FilterCrashingLegs::Finalize(UPoseSearchSchema* Schema)
+bool UPoseSearchFeatureChannel_FilterCrashingLegs::Finalize(UPoseSearchSchema* Schema)
 {
 	ChannelDataOffset = Schema->SchemaCardinality;
 	ChannelCardinality = 1;
 	Schema->SchemaCardinality += ChannelCardinality;
 
-	LeftThighIdx = Schema->AddBoneReference(LeftThigh);
-	RightThighIdx = Schema->AddBoneReference(RightThigh);
-	LeftFootIdx = Schema->AddBoneReference(LeftFoot);
-	RightFootIdx = Schema->AddBoneReference(RightFoot);
+	LeftThighIdx = Schema->AddBoneReference(LeftThigh, SampleRole);
+	RightThighIdx = Schema->AddBoneReference(RightThigh, SampleRole);
+	LeftFootIdx = Schema->AddBoneReference(LeftFoot, SampleRole);
+	RightFootIdx = Schema->AddBoneReference(RightFoot, SampleRole);
+
+	return LeftThighIdx >= 0 && RightThighIdx >= 0 && LeftFootIdx >= 0 && RightFootIdx >= 0;
 }
 
 void UPoseSearchFeatureChannel_FilterCrashingLegs::AddDependentChannels(UPoseSearchSchema* Schema) const
 {
 	if (Schema->bInjectAdditionalDebugChannels)
 	{
-		UPoseSearchFeatureChannel_Position::FindOrAddToSchema(Schema, 0.f, LeftThigh.BoneName);
-		UPoseSearchFeatureChannel_Position::FindOrAddToSchema(Schema, 0.f, RightThigh.BoneName);
-		UPoseSearchFeatureChannel_Position::FindOrAddToSchema(Schema, 0.f, LeftFoot.BoneName);
-		UPoseSearchFeatureChannel_Position::FindOrAddToSchema(Schema, 0.f, RightFoot.BoneName);
+		UPoseSearchFeatureChannel_Position::FindOrAddToSchema(Schema, 0.f, LeftThigh.BoneName, SampleRole);
+		UPoseSearchFeatureChannel_Position::FindOrAddToSchema(Schema, 0.f, RightThigh.BoneName, SampleRole);
+		UPoseSearchFeatureChannel_Position::FindOrAddToSchema(Schema, 0.f, LeftFoot.BoneName, SampleRole);
+		UPoseSearchFeatureChannel_Position::FindOrAddToSchema(Schema, 0.f, RightFoot.BoneName, SampleRole);
 	}
 }
 
@@ -59,11 +61,12 @@ void UPoseSearchFeatureChannel_FilterCrashingLegs::BuildQuery(UE::PoseSearch::FS
 	{
 		// composing a unique identifier to specify this channel with all the required properties to be able to share the query data with other channels of the same type
 		uint32 UniqueIdentifier = GetClass()->GetUniqueID();
-		UniqueIdentifier = HashCombineFast(UniqueIdentifier, ::GetTypeHash(LeftThighIdx));
-		UniqueIdentifier = HashCombineFast(UniqueIdentifier, ::GetTypeHash(RightThighIdx));
-		UniqueIdentifier = HashCombineFast(UniqueIdentifier, ::GetTypeHash(LeftFootIdx));
-		UniqueIdentifier = HashCombineFast(UniqueIdentifier, ::GetTypeHash(RightFootIdx));
-		UniqueIdentifier = HashCombineFast(UniqueIdentifier, ::GetTypeHash(InputQueryPose));
+		UniqueIdentifier = HashCombineFast(UniqueIdentifier, GetTypeHash(SampleRole));
+		UniqueIdentifier = HashCombineFast(UniqueIdentifier, GetTypeHash(LeftThighIdx));
+		UniqueIdentifier = HashCombineFast(UniqueIdentifier, GetTypeHash(RightThighIdx));
+		UniqueIdentifier = HashCombineFast(UniqueIdentifier, GetTypeHash(LeftFootIdx));
+		UniqueIdentifier = HashCombineFast(UniqueIdentifier, GetTypeHash(RightFootIdx));
+		UniqueIdentifier = HashCombineFast(UniqueIdentifier, GetTypeHash(InputQueryPose));
 
 		TConstArrayView<float> CachedChannelData;
 		if (const UPoseSearchFeatureChannel* CachedChannel = SearchContext.GetCachedChannelData(UniqueIdentifier, this, CachedChannelData))
@@ -75,6 +78,7 @@ void UPoseSearchFeatureChannel_FilterCrashingLegs::BuildQuery(UE::PoseSearch::FS
 			check(CachedChannelData.Num() == ChannelCardinality);
 
 			// making sure there were no hash collisions
+			check(CachedFilterCrashingLegs->SampleRole == SampleRole);
 			check(CachedFilterCrashingLegs->LeftThighIdx == LeftThighIdx);
 			check(CachedFilterCrashingLegs->RightThighIdx == RightThighIdx);
 			check(CachedFilterCrashingLegs->LeftFootIdx == LeftFootIdx);
@@ -91,7 +95,7 @@ void UPoseSearchFeatureChannel_FilterCrashingLegs::BuildQuery(UE::PoseSearch::FS
 	// trying to get the BuildQuery data from the continuing pose
 	const bool bCanUseCurrentResult = SearchContext.CanUseCurrentResult();
 	const bool bSkip = InputQueryPose != EInputQueryPose::UseCharacterPose && bCanUseCurrentResult;
-	if (bSkip || !SearchContext.GetHistory())
+	if (bSkip || !SearchContext.ArePoseHistoriesValid())
 	{
 		if (bCanUseCurrentResult)
 		{
@@ -99,16 +103,16 @@ void UPoseSearchFeatureChannel_FilterCrashingLegs::BuildQuery(UE::PoseSearch::FS
 			return;
 		}
 
-		// we leave the SearchContext.EditFeatureVector() set to zero since the SearchContext.History is invalid and it'll fail if we continue
+		// we leave the SearchContext.EditFeatureVector() set to zero since the SearchContext.PoseHistory is invalid and it'll fail if we continue
 		UE_LOG(LogPoseSearch, Error, TEXT("UPoseSearchFeatureChannel_FilterCrashingLegs::BuildQuery - Failed because Pose History Node is missing."));
 		return;
 	}
 
 	// composing the BuildQuery data from the bones
-	const FVector RightThighPosition = SearchContext.GetSamplePosition(0.f, 0.f, RightThighIdx);
-	const FVector LeftThighPosition = SearchContext.GetSamplePosition(0.f, 0.f, LeftThighIdx);
-	const FVector RightFootPosition = SearchContext.GetSamplePosition(0.f, 0.f, RightFootIdx);
-	const FVector LeftFootPosition = SearchContext.GetSamplePosition(0.f, 0.f, LeftFootIdx);
+	const FVector RightThighPosition = SearchContext.GetSamplePosition(0.f, 0.f, RightThighIdx, RootSchemaBoneIdx, SampleRole, SampleRole);
+	const FVector LeftThighPosition = SearchContext.GetSamplePosition(0.f, 0.f, LeftThighIdx, RootSchemaBoneIdx, SampleRole, SampleRole);
+	const FVector RightFootPosition = SearchContext.GetSamplePosition(0.f, 0.f, RightFootIdx, RootSchemaBoneIdx, SampleRole, SampleRole);
+	const FVector LeftFootPosition = SearchContext.GetSamplePosition(0.f, 0.f, LeftFootIdx, RootSchemaBoneIdx, SampleRole, SampleRole);
 
 	const float CrashingLegsValue = ComputeCrashingLegsValue(RightThighPosition, LeftThighPosition, RightFootPosition, LeftFootPosition);
 
@@ -122,10 +126,10 @@ void UPoseSearchFeatureChannel_FilterCrashingLegs::DebugDraw(const UE::PoseSearc
 
 	const float CrashingLegsValue = FFeatureVectorHelper::DecodeFloat(PoseVector, ChannelDataOffset);
 
-	const FVector LeftThighPosition = DrawParams.ExtractPosition(PoseVector, 0.f, LeftThighIdx);
-	const FVector RightThighPosition = DrawParams.ExtractPosition(PoseVector, 0.f, RightThighIdx);
-	const FVector LeftFootPosition = DrawParams.ExtractPosition(PoseVector, 0.f, LeftFootIdx);
-	const FVector RightFootPosition = DrawParams.ExtractPosition(PoseVector, 0.f, RightFootIdx);
+	const FVector LeftThighPosition = DrawParams.ExtractPosition(PoseVector, 0.f, LeftThighIdx, SampleRole);
+	const FVector RightThighPosition = DrawParams.ExtractPosition(PoseVector, 0.f, RightThighIdx, SampleRole);
+	const FVector LeftFootPosition = DrawParams.ExtractPosition(PoseVector, 0.f, LeftFootIdx, SampleRole);
+	const FVector RightFootPosition = DrawParams.ExtractPosition(PoseVector, 0.f, RightFootIdx, SampleRole);
 
 	const float FeetDistance = (RightFootPosition - LeftFootPosition).Length();
 
@@ -189,10 +193,10 @@ bool UPoseSearchFeatureChannel_FilterCrashingLegs::IndexAsset(UE::PoseSearch::FA
 	FVector RightThighPosition, LeftThighPosition, RightFootPosition, LeftFootPosition;
 	for (int32 SampleIdx = Indexer.GetBeginSampleIdx(); SampleIdx != Indexer.GetEndSampleIdx(); ++SampleIdx)
 	{
-		if (Indexer.GetSamplePosition(RightThighPosition, 0.f, 0.f, SampleIdx, RightThighIdx) &&
-			Indexer.GetSamplePosition(LeftThighPosition, 0.f, 0.f, SampleIdx, LeftThighIdx) &&
-			Indexer.GetSamplePosition(RightFootPosition, 0.f, 0.f, SampleIdx, RightFootIdx) &&
-			Indexer.GetSamplePosition(LeftFootPosition, 0.f, 0.f, SampleIdx, LeftFootIdx))
+		if (Indexer.GetSamplePosition(RightThighPosition, 0.f, 0.f, SampleIdx, RightThighIdx, RootSchemaBoneIdx, SampleRole, SampleRole) &&
+			Indexer.GetSamplePosition(LeftThighPosition, 0.f, 0.f, SampleIdx, LeftThighIdx, RootSchemaBoneIdx, SampleRole, SampleRole) &&
+			Indexer.GetSamplePosition(RightFootPosition, 0.f, 0.f, SampleIdx, RightFootIdx, RootSchemaBoneIdx, SampleRole, SampleRole) &&
+			Indexer.GetSamplePosition(LeftFootPosition, 0.f, 0.f, SampleIdx, LeftFootIdx, RootSchemaBoneIdx, SampleRole, SampleRole))
 		{
 			const float CrashingLegsValue = ComputeCrashingLegsValue(RightThighPosition, LeftThighPosition, RightFootPosition, LeftFootPosition);
 			FFeatureVectorHelper::EncodeFloat(Indexer.GetPoseVector(SampleIdx), ChannelDataOffset, CrashingLegsValue);
diff --git a/Engine/Plugins/Experimental/Animation/PoseSearch/Source/Runtime/Private/PoseSearchFeatureChannel_Group.cpp b/Engine/Plugins/Experimental/Animation/PoseSearch/Source/Runtime/Private/PoseSearchFeatureChannel_Group.cpp
--- a/Engine/Plugins/Experimental/Animation/PoseSearch/Source/Runtime/Private/PoseSearchFeatureChannel_Group.cpp
+++ b/Engine/Plugins/Experimental/Animation/PoseSearch/Source/Runtime/Private/PoseSearchFeatureChannel_Group.cpp
@@ -3,17 +3,21 @@
 #include "PoseSearchFeatureChannel_Group.h"
 #include "PoseSearch/PoseSearchSchema.h"
 
-void UPoseSearchFeatureChannel_GroupBase::Finalize(UPoseSearchSchema* Schema)
+bool UPoseSearchFeatureChannel_GroupBase::Finalize(UPoseSearchSchema* Schema)
 {
 	ChannelDataOffset = Schema->SchemaCardinality;
 	for (TObjectPtr<UPoseSearchFeatureChannel>& SubChannelPtr : GetSubChannels())
 	{
 		if (UPoseSearchFeatureChannel* SubChannel = SubChannelPtr.Get())
 		{
-			SubChannel->Finalize(Schema);
+			if (!SubChannel->Finalize(Schema))
+			{
+				return false;
+			}
 		}
 	}
 	ChannelCardinality = Schema->SchemaCardinality - ChannelDataOffset;
+	return true;
 }
 
 void UPoseSearchFeatureChannel_GroupBase::AddDependentChannels(UPoseSearchSchema* Schema) const
diff --git a/Engine/Plugins/Experimental/Animation/PoseSearch/Source/Runtime/Private/PoseSearchFeatureChannel_Heading.cpp b/Engine/Plugins/Experimental/Animation/PoseSearch/Source/Runtime/Private/PoseSearchFeatureChannel_Heading.cpp
--- a/Engine/Plugins/Experimental/Animation/PoseSearch/Source/Runtime/Private/PoseSearchFeatureChannel_Heading.cpp
+++ b/Engine/Plugins/Experimental/Animation/PoseSearch/Source/Runtime/Private/PoseSearchFeatureChannel_Heading.cpp
@@ -1,26 +1,37 @@
 // Copyright Epic Games, Inc. All Rights Reserved.
 
 #include "PoseSearchFeatureChannel_Heading.h"
+#include "Engine/BlueprintGeneratedClass.h"
 #include "PoseSearch/PoseSearchAssetIndexer.h"
 #include "PoseSearch/PoseSearchAssetSampler.h"
 #include "PoseSearch/PoseSearchContext.h"
 #include "PoseSearch/PoseSearchDatabase.h"
 #include "PoseSearch/PoseSearchSchema.h"
 #include "PoseSearchFeatureChannel_Position.h"
-#include "Engine/BlueprintGeneratedClass.h"
+
+#if WITH_EDITOR
+#include "PropertyHandle.h"
+#endif // WITH_EDITOR
 
 UPoseSearchFeatureChannel_Heading::UPoseSearchFeatureChannel_Heading()
 {
 	bUseBlueprintQueryOverride = Cast<UBlueprintGeneratedClass>(GetClass()) != nullptr;
 }
 
-void UPoseSearchFeatureChannel_Heading::FindOrAddToSchema(UPoseSearchSchema* Schema, float SampleTimeOffset, const FName& BoneName, EHeadingAxis HeadingAxis, EPermutationTimeType PermutationTimeType)
+void UPoseSearchFeatureChannel_Heading::FindOrAddToSchema(UPoseSearchSchema* Schema, float SampleTimeOffset, const FName& BoneName, const UE::PoseSearch::FRole& Role, EHeadingAxis HeadingAxis, EPermutationTimeType PermutationTimeType)
 {
-	if (!Schema->FindChannel([SampleTimeOffset, &BoneName, HeadingAxis, PermutationTimeType](const UPoseSearchFeatureChannel* Channel) -> const UPoseSearchFeatureChannel_Heading*
+	if (!Schema->FindChannel([SampleTimeOffset, &BoneName, Role, HeadingAxis, PermutationTimeType](const UPoseSearchFeatureChannel* Channel) -> const UPoseSearchFeatureChannel_Heading*
 		{
 			if (const UPoseSearchFeatureChannel_Heading* Heading = Cast<UPoseSearchFeatureChannel_Heading>(Channel))
 			{
-				if (Heading->Bone.BoneName == BoneName && Heading->OriginBone.BoneName == NAME_None && Heading->SampleTimeOffset == SampleTimeOffset && Heading->OriginTimeOffset == 0.f && Heading->HeadingAxis == HeadingAxis && Heading->PermutationTimeType == PermutationTimeType)
+				if (Heading->Bone.BoneName == BoneName &&
+					Heading->OriginBone.BoneName == NAME_None && 
+					Heading->SampleTimeOffset == SampleTimeOffset &&
+					Heading->OriginTimeOffset == 0.f &&
+					Heading->HeadingAxis == HeadingAxis &&
+					Heading->PermutationTimeType == PermutationTimeType &&
+					Heading->SampleRole == Role &&
+					Heading->OriginRole == Role)
 				{
 					return Heading;
 				}
@@ -30,6 +41,8 @@ void UPoseSearchFeatureChannel_Heading::FindOrAddToSchema(UPoseSearchSchema* Sch
 	{
 		UPoseSearchFeatureChannel_Heading* Heading = NewObject<UPoseSearchFeatureChannel_Heading>(Schema, NAME_None, RF_Transient);
 		Heading->Bone.BoneName = BoneName;
+		Heading->SampleRole = Role;
+		Heading->OriginRole = Role;
 #if WITH_EDITORONLY_DATA
 		Heading->Weight = 0.f;
 		// @todo: perhaps add a tunable color for injected channels
@@ -42,29 +55,31 @@ void UPoseSearchFeatureChannel_Heading::FindOrAddToSchema(UPoseSearchSchema* Sch
 	}
 }
 
-void UPoseSearchFeatureChannel_Heading::Finalize(UPoseSearchSchema* Schema)
+bool UPoseSearchFeatureChannel_Heading::Finalize(UPoseSearchSchema* Schema)
 {
 	ChannelDataOffset = Schema->SchemaCardinality;
 	ChannelCardinality = UE::PoseSearch::FFeatureVectorHelper::GetVectorCardinality(ComponentStripping);
 	Schema->SchemaCardinality += ChannelCardinality;
 
-	SchemaBoneIdx = Schema->AddBoneReference(Bone);
-	SchemaOriginBoneIdx = Schema->AddBoneReference(OriginBone);
+	SchemaBoneIdx = Schema->AddBoneReference(Bone, SampleRole);
+	SchemaOriginBoneIdx = Schema->AddBoneReference(OriginBone, OriginRole);
+
+	return SchemaBoneIdx >= 0 && SchemaOriginBoneIdx >= 0;
 }
 
 void UPoseSearchFeatureChannel_Heading::AddDependentChannels(UPoseSearchSchema* Schema) const
 {
 	if (Schema->bInjectAdditionalDebugChannels)
 	{
-		UPoseSearchFeatureChannel_Position::FindOrAddToSchema(Schema, SampleTimeOffset, Bone.BoneName);
+		UPoseSearchFeatureChannel_Position::FindOrAddToSchema(Schema, SampleTimeOffset, Bone.BoneName, SampleRole);
 		if (!FMath::IsNearlyZero(OriginTimeOffset))
 		{
 			// adding the position OriginTimeOffset seconds ahead
-			UPoseSearchFeatureChannel_Position::FindOrAddToSchema(Schema, OriginTimeOffset, Bone.BoneName);
+			UPoseSearchFeatureChannel_Position::FindOrAddToSchema(Schema, OriginTimeOffset, Bone.BoneName, OriginRole);
 			
 			// adding the rotation (X, Y axis) OriginTimeOffset seconds ahead
-			UPoseSearchFeatureChannel_Heading::FindOrAddToSchema(Schema, OriginTimeOffset, Bone.BoneName, EHeadingAxis::X);
-			UPoseSearchFeatureChannel_Heading::FindOrAddToSchema(Schema, OriginTimeOffset, Bone.BoneName, EHeadingAxis::Y);
+			UPoseSearchFeatureChannel_Heading::FindOrAddToSchema(Schema, OriginTimeOffset, Bone.BoneName, OriginRole, EHeadingAxis::X);
+			UPoseSearchFeatureChannel_Heading::FindOrAddToSchema(Schema, OriginTimeOffset, Bone.BoneName, OriginRole, EHeadingAxis::Y);
 		}
 	}
 }
@@ -92,8 +107,8 @@ void UPoseSearchFeatureChannel_Heading::BuildQuery(UE::PoseSearch::FSearchContex
 	const bool bIsRootBone = SchemaBoneIdx == RootSchemaBoneIdx;
 	if (bUseBlueprintQueryOverride)
 	{
-		const FQuat BoneRotationWorld = BP_GetWorldRotation(SearchContext.GetAnimInstance());
-		const FQuat BoneRotation = SearchContext.GetSampleRotation(SampleTimeOffset, OriginTimeOffset, SchemaBoneIdx, SchemaOriginBoneIdx, EPermutationTimeType::UseSampleTime, &BoneRotationWorld);
+		const FQuat BoneRotationWorld = BP_GetWorldRotation(SearchContext.GetAnimInstance(SampleRole));
+		const FQuat BoneRotation = SearchContext.GetSampleRotation(SampleTimeOffset, OriginTimeOffset, SchemaBoneIdx, SchemaOriginBoneIdx, SampleRole, OriginRole, EPermutationTimeType::UseSampleTime, &BoneRotationWorld);
 		FFeatureVectorHelper::EncodeVector(SearchContext.EditFeatureVector(), ChannelDataOffset, GetAxis(BoneRotation), ComponentStripping);
 		return;
 	}
@@ -103,15 +118,17 @@ void UPoseSearchFeatureChannel_Heading::BuildQuery(UE::PoseSearch::FSearchContex
 	{
 		// composing a unique identifier to specify this channel with all the required properties to be able to share the query data with other channels of the same type
 		uint32 UniqueIdentifier = GetClass()->GetUniqueID();
-		UniqueIdentifier = HashCombineFast(UniqueIdentifier, ::GetTypeHash(SamplingAttributeId));
-		UniqueIdentifier = HashCombineFast(UniqueIdentifier, ::GetTypeHash(SampleTimeOffset));
-		UniqueIdentifier = HashCombineFast(UniqueIdentifier, ::GetTypeHash(OriginTimeOffset));
-		UniqueIdentifier = HashCombineFast(UniqueIdentifier, ::GetTypeHash(HeadingAxis));
-		UniqueIdentifier = HashCombineFast(UniqueIdentifier, ::GetTypeHash(SchemaBoneIdx));
-		UniqueIdentifier = HashCombineFast(UniqueIdentifier, ::GetTypeHash(SchemaOriginBoneIdx));
-		UniqueIdentifier = HashCombineFast(UniqueIdentifier, ::GetTypeHash(InputQueryPose));
-		UniqueIdentifier = HashCombineFast(UniqueIdentifier, ::GetTypeHash(ComponentStripping));
-		UniqueIdentifier = HashCombineFast(UniqueIdentifier, ::GetTypeHash(PermutationTimeType));
+		UniqueIdentifier = HashCombineFast(UniqueIdentifier, GetTypeHash(SampleRole));
+		UniqueIdentifier = HashCombineFast(UniqueIdentifier, GetTypeHash(OriginRole));
+		UniqueIdentifier = HashCombineFast(UniqueIdentifier, GetTypeHash(SamplingAttributeId));
+		UniqueIdentifier = HashCombineFast(UniqueIdentifier, GetTypeHash(SampleTimeOffset));
+		UniqueIdentifier = HashCombineFast(UniqueIdentifier, GetTypeHash(OriginTimeOffset));
+		UniqueIdentifier = HashCombineFast(UniqueIdentifier, GetTypeHash(HeadingAxis));
+		UniqueIdentifier = HashCombineFast(UniqueIdentifier, GetTypeHash(SchemaBoneIdx));
+		UniqueIdentifier = HashCombineFast(UniqueIdentifier, GetTypeHash(SchemaOriginBoneIdx));
+		UniqueIdentifier = HashCombineFast(UniqueIdentifier, GetTypeHash(InputQueryPose));
+		UniqueIdentifier = HashCombineFast(UniqueIdentifier, GetTypeHash(ComponentStripping));
+		UniqueIdentifier = HashCombineFast(UniqueIdentifier, GetTypeHash(PermutationTimeType));
 
 		TConstArrayView<float> CachedChannelData;
 		if (const UPoseSearchFeatureChannel* CachedChannel = SearchContext.GetCachedChannelData(UniqueIdentifier, this, CachedChannelData))
@@ -123,6 +140,8 @@ void UPoseSearchFeatureChannel_Heading::BuildQuery(UE::PoseSearch::FSearchContex
 			check(CachedChannelData.Num() == ChannelCardinality);
 
 			// making sure there were no hash collisions
+			check(CachedHeadingChannel->SampleRole == SampleRole);
+			check(CachedHeadingChannel->OriginRole == OriginRole);
 			check(CachedHeadingChannel->SamplingAttributeId == SamplingAttributeId);
 			check(CachedHeadingChannel->SampleTimeOffset == SampleTimeOffset);
 			check(CachedHeadingChannel->OriginTimeOffset == OriginTimeOffset);
@@ -141,22 +160,22 @@ void UPoseSearchFeatureChannel_Heading::BuildQuery(UE::PoseSearch::FSearchContex
 	}
 
 	const bool bCanUseCurrentResult = SearchContext.CanUseCurrentResult();
-	const bool bSkip = InputQueryPose != EInputQueryPose::UseCharacterPose && bCanUseCurrentResult;
-	if (bSkip || (!SearchContext.GetHistory() && !bIsRootBone))
-		{
+	const bool bSkip = InputQueryPose != EInputQueryPose::UseCharacterPose && bCanUseCurrentResult && SampleRole == OriginRole;
+	if (bSkip || (!SearchContext.ArePoseHistoriesValid() && !bIsRootBone))
+	{
 		if (bCanUseCurrentResult)
-			{
+		{
 			FFeatureVectorHelper::Copy(SearchContext.EditFeatureVector(), ChannelDataOffset, ChannelCardinality, SearchContext.GetCurrentResultPoseVector());
 			return;
-			}
-		
-		// we leave the SearchContext.EditFeatureVector() set to zero since the SearchContext.History is invalid and it'll fail if we continue
-				UE_LOG(LogPoseSearch, Error, TEXT("UPoseSearchFeatureChannel_Heading::BuildQuery - Failed because Pose History Node is missing."));
+		}
+
+		// we leave the SearchContext.EditFeatureVector() set to zero since the SearchContext.PoseHistory is invalid and it'll fail if we continue
+		UE_LOG(LogPoseSearch, Error, TEXT("UPoseSearchFeatureChannel_Heading::BuildQuery - Failed because Pose History Node is missing."));
 		return;
-			}
+	}
 	
-			// calculating the BoneRotation in component space for the bone indexed by SchemaBoneIdx
-	const FQuat BoneRotation = SearchContext.GetSampleRotation(SampleTimeOffset, OriginTimeOffset, SchemaBoneIdx, SchemaOriginBoneIdx, PermutationTimeType);
+	// calculating the BoneRotation in component space for the bone indexed by SchemaBoneIdx
+	const FQuat BoneRotation = SearchContext.GetSampleRotation(SampleTimeOffset, OriginTimeOffset, SchemaBoneIdx, SchemaOriginBoneIdx, SampleRole, OriginRole, PermutationTimeType);
 	FFeatureVectorHelper::EncodeVector(SearchContext.EditFeatureVector(), ChannelDataOffset, GetAxis(BoneRotation), ComponentStripping);
 }
 
@@ -172,8 +191,8 @@ void UPoseSearchFeatureChannel_Heading::DebugDraw(const UE::PoseSearch::FDebugDr
 	Color = FLinearColor::White.ToFColor(true);
 #endif // WITH_EDITORONLY_DATA
 
-	const FVector BoneHeading = DrawParams.ExtractRotation(PoseVector, OriginTimeOffset).RotateVector(FFeatureVectorHelper::DecodeVector(PoseVector, ChannelDataOffset, ComponentStripping));
-	const FVector BonePos = DrawParams.ExtractPosition(PoseVector, SampleTimeOffset, SchemaBoneIdx, PermutationTimeType, SamplingAttributeId);
+	const FVector BoneHeading = DrawParams.ExtractRotation(PoseVector, OriginTimeOffset, RootSchemaBoneIdx, OriginRole).RotateVector(FFeatureVectorHelper::DecodeVector(PoseVector, ChannelDataOffset, ComponentStripping));
+	const FVector BonePos = DrawParams.ExtractPosition(PoseVector, SampleTimeOffset, SchemaBoneIdx, SampleRole, PermutationTimeType, SamplingAttributeId);
 
 	DrawParams.DrawPoint(BonePos, Color, 3.f);
 	DrawParams.DrawLine(BonePos + BoneHeading * 4.f, BonePos + BoneHeading * 15.f, Color);
@@ -196,7 +215,7 @@ bool UPoseSearchFeatureChannel_Heading::IndexAsset(UE::PoseSearch::FAssetIndexer
 	FQuat SampleRotation = FQuat::Identity;
 	for (int32 SampleIdx = Indexer.GetBeginSampleIdx(); SampleIdx != Indexer.GetEndSampleIdx(); ++SampleIdx)
 	{
-		if (Indexer.GetSampleRotation(SampleRotation, SampleTimeOffset, OriginTimeOffset, SampleIdx, SchemaBoneIdx, SchemaOriginBoneIdx, PermutationTimeType, SamplingAttributeId))
+		if (Indexer.GetSampleRotation(SampleRotation, SampleTimeOffset, OriginTimeOffset, SampleIdx, SchemaBoneIdx, SchemaOriginBoneIdx, SampleRole, OriginRole, PermutationTimeType, SamplingAttributeId))
 		{
 			FFeatureVectorHelper::EncodeVector(Indexer.GetPoseVector(SampleIdx), ChannelDataOffset, GetAxis(SampleRotation), ComponentStripping);
 		}
@@ -243,13 +262,27 @@ UE::PoseSearch::TLabelBuilder& UPoseSearchFeatureChannel_Heading::GetLabel(UE::P
 	if (SchemaBoneIdx != RootSchemaBoneIdx)
 	{
 		LabelBuilder.Append(TEXT("_"));
-		LabelBuilder.Append(Schema->BoneReferences[SchemaBoneIdx].BoneName.ToString());
+		LabelBuilder.Append(Schema->GetBoneReferences(SampleRole)[SchemaBoneIdx].BoneName.ToString());
+	}
+
+	if (SampleRole != DefaultRole)
+	{
+		LabelBuilder.Append(TEXT("["));
+		LabelBuilder.Append(SampleRole.ToString());
+		LabelBuilder.Append(TEXT("]"));
 	}
 
 	if (SchemaOriginBoneIdx != RootSchemaBoneIdx)
 	{
 		LabelBuilder.Append(TEXT("_"));
-		LabelBuilder.Append(Schema->BoneReferences[SchemaOriginBoneIdx].BoneName.ToString());
+		LabelBuilder.Append(Schema->GetBoneReferences(OriginRole)[SchemaOriginBoneIdx].BoneName.ToString());
+	}
+
+	if (OriginRole != DefaultRole)
+	{
+		LabelBuilder.Append(TEXT("["));
+		LabelBuilder.Append(OriginRole.ToString());
+		LabelBuilder.Append(TEXT("]"));
 	}
 
 	AppendLabelSeparator(LabelBuilder, LabelFormat, true);
@@ -263,4 +296,19 @@ UE::PoseSearch::TLabelBuilder& UPoseSearchFeatureChannel_Heading::GetLabel(UE::P
 
 	return LabelBuilder;
 }
+
+USkeleton* UPoseSearchFeatureChannel_Heading::GetSkeleton(bool& bInvalidSkeletonIsError, const IPropertyHandle* PropertyHandle)
+{
+	// blueprint generated classes don't have a schema, until they're instanced by the schema
+	if (const UPoseSearchSchema* Schema = GetSchema())
+	{
+		bInvalidSkeletonIsError = false;
+		if (PropertyHandle && PropertyHandle->GetProperty()->GetFName() == GET_MEMBER_NAME_CHECKED(UPoseSearchFeatureChannel_Heading, OriginBone))
+		{
+			return Schema->GetSkeleton(OriginRole);
+		}
+	}
+
+	return Super::GetSkeleton(bInvalidSkeletonIsError, PropertyHandle);
+}
 #endif
\ No newline at end of file
diff --git a/Engine/Plugins/Experimental/Animation/PoseSearch/Source/Runtime/Private/PoseSearchFeatureChannel_Heading.h b/Engine/Plugins/Experimental/Animation/PoseSearch/Source/Runtime/Private/PoseSearchFeatureChannel_Heading.h
--- a/Engine/Plugins/Experimental/Animation/PoseSearch/Source/Runtime/Private/PoseSearchFeatureChannel_Heading.h
+++ b/Engine/Plugins/Experimental/Animation/PoseSearch/Source/Runtime/Private/PoseSearchFeatureChannel_Heading.h
@@ -25,9 +25,15 @@ public:
 	UPROPERTY(EditAnywhere, Category = "Settings")
 	FBoneReference Bone;
 
+	UPROPERTY(EditAnywhere, Category = "Settings")
+	FName SampleRole = UE::PoseSearch::DefaultRole;
+
 	UPROPERTY(EditAnywhere, Category = "Settings")
 	FBoneReference OriginBone;
 
+	UPROPERTY(EditAnywhere, Category = "Settings")
+	FName OriginRole = UE::PoseSearch::DefaultRole;
+
 #if WITH_EDITORONLY_DATA
 	UPROPERTY(EditAnywhere, Category = "Settings")
 	float Weight = 1.f;
@@ -82,7 +88,7 @@ public:
 	UPoseSearchFeatureChannel_Heading();
 
 	// UPoseSearchFeatureChannel interface
-	virtual void Finalize(UPoseSearchSchema* Schema) override;
+	virtual bool Finalize(UPoseSearchSchema* Schema) override;
 	virtual void BuildQuery(UE::PoseSearch::FSearchContext& SearchContext) const override;
 
 	virtual EPermutationTimeType GetPermutationTimeType() const override { return PermutationTimeType; }
@@ -96,9 +102,11 @@ public:
 	virtual void FillWeights(TArrayView<float> Weights) const override;
 	virtual bool IndexAsset(UE::PoseSearch::FAssetIndexer& Indexer) const override;
 	virtual UE::PoseSearch::TLabelBuilder& GetLabel(UE::PoseSearch::TLabelBuilder& LabelBuilder, UE::PoseSearch::ELabelFormat LabelFormat = UE::PoseSearch::ELabelFormat::Full_Horizontal) const override;
+
+	// IBoneReferenceSkeletonProvider interface
+	USkeleton* GetSkeleton(bool& bInvalidSkeletonIsError, const IPropertyHandle* PropertyHandle) override;
 #endif
 
 	FVector GetAxis(const FQuat& Rotation) const;
-
-	static void FindOrAddToSchema(UPoseSearchSchema* Schema, float SampleTimeOffset, const FName& BoneName = NAME_None, EHeadingAxis HeadingAxis = EHeadingAxis::X, EPermutationTimeType PermutationTimeType = EPermutationTimeType::UseSampleTime);
+	static void FindOrAddToSchema(UPoseSearchSchema* Schema, float SampleTimeOffset, const FName& BoneName, const UE::PoseSearch::FRole& Role, EHeadingAxis HeadingAxis = EHeadingAxis::X, EPermutationTimeType PermutationTimeType = EPermutationTimeType::UseSampleTime);
 };
\ No newline at end of file
diff --git a/Engine/Plugins/Experimental/Animation/PoseSearch/Source/Runtime/Private/PoseSearchFeatureChannel_Padding.cpp b/Engine/Plugins/Experimental/Animation/PoseSearch/Source/Runtime/Private/PoseSearchFeatureChannel_Padding.cpp
--- a/Engine/Plugins/Experimental/Animation/PoseSearch/Source/Runtime/Private/PoseSearchFeatureChannel_Padding.cpp
+++ b/Engine/Plugins/Experimental/Animation/PoseSearch/Source/Runtime/Private/PoseSearchFeatureChannel_Padding.cpp
@@ -12,11 +12,12 @@ void UPoseSearchFeatureChannel_Padding::AddToSchema(UPoseSearchSchema* Schema, i
 	Schema->AddTemporaryChannel(Padding);
 }
 
-void UPoseSearchFeatureChannel_Padding::Finalize(UPoseSearchSchema* Schema)
+bool UPoseSearchFeatureChannel_Padding::Finalize(UPoseSearchSchema* Schema)
 {
 	ChannelDataOffset = Schema->SchemaCardinality;
 	ChannelCardinality = PaddingSize;
 	Schema->SchemaCardinality += ChannelCardinality;
+	return true;
 }
 
 void UPoseSearchFeatureChannel_Padding::BuildQuery(UE::PoseSearch::FSearchContext& SearchContext) const
diff --git a/Engine/Plugins/Experimental/Animation/PoseSearch/Source/Runtime/Private/PoseSearchFeatureChannel_Padding.h b/Engine/Plugins/Experimental/Animation/PoseSearch/Source/Runtime/Private/PoseSearchFeatureChannel_Padding.h
--- a/Engine/Plugins/Experimental/Animation/PoseSearch/Source/Runtime/Private/PoseSearchFeatureChannel_Padding.h
+++ b/Engine/Plugins/Experimental/Animation/PoseSearch/Source/Runtime/Private/PoseSearchFeatureChannel_Padding.h
@@ -15,7 +15,7 @@ public:
 	int32 PaddingSize = 1;
 
 	// UPoseSearchFeatureChannel interface
-	virtual void Finalize(UPoseSearchSchema* Schema) override;
+	virtual bool Finalize(UPoseSearchSchema* Schema) override;
 	virtual void BuildQuery(UE::PoseSearch::FSearchContext& SearchContext) const override;
 
 #if WITH_EDITOR
diff --git a/Engine/Plugins/Experimental/Animation/PoseSearch/Source/Runtime/Private/PoseSearchFeatureChannel_PermutationTime.cpp b/Engine/Plugins/Experimental/Animation/PoseSearch/Source/Runtime/Private/PoseSearchFeatureChannel_PermutationTime.cpp
--- a/Engine/Plugins/Experimental/Animation/PoseSearch/Source/Runtime/Private/PoseSearchFeatureChannel_PermutationTime.cpp
+++ b/Engine/Plugins/Experimental/Animation/PoseSearch/Source/Runtime/Private/PoseSearchFeatureChannel_PermutationTime.cpp
@@ -14,14 +14,17 @@ void UPoseSearchFeatureChannel_PermutationTime::FindOrAddToSchema(UPoseSearchSch
 		PermutationTime->Weight = 0.f;
 #endif // WITH_EDITORONLY_DATA
 		Schema->AddTemporaryChannel(PermutationTime);
+
+		UE_LOG(LogPoseSearch, Warning, TEXT("required UPoseSearchFeatureChannel_PermutationTime has been added by the Schema '%s' with Weight of zero. Explicitly add one of this channel to the Schema to be able to tune its weight"), *Schema->GetName());
 	}
 }
 
-void UPoseSearchFeatureChannel_PermutationTime::Finalize(UPoseSearchSchema* Schema)
+bool UPoseSearchFeatureChannel_PermutationTime::Finalize(UPoseSearchSchema* Schema)
 {
 	ChannelDataOffset = Schema->SchemaCardinality;
 	ChannelCardinality = 1;
 	Schema->SchemaCardinality += ChannelCardinality;
+	return true;
 }
 
 void UPoseSearchFeatureChannel_PermutationTime::BuildQuery(UE::PoseSearch::FSearchContext& SearchContext) const
diff --git a/Engine/Plugins/Experimental/Animation/PoseSearch/Source/Runtime/Private/PoseSearchFeatureChannel_PermutationTime.h b/Engine/Plugins/Experimental/Animation/PoseSearch/Source/Runtime/Private/PoseSearchFeatureChannel_PermutationTime.h
--- a/Engine/Plugins/Experimental/Animation/PoseSearch/Source/Runtime/Private/PoseSearchFeatureChannel_PermutationTime.h
+++ b/Engine/Plugins/Experimental/Animation/PoseSearch/Source/Runtime/Private/PoseSearchFeatureChannel_PermutationTime.h
@@ -17,7 +17,7 @@ public:
 #endif // WITH_EDITORONLY_DATA
 
 	// UPoseSearchFeatureChannel interface
-	virtual void Finalize(UPoseSearchSchema* Schema) override;
+	virtual bool Finalize(UPoseSearchSchema* Schema) override;
 	virtual void BuildQuery(UE::PoseSearch::FSearchContext& SearchContext) const override;
 
 #if WITH_EDITOR
diff --git a/Engine/Plugins/Experimental/Animation/PoseSearch/Source/Runtime/Private/PoseSearchFeatureChannel_Phase.cpp b/Engine/Plugins/Experimental/Animation/PoseSearch/Source/Runtime/Private/PoseSearchFeatureChannel_Phase.cpp
--- a/Engine/Plugins/Experimental/Animation/PoseSearch/Source/Runtime/Private/PoseSearchFeatureChannel_Phase.cpp
+++ b/Engine/Plugins/Experimental/Animation/PoseSearch/Source/Runtime/Private/PoseSearchFeatureChannel_Phase.cpp
@@ -41,7 +41,7 @@ namespace UE::PoseSearch
 		return (Values[Num - 1] - Values[Num - 2]) * (Sample - (Num - 1)) + Values[Num - 1];
 	}
 
-	static bool CollectBonePositions(TArray<FVector>& BonePositions, FAssetIndexer& Indexer, int8 SchemaBoneIdx)
+	static bool CollectBonePositions(TArray<FVector>& BonePositions, FAssetIndexer& Indexer, int8 SchemaBoneIdx, const FRole& Role)
 	{
 		const int32 NumSamples = Indexer.GetEndSampleIdx() - Indexer.GetBeginSampleIdx();
 
@@ -50,7 +50,7 @@ namespace UE::PoseSearch
 		BonePositions.AddDefaulted(NumSamples);
 		for (int32 SampleIdx = 0; SampleIdx != NumSamples; ++SampleIdx)
 		{
-			if (!Indexer.GetSamplePosition(BonePositions[SampleIdx], 0.f, 0.f, SampleIdx, SchemaBoneIdx, RootSchemaBoneIdx, EPermutationTimeType::UseSampleTime))
+			if (!Indexer.GetSamplePosition(BonePositions[SampleIdx], 0.f, 0.f, SampleIdx, SchemaBoneIdx, RootSchemaBoneIdx, Role, Role))
 			{
 				return false;
 			}
@@ -284,19 +284,22 @@ namespace UE::PoseSearch
 } // namespace UE::PoseSearch
 #endif // WITH_EDITOR
 
-void UPoseSearchFeatureChannel_Phase::Finalize(UPoseSearchSchema* Schema)
+bool UPoseSearchFeatureChannel_Phase::Finalize(UPoseSearchSchema* Schema)
 {
 	ChannelDataOffset = Schema->SchemaCardinality;
 	ChannelCardinality = 2;
 	Schema->SchemaCardinality += ChannelCardinality;
-	SchemaBoneIdx = Schema->AddBoneReference(Bone);
+
+	SchemaBoneIdx = Schema->AddBoneReference(Bone, SampleRole);
+
+	return SchemaBoneIdx >= 0;
 }
 
 void UPoseSearchFeatureChannel_Phase::AddDependentChannels(UPoseSearchSchema* Schema) const
 {
 	if (Schema->bInjectAdditionalDebugChannels)
 	{
-		UPoseSearchFeatureChannel_Position::FindOrAddToSchema(Schema, 0.f, Bone.BoneName);
+		UPoseSearchFeatureChannel_Position::FindOrAddToSchema(Schema, 0.f, Bone.BoneName, SampleRole);
 	}
 }
 
@@ -309,8 +312,8 @@ void UPoseSearchFeatureChannel_Phase::BuildQuery(UE::PoseSearch::FSearchContext&
 	{
 		// composing a unique identifier to specify this channel with all the required properties to be able to share the query data with other channels of the same type
 		uint32 UniqueIdentifier = GetClass()->GetUniqueID();
-		UniqueIdentifier = HashCombineFast(UniqueIdentifier, ::GetTypeHash(SchemaBoneIdx));
-		UniqueIdentifier = HashCombineFast(UniqueIdentifier, ::GetTypeHash(InputQueryPose));
+		UniqueIdentifier = HashCombineFast(UniqueIdentifier, GetTypeHash(SchemaBoneIdx));
+		UniqueIdentifier = HashCombineFast(UniqueIdentifier, GetTypeHash(InputQueryPose));
 
 		TConstArrayView<float> CachedChannelData;
 		if (const UPoseSearchFeatureChannel* CachedChannel = SearchContext.GetCachedChannelData(UniqueIdentifier, this, CachedChannelData))
@@ -334,7 +337,7 @@ void UPoseSearchFeatureChannel_Phase::BuildQuery(UE::PoseSearch::FSearchContext&
 
 	const bool bCanUseCurrentResult = SearchContext.CanUseCurrentResult();
 	const bool bSkip = InputQueryPose != EInputQueryPose::UseCharacterPose && bCanUseCurrentResult;
-	if (bSkip || !SearchContext.GetHistory())
+	if (bSkip || !SearchContext.ArePoseHistoriesValid())
 	{
 		if (bCanUseCurrentResult)
 		{
@@ -342,7 +345,7 @@ void UPoseSearchFeatureChannel_Phase::BuildQuery(UE::PoseSearch::FSearchContext&
 			return;
 		}
 		
-		// we leave the SearchContext.EditFeatureVector() set to zero since the SearchContext.History is invalid and it'll fail if we continue
+		// we leave the SearchContext.EditFeatureVector() set to zero since the SearchContext.PoseHistory is invalid and it'll fail if we continue
 		UE_LOG(LogPoseSearch, Error, TEXT("UPoseSearchFeatureChannel_Phase::BuildQuery - Failed because Pose History Node is missing."));
 		return;
 	}
@@ -366,11 +369,11 @@ void UPoseSearchFeatureChannel_Phase::DebugDraw(const UE::PoseSearch::FDebugDraw
 #endif // WITH_EDITORONLY_DATA
 
 	const FVector2D Phase = FFeatureVectorHelper::DecodeVector2D(PoseVector, ChannelDataOffset);
-	const FVector BonePos = DrawParams.ExtractPosition(PoseVector, 0.f, SchemaBoneIdx);
+	const FVector BonePos = DrawParams.ExtractPosition(PoseVector, 0.f, SchemaBoneIdx, SampleRole);
 
-	const FVector TransformXAxisVector = DrawParams.GetRootTransform().TransformVector(FVector::XAxisVector);
-	const FVector TransformYAxisVector = DrawParams.GetRootTransform().TransformVector(FVector::YAxisVector);
-	const FVector TransformZAxisVector = DrawParams.GetRootTransform().TransformVector(FVector::ZAxisVector);
+	const FVector TransformXAxisVector = DrawParams.GetRootTransform(SampleRole).TransformVector(FVector::XAxisVector);
+	const FVector TransformYAxisVector = DrawParams.GetRootTransform(SampleRole).TransformVector(FVector::YAxisVector);
+	const FVector TransformZAxisVector = DrawParams.GetRootTransform(SampleRole).TransformVector(FVector::ZAxisVector);
 
 	const FVector PhaseVector = (TransformZAxisVector * Phase.X + TransformYAxisVector * Phase.Y) * ScaleFactor;
 	DrawParams.DrawLine(BonePos, BonePos + PhaseVector, Color);
@@ -408,7 +411,7 @@ bool UPoseSearchFeatureChannel_Phase::IndexAsset(UE::PoseSearch::FAssetIndexer&
 	TArray<LocalMinMax> LocalMinMax;
 	TArray<FVector> BonePositions;
 
-	if (!CollectBonePositions(BonePositions, Indexer, SchemaBoneIdx))
+	if (!CollectBonePositions(BonePositions, Indexer, SchemaBoneIdx, SampleRole))
 	{
 		return false;
 	}
@@ -448,7 +451,7 @@ UE::PoseSearch::TLabelBuilder& UPoseSearchFeatureChannel_Phase::GetLabel(UE::Pos
 	if (SchemaBoneIdx != RootSchemaBoneIdx)
 	{
 		LabelBuilder.Append(TEXT("_"));
-		LabelBuilder.Append(Schema->BoneReferences[SchemaBoneIdx].BoneName.ToString());
+		LabelBuilder.Append(Schema->GetBoneReferences(SampleRole)[SchemaBoneIdx].BoneName.ToString());
 	}
 
 	return LabelBuilder;
diff --git a/Engine/Plugins/Experimental/Animation/PoseSearch/Source/Runtime/Private/PoseSearchFeatureChannel_Pose.cpp b/Engine/Plugins/Experimental/Animation/PoseSearch/Source/Runtime/Private/PoseSearchFeatureChannel_Pose.cpp
--- a/Engine/Plugins/Experimental/Animation/PoseSearch/Source/Runtime/Private/PoseSearchFeatureChannel_Pose.cpp
+++ b/Engine/Plugins/Experimental/Animation/PoseSearch/Source/Runtime/Private/PoseSearchFeatureChannel_Pose.cpp
@@ -31,7 +31,7 @@ UPoseSearchFeatureChannel_Pose::UPoseSearchFeatureChannel_Pose()
 		}));
 }
 
-void UPoseSearchFeatureChannel_Pose::Finalize(UPoseSearchSchema* Schema)
+bool UPoseSearchFeatureChannel_Pose::Finalize(UPoseSearchSchema* Schema)
 {
 	SubChannels.Reset();
 
@@ -42,6 +42,8 @@ void UPoseSearchFeatureChannel_Pose::Finalize(UPoseSearchSchema* Schema)
 		{
 			UPoseSearchFeatureChannel_Position* Position = NewObject<UPoseSearchFeatureChannel_Position>(this, NAME_None, RF_Transient);
 			Position->Bone = SampledBone.Reference;
+			Position->SampleRole = SampleRole;
+			Position->OriginRole = SampleRole;
 #if WITH_EDITORONLY_DATA
 			Position->Weight = SampledBone.Weight * Weight;
 			Position->DebugColor = SampledBone.DebugColor;
@@ -55,6 +57,8 @@ void UPoseSearchFeatureChannel_Pose::Finalize(UPoseSearchSchema* Schema)
 		{
 			UPoseSearchFeatureChannel_Heading* HeadingX = NewObject<UPoseSearchFeatureChannel_Heading>(this, NAME_None, RF_Transient);
 			HeadingX->Bone = SampledBone.Reference;
+			HeadingX->SampleRole = SampleRole;
+			HeadingX->OriginRole = SampleRole;
 #if WITH_EDITORONLY_DATA
 			HeadingX->Weight = SampledBone.Weight * Weight;
 			HeadingX->DebugColor = SampledBone.DebugColor;
@@ -66,6 +70,8 @@ void UPoseSearchFeatureChannel_Pose::Finalize(UPoseSearchSchema* Schema)
 
 			UPoseSearchFeatureChannel_Heading* HeadingY = NewObject<UPoseSearchFeatureChannel_Heading>(this, NAME_None, RF_Transient);
 			HeadingY->Bone = SampledBone.Reference;
+			HeadingY->SampleRole = SampleRole;
+			HeadingY->OriginRole = SampleRole;
 #if WITH_EDITORONLY_DATA
 			HeadingY->Weight = SampledBone.Weight * Weight;
 			HeadingY->DebugColor = SampledBone.DebugColor;
@@ -80,6 +86,8 @@ void UPoseSearchFeatureChannel_Pose::Finalize(UPoseSearchSchema* Schema)
 		{
 			UPoseSearchFeatureChannel_Velocity* Velocity = NewObject<UPoseSearchFeatureChannel_Velocity>(this, NAME_None, RF_Transient);
 			Velocity->Bone = SampledBone.Reference;
+			Velocity->SampleRole = SampleRole;
+			Velocity->OriginRole = SampleRole;
 #if WITH_EDITORONLY_DATA
 			Velocity->Weight = SampledBone.Weight * Weight;
 			Velocity->DebugColor = SampledBone.DebugColor;
@@ -94,6 +102,7 @@ void UPoseSearchFeatureChannel_Pose::Finalize(UPoseSearchSchema* Schema)
 		{
 			UPoseSearchFeatureChannel_Phase* Phase = NewObject<UPoseSearchFeatureChannel_Phase>(this, NAME_None, RF_Transient);
 			Phase->Bone = SampledBone.Reference;
+			Phase->SampleRole = SampleRole;
 #if WITH_EDITORONLY_DATA
 			Phase->Weight = SampledBone.Weight * Weight;
 			Phase->DebugColor = SampledBone.DebugColor;
@@ -103,7 +112,7 @@ void UPoseSearchFeatureChannel_Pose::Finalize(UPoseSearchSchema* Schema)
 		}
 	}
 
-	Super::Finalize(Schema);
+	return Super::Finalize(Schema);
 }
 
 #if WITH_EDITOR
diff --git a/Engine/Plugins/Experimental/Animation/PoseSearch/Source/Runtime/Private/PoseSearchFeatureChannel_Position.cpp b/Engine/Plugins/Experimental/Animation/PoseSearch/Source/Runtime/Private/PoseSearchFeatureChannel_Position.cpp
--- a/Engine/Plugins/Experimental/Animation/PoseSearch/Source/Runtime/Private/PoseSearchFeatureChannel_Position.cpp
+++ b/Engine/Plugins/Experimental/Animation/PoseSearch/Source/Runtime/Private/PoseSearchFeatureChannel_Position.cpp
@@ -1,25 +1,35 @@
 // Copyright Epic Games, Inc. All Rights Reserved.
 
 #include "PoseSearchFeatureChannel_Position.h"
+#include "Engine/BlueprintGeneratedClass.h"
 #include "PoseSearch/PoseSearchAssetIndexer.h"
 #include "PoseSearch/PoseSearchAssetSampler.h"
 #include "PoseSearch/PoseSearchContext.h"
 #include "PoseSearch/PoseSearchDatabase.h"
 #include "PoseSearch/PoseSearchSchema.h"
-#include "Engine/BlueprintGeneratedClass.h"
+
+#if WITH_EDITOR
+#include "PropertyHandle.h"
+#endif // WITH_EDITOR
 
 UPoseSearchFeatureChannel_Position::UPoseSearchFeatureChannel_Position()
 {
 	bUseBlueprintQueryOverride = Cast<UBlueprintGeneratedClass>(GetClass()) != nullptr;
 }
 
-void UPoseSearchFeatureChannel_Position::FindOrAddToSchema(UPoseSearchSchema* Schema, float SampleTimeOffset, const FName& BoneName, EPermutationTimeType PermutationTimeType)
+void UPoseSearchFeatureChannel_Position::FindOrAddToSchema(UPoseSearchSchema* Schema, float SampleTimeOffset, const FName& BoneName, const UE::PoseSearch::FRole& Role, EPermutationTimeType PermutationTimeType)
 {
-	if (!Schema->FindChannel([SampleTimeOffset, &BoneName, PermutationTimeType](const UPoseSearchFeatureChannel* Channel) -> const UPoseSearchFeatureChannel_Position*
+	if (!Schema->FindChannel([SampleTimeOffset, &BoneName, &Role, PermutationTimeType](const UPoseSearchFeatureChannel* Channel) -> const UPoseSearchFeatureChannel_Position*
 		{
 			if (const UPoseSearchFeatureChannel_Position* Position = Cast<UPoseSearchFeatureChannel_Position>(Channel))
 			{
-				if (Position->Bone.BoneName == BoneName && Position->OriginBone.BoneName == NAME_None && Position->SampleTimeOffset == SampleTimeOffset && Position->OriginTimeOffset == 0.f && Position->PermutationTimeType == PermutationTimeType)
+				if (Position->Bone.BoneName == BoneName &&
+					Position->OriginBone.BoneName == NAME_None &&
+					Position->SampleTimeOffset == SampleTimeOffset &&
+					Position->OriginTimeOffset == 0.f &&
+					Position->PermutationTimeType == PermutationTimeType &&
+					Position->SampleRole == Role &&
+					Position->OriginRole == Role)
 				{
 					return Position;
 				}
@@ -29,6 +39,8 @@ void UPoseSearchFeatureChannel_Position::FindOrAddToSchema(UPoseSearchSchema* Sc
 	{
 		UPoseSearchFeatureChannel_Position* Position = NewObject<UPoseSearchFeatureChannel_Position>(Schema, NAME_None, RF_Transient);
 		Position->Bone.BoneName = BoneName;
+		Position->SampleRole = Role;
+		Position->OriginRole = Role;
 #if WITH_EDITORONLY_DATA
 		Position->Weight = 0.f;
 		// @todo: perhaps add a tunable color for injected channels
@@ -40,14 +52,16 @@ void UPoseSearchFeatureChannel_Position::FindOrAddToSchema(UPoseSearchSchema* Sc
 	}
 }
 
-void UPoseSearchFeatureChannel_Position::Finalize(UPoseSearchSchema* Schema)
+bool UPoseSearchFeatureChannel_Position::Finalize(UPoseSearchSchema* Schema)
 {
 	ChannelDataOffset = Schema->SchemaCardinality;
 	ChannelCardinality = UE::PoseSearch::FFeatureVectorHelper::GetVectorCardinality(ComponentStripping);
 	Schema->SchemaCardinality += ChannelCardinality;
 
-	SchemaBoneIdx = Schema->AddBoneReference(Bone);
-	SchemaOriginBoneIdx = Schema->AddBoneReference(OriginBone);
+	SchemaBoneIdx = Schema->AddBoneReference(Bone, SampleRole);
+	SchemaOriginBoneIdx = Schema->AddBoneReference(OriginBone, OriginRole);
+	
+	return SchemaBoneIdx >= 0 && SchemaOriginBoneIdx >= 0;
 }
 
 void UPoseSearchFeatureChannel_Position::AddDependentChannels(UPoseSearchSchema* Schema) const
@@ -59,7 +73,7 @@ void UPoseSearchFeatureChannel_Position::AddDependentChannels(UPoseSearchSchema*
 		if (SchemaOriginBoneIdx != RootSchemaBoneIdx)
 		{
 			const EPermutationTimeType DependentChannelsPermutationTimeType = PermutationTimeType == EPermutationTimeType::UsePermutationTime ? EPermutationTimeType::UseSampleToPermutationTime : EPermutationTimeType::UseSampleTime;
-			UPoseSearchFeatureChannel_Position::FindOrAddToSchema(Schema, 0.f, OriginBone.BoneName, DependentChannelsPermutationTimeType);
+			UPoseSearchFeatureChannel_Position::FindOrAddToSchema(Schema, 0.f, OriginBone.BoneName, OriginRole, DependentChannelsPermutationTimeType);
 		}
 	}
 }
@@ -70,8 +84,8 @@ void UPoseSearchFeatureChannel_Position::BuildQuery(UE::PoseSearch::FSearchConte
 
 	if (bUseBlueprintQueryOverride)
 	{
-		const FVector BonePositionWorld = BP_GetWorldPosition(SearchContext.GetAnimInstance());
-		const FVector BonePosition = SearchContext.GetSamplePosition(SampleTimeOffset, OriginTimeOffset, SchemaBoneIdx, SchemaOriginBoneIdx, PermutationTimeType, &BonePositionWorld);
+		const FVector BonePositionWorld = BP_GetWorldPosition(SearchContext.GetAnimInstance(SampleRole));
+		const FVector BonePosition = SearchContext.GetSamplePosition(SampleTimeOffset, OriginTimeOffset, SchemaBoneIdx, SchemaOriginBoneIdx, SampleRole, OriginRole, PermutationTimeType, &BonePositionWorld);
   		FFeatureVectorHelper::EncodeVector(SearchContext.EditFeatureVector(), ChannelDataOffset, BonePosition, ComponentStripping);
 		return;
 	}
@@ -81,14 +95,16 @@ void UPoseSearchFeatureChannel_Position::BuildQuery(UE::PoseSearch::FSearchConte
 	{
 		// composing a unique identifier to specify this channel with all the required properties to be able to share the query data with other channels of the same type
 		uint32 UniqueIdentifier = GetClass()->GetUniqueID();
-		UniqueIdentifier = HashCombineFast(UniqueIdentifier, ::GetTypeHash(SamplingAttributeId));
-		UniqueIdentifier = HashCombineFast(UniqueIdentifier, ::GetTypeHash(SampleTimeOffset));
-		UniqueIdentifier = HashCombineFast(UniqueIdentifier, ::GetTypeHash(OriginTimeOffset));
-		UniqueIdentifier = HashCombineFast(UniqueIdentifier, ::GetTypeHash(SchemaBoneIdx));
-		UniqueIdentifier = HashCombineFast(UniqueIdentifier, ::GetTypeHash(SchemaOriginBoneIdx));
-		UniqueIdentifier = HashCombineFast(UniqueIdentifier, ::GetTypeHash(InputQueryPose));
-		UniqueIdentifier = HashCombineFast(UniqueIdentifier, ::GetTypeHash(ComponentStripping));
-		UniqueIdentifier = HashCombineFast(UniqueIdentifier, ::GetTypeHash(PermutationTimeType));
+		UniqueIdentifier = HashCombineFast(UniqueIdentifier, GetTypeHash(SampleRole));
+		UniqueIdentifier = HashCombineFast(UniqueIdentifier, GetTypeHash(OriginRole));
+		UniqueIdentifier = HashCombineFast(UniqueIdentifier, GetTypeHash(SamplingAttributeId));
+		UniqueIdentifier = HashCombineFast(UniqueIdentifier, GetTypeHash(SampleTimeOffset));
+		UniqueIdentifier = HashCombineFast(UniqueIdentifier, GetTypeHash(OriginTimeOffset));
+		UniqueIdentifier = HashCombineFast(UniqueIdentifier, GetTypeHash(SchemaBoneIdx));
+		UniqueIdentifier = HashCombineFast(UniqueIdentifier, GetTypeHash(SchemaOriginBoneIdx));
+		UniqueIdentifier = HashCombineFast(UniqueIdentifier, GetTypeHash(InputQueryPose));
+		UniqueIdentifier = HashCombineFast(UniqueIdentifier, GetTypeHash(ComponentStripping));
+		UniqueIdentifier = HashCombineFast(UniqueIdentifier, GetTypeHash(PermutationTimeType));
 
 		TConstArrayView<float> CachedChannelData;
 		if (const UPoseSearchFeatureChannel* CachedChannel = SearchContext.GetCachedChannelData(UniqueIdentifier, this, CachedChannelData))
@@ -100,6 +116,8 @@ void UPoseSearchFeatureChannel_Position::BuildQuery(UE::PoseSearch::FSearchConte
 			check(CachedChannelData.Num() == ChannelCardinality);
 
 			// making sure there were no hash collisions
+			check(CachedPositionChannel->SampleRole == SampleRole);
+			check(CachedPositionChannel->OriginRole == OriginRole);
 			check(CachedPositionChannel->SamplingAttributeId == SamplingAttributeId);
 			check(CachedPositionChannel->SampleTimeOffset == SampleTimeOffset);
 			check(CachedPositionChannel->OriginTimeOffset == OriginTimeOffset);
@@ -117,23 +135,23 @@ void UPoseSearchFeatureChannel_Position::BuildQuery(UE::PoseSearch::FSearchConte
 	}
 
 	const bool bCanUseCurrentResult = SearchContext.CanUseCurrentResult();
-	const bool bSkip = InputQueryPose != EInputQueryPose::UseCharacterPose && bCanUseCurrentResult;
-		const bool bIsRootBone = SchemaBoneIdx == RootSchemaBoneIdx;
-	if (bSkip || (!SearchContext.GetHistory() && !bIsRootBone))
-		{
+	const bool bSkip = InputQueryPose != EInputQueryPose::UseCharacterPose && bCanUseCurrentResult && SampleRole == OriginRole;
+	const bool bIsRootBone = SchemaBoneIdx == RootSchemaBoneIdx;
+	if (bSkip || (!SearchContext.ArePoseHistoriesValid() && !bIsRootBone))
+	{
 		if (bCanUseCurrentResult)
-			{
+		{
 			FFeatureVectorHelper::Copy(SearchContext.EditFeatureVector(), ChannelDataOffset, ChannelCardinality, SearchContext.GetCurrentResultPoseVector());
 			return;
-			}
+		}
 
-		// we leave the SearchContext.EditFeatureVector() set to zero since the SearchContext.History is invalid and it'll fail if we continue
-				UE_LOG(LogPoseSearch, Error, TEXT("UPoseSearchFeatureChannel_Position::BuildQuery - Failed because Pose History Node is missing."));
+		// we leave the SearchContext.EditFeatureVector() set to zero since the SearchContext.PoseHistory is invalid and it'll fail if we continue
+		UE_LOG(LogPoseSearch, Error, TEXT("UPoseSearchFeatureChannel_Position::BuildQuery - Failed because Pose History Node is missing."));
 		return;
-			}
+	}
 	
-			// calculating the BonePosition in root bone space for the bone indexed by SchemaBoneIdx
-	const FVector BonePosition = SearchContext.GetSamplePosition(SampleTimeOffset, OriginTimeOffset, SchemaBoneIdx, SchemaOriginBoneIdx, PermutationTimeType);
+	// calculating the BonePosition in root bone space for the bone indexed by SchemaBoneIdx
+	const FVector BonePosition = SearchContext.GetSamplePosition(SampleTimeOffset, OriginTimeOffset, SchemaBoneIdx, SchemaOriginBoneIdx, SampleRole, OriginRole, PermutationTimeType);
 	FFeatureVectorHelper::EncodeVector(SearchContext.EditFeatureVector(), ChannelDataOffset, BonePosition, ComponentStripping);
 }
 
@@ -150,21 +168,19 @@ void UPoseSearchFeatureChannel_Position::DebugDraw(const UE::PoseSearch::FDebugD
 #endif // WITH_EDITORONLY_DATA
 
 	const FVector FeaturesVector = FFeatureVectorHelper::DecodeVector(PoseVector, ChannelDataOffset, ComponentStripping);
-	if (SchemaOriginBoneIdx == RootSchemaBoneIdx && FMath::IsNearlyZero(OriginTimeOffset))
-	{
-		const FVector BonePos = DrawParams.GetRootTransform().TransformPosition(FeaturesVector);
-		DrawParams.DrawPoint(BonePos, Color);
-	}
-	else
+	const EPermutationTimeType TimeType = PermutationTimeType == EPermutationTimeType::UsePermutationTime ? EPermutationTimeType::UseSampleToPermutationTime : EPermutationTimeType::UseSampleTime;
+	const FVector OriginBonePos = DrawParams.ExtractPosition(PoseVector, OriginTimeOffset, SchemaOriginBoneIdx, OriginRole, TimeType);
+	const FVector DeltaPos = DrawParams.GetRootTransform(OriginRole).TransformVector(FeaturesVector);
+	const FVector BonePos = OriginBonePos + DeltaPos;
+	DrawParams.DrawPoint(BonePos, Color);
+
+	const bool bDrawOrigin = !DeltaPos.IsNearlyZero() && (SchemaOriginBoneIdx != RootSchemaBoneIdx || !FMath::IsNearlyZero(OriginTimeOffset) ||
+							 SampleRole != OriginRole || PermutationTimeType != EPermutationTimeType::UseSampleTime);
+	if (bDrawOrigin)
 	{
-		const EPermutationTimeType TimeType = PermutationTimeType == EPermutationTimeType::UsePermutationTime ? EPermutationTimeType::UseSampleToPermutationTime : EPermutationTimeType::UseSampleTime;
-		const FVector OriginBonePos = DrawParams.ExtractPosition(PoseVector, OriginTimeOffset, SchemaOriginBoneIdx, TimeType);
-		const FVector DeltaPos = DrawParams.GetRootTransform().TransformVector(FeaturesVector);
-		const FVector BonePos = OriginBonePos + DeltaPos;
-		DrawParams.DrawLine(OriginBonePos, BonePos, Color);
 		DrawParams.DrawPoint(OriginBonePos, Color);
-		DrawParams.DrawPoint(BonePos, Color);
-	}	
+		DrawParams.DrawLine(OriginBonePos, BonePos, Color);
+	}
 }
 #endif // ENABLE_DRAW_DEBUG
 
@@ -184,7 +200,7 @@ bool UPoseSearchFeatureChannel_Position::IndexAsset(UE::PoseSearch::FAssetIndexe
 	FVector BonePosition;
 	for (int32 SampleIdx = Indexer.GetBeginSampleIdx(); SampleIdx != Indexer.GetEndSampleIdx(); ++SampleIdx)
 	{
-		if (Indexer.GetSamplePosition(BonePosition, SampleTimeOffset, OriginTimeOffset, SampleIdx, SchemaBoneIdx, SchemaOriginBoneIdx, PermutationTimeType, SamplingAttributeId))
+		if (Indexer.GetSamplePosition(BonePosition, SampleTimeOffset, OriginTimeOffset, SampleIdx, SchemaBoneIdx, SchemaOriginBoneIdx, SampleRole, OriginRole, PermutationTimeType, SamplingAttributeId))
 		{
 			FFeatureVectorHelper::EncodeVector(Indexer.GetPoseVector(SampleIdx), ChannelDataOffset, BonePosition, ComponentStripping);
 		}
@@ -219,13 +235,27 @@ UE::PoseSearch::TLabelBuilder& UPoseSearchFeatureChannel_Position::GetLabel(UE::
 	if (SchemaBoneIdx != RootSchemaBoneIdx)
 	{
 		LabelBuilder.Append(TEXT("_"));
-		LabelBuilder.Append(Schema->BoneReferences[SchemaBoneIdx].BoneName.ToString());
+		LabelBuilder.Append(Schema->GetBoneReferences(SampleRole)[SchemaBoneIdx].BoneName.ToString());
+	}
+
+	if (SampleRole != DefaultRole)
+	{
+		LabelBuilder.Append(TEXT("["));
+		LabelBuilder.Append(SampleRole.ToString());
+		LabelBuilder.Append(TEXT("]"));
 	}
 
 	if (SchemaOriginBoneIdx != RootSchemaBoneIdx)
 	{
 		LabelBuilder.Append(TEXT("_"));
-		LabelBuilder.Append(Schema->BoneReferences[SchemaOriginBoneIdx].BoneName.ToString());
+		LabelBuilder.Append(Schema->GetBoneReferences(OriginRole)[SchemaOriginBoneIdx].BoneName.ToString());
+	}
+
+	if (OriginRole != DefaultRole)
+	{
+		LabelBuilder.Append(TEXT("["));
+		LabelBuilder.Append(OriginRole.ToString());
+		LabelBuilder.Append(TEXT("]"));
 	}
 
 	AppendLabelSeparator(LabelBuilder, LabelFormat, true);
@@ -239,4 +269,19 @@ UE::PoseSearch::TLabelBuilder& UPoseSearchFeatureChannel_Position::GetLabel(UE::
 
 	return LabelBuilder;
 }
+
+USkeleton* UPoseSearchFeatureChannel_Position::GetSkeleton(bool& bInvalidSkeletonIsError, const IPropertyHandle* PropertyHandle)
+{
+	// blueprint generated classes don't have a schema, until they're instanced by the schema
+	if (const UPoseSearchSchema* Schema = GetSchema())
+	{
+		bInvalidSkeletonIsError = false;
+		if (PropertyHandle && PropertyHandle->GetProperty()->GetFName() == GET_MEMBER_NAME_CHECKED(UPoseSearchFeatureChannel_Position, OriginBone))
+		{
+			return Schema->GetSkeleton(OriginRole);
+		}
+	}
+
+	return Super::GetSkeleton(bInvalidSkeletonIsError, PropertyHandle);
+}
 #endif
\ No newline at end of file
diff --git a/Engine/Plugins/Experimental/Animation/PoseSearch/Source/Runtime/Private/PoseSearchFeatureChannel_Position.h b/Engine/Plugins/Experimental/Animation/PoseSearch/Source/Runtime/Private/PoseSearchFeatureChannel_Position.h
--- a/Engine/Plugins/Experimental/Animation/PoseSearch/Source/Runtime/Private/PoseSearchFeatureChannel_Position.h
+++ b/Engine/Plugins/Experimental/Animation/PoseSearch/Source/Runtime/Private/PoseSearchFeatureChannel_Position.h
@@ -15,9 +15,15 @@ public:
 	UPROPERTY(EditAnywhere, Category = "Settings")
 	FBoneReference Bone;
 
+	UPROPERTY(EditAnywhere, Category = "Settings")
+	FName SampleRole = UE::PoseSearch::DefaultRole;
+
 	UPROPERTY(EditAnywhere, Category = "Settings")
 	FBoneReference OriginBone;
 
+	UPROPERTY(EditAnywhere, Category = "Settings")
+	FName OriginRole = UE::PoseSearch::DefaultRole;
+
 #if WITH_EDITORONLY_DATA
 	UPROPERTY(EditAnywhere, Category = "Settings")
 	float Weight = 1.f;
@@ -69,7 +75,7 @@ public:
 	UPoseSearchFeatureChannel_Position();
 
 	// UPoseSearchFeatureChannel interface
-	virtual void Finalize(UPoseSearchSchema* Schema) override;
+	virtual bool Finalize(UPoseSearchSchema* Schema) override;
 	virtual void BuildQuery(UE::PoseSearch::FSearchContext& SearchContext) const override;
 
 	virtual EPermutationTimeType GetPermutationTimeType() const override { return PermutationTimeType; }
@@ -83,7 +89,10 @@ public:
 	virtual void FillWeights(TArrayView<float> Weights) const override;
 	virtual bool IndexAsset(UE::PoseSearch::FAssetIndexer& Indexer) const override;
 	virtual UE::PoseSearch::TLabelBuilder& GetLabel(UE::PoseSearch::TLabelBuilder& LabelBuilder, UE::PoseSearch::ELabelFormat LabelFormat = UE::PoseSearch::ELabelFormat::Full_Horizontal) const override;
+
+	// IBoneReferenceSkeletonProvider interface
+	USkeleton* GetSkeleton(bool& bInvalidSkeletonIsError, const IPropertyHandle* PropertyHandle) override;
 #endif
 
-	static void FindOrAddToSchema(UPoseSearchSchema* Schema, float SampleTimeOffset, const FName& BoneName = NAME_None, EPermutationTimeType PermutationTimeType = EPermutationTimeType::UseSampleTime);
+	static void FindOrAddToSchema(UPoseSearchSchema* Schema, float SampleTimeOffset, const FName& BoneName, const UE::PoseSearch::FRole& Role, EPermutationTimeType PermutationTimeType = EPermutationTimeType::UseSampleTime);
 };
diff --git a/Engine/Plugins/Experimental/Animation/PoseSearch/Source/Runtime/Private/PoseSearchFeatureChannel_SamplingTime.cpp b/Engine/Plugins/Experimental/Animation/PoseSearch/Source/Runtime/Private/PoseSearchFeatureChannel_SamplingTime.cpp
--- a/Engine/Plugins/Experimental/Animation/PoseSearch/Source/Runtime/Private/PoseSearchFeatureChannel_SamplingTime.cpp
+++ b/Engine/Plugins/Experimental/Animation/PoseSearch/Source/Runtime/Private/PoseSearchFeatureChannel_SamplingTime.cpp
@@ -5,11 +5,12 @@
 #include "PoseSearch/PoseSearchContext.h"
 #include "PoseSearch/PoseSearchSchema.h"
 
-void UPoseSearchFeatureChannel_SamplingTime::Finalize(UPoseSearchSchema* Schema)
+bool UPoseSearchFeatureChannel_SamplingTime::Finalize(UPoseSearchSchema* Schema)
 {
 	ChannelDataOffset = Schema->SchemaCardinality;
 	ChannelCardinality = 1;
 	Schema->SchemaCardinality += ChannelCardinality;
+	return true;
 }
 
 void UPoseSearchFeatureChannel_SamplingTime::BuildQuery(UE::PoseSearch::FSearchContext& SearchContext) const
diff --git a/Engine/Plugins/Experimental/Animation/PoseSearch/Source/Runtime/Private/PoseSearchFeatureChannel_Trajectory.cpp b/Engine/Plugins/Experimental/Animation/PoseSearch/Source/Runtime/Private/PoseSearchFeatureChannel_Trajectory.cpp
--- a/Engine/Plugins/Experimental/Animation/PoseSearch/Source/Runtime/Private/PoseSearchFeatureChannel_Trajectory.cpp
+++ b/Engine/Plugins/Experimental/Animation/PoseSearch/Source/Runtime/Private/PoseSearchFeatureChannel_Trajectory.cpp
@@ -44,7 +44,7 @@ UPoseSearchFeatureChannel_Trajectory::UPoseSearchFeatureChannel_Trajectory()
 		}));
 }
 
-void UPoseSearchFeatureChannel_Trajectory::Finalize(UPoseSearchSchema* Schema)
+bool UPoseSearchFeatureChannel_Trajectory::Finalize(UPoseSearchSchema* Schema)
 {
 	SubChannels.Reset();
 
@@ -53,12 +53,15 @@ void UPoseSearchFeatureChannel_Trajectory::Finalize(UPoseSearchSchema* Schema)
 		if (EnumHasAnyFlags(Sample.Flags, EPoseSearchTrajectoryFlags::Position | EPoseSearchTrajectoryFlags::PositionXY))
 		{
 			UPoseSearchFeatureChannel_Position* Position = NewObject<UPoseSearchFeatureChannel_Position>(this, NAME_None, RF_Transient);
+			Position->SampleRole = SampleRole;
+			Position->OriginRole = SampleRole;
 #if WITH_EDITORONLY_DATA
 			Position->Weight = Sample.Weight * Weight;
 			Position->DebugColor = Sample.DebugColor;
 #endif // WITH_EDITORONLY_DATA
 			Position->SampleTimeOffset = Sample.Offset;
 			Position->InputQueryPose = EInputQueryPose::UseCharacterPose;
+	
 			if (EnumHasAnyFlags(Sample.Flags, EPoseSearchTrajectoryFlags::PositionXY))
 			{
 				Position->ComponentStripping = EComponentStrippingVector::StripZ;
@@ -69,6 +72,8 @@ void UPoseSearchFeatureChannel_Trajectory::Finalize(UPoseSearchSchema* Schema)
 		if (EnumHasAnyFlags(Sample.Flags, EPoseSearchTrajectoryFlags::Velocity | EPoseSearchTrajectoryFlags::VelocityXY))
 		{
 			UPoseSearchFeatureChannel_Velocity* Velocity = NewObject<UPoseSearchFeatureChannel_Velocity>(this, NAME_None, RF_Transient);
+			Velocity->SampleRole = SampleRole;
+			Velocity->OriginRole = SampleRole;
 #if WITH_EDITORONLY_DATA
 			Velocity->Weight = Sample.Weight * Weight;
 			Velocity->DebugColor = Sample.DebugColor;
@@ -86,6 +91,8 @@ void UPoseSearchFeatureChannel_Trajectory::Finalize(UPoseSearchSchema* Schema)
 		if (EnumHasAnyFlags(Sample.Flags, EPoseSearchTrajectoryFlags::VelocityDirection | EPoseSearchTrajectoryFlags::VelocityDirectionXY))
 		{
 			UPoseSearchFeatureChannel_Velocity* Velocity = NewObject<UPoseSearchFeatureChannel_Velocity>(this, NAME_None, RF_Transient);
+			Velocity->SampleRole = SampleRole;
+			Velocity->OriginRole = SampleRole;
 #if WITH_EDITORONLY_DATA
 			Velocity->Weight = Sample.Weight * Weight;
 			Velocity->DebugColor = Sample.DebugColor;
@@ -104,6 +111,8 @@ void UPoseSearchFeatureChannel_Trajectory::Finalize(UPoseSearchSchema* Schema)
 		if (EnumHasAnyFlags(Sample.Flags, EPoseSearchTrajectoryFlags::FacingDirection | EPoseSearchTrajectoryFlags::FacingDirectionXY))
 		{
 			UPoseSearchFeatureChannel_Heading* Heading = NewObject<UPoseSearchFeatureChannel_Heading>(this, NAME_None, RF_Transient);
+			Heading->SampleRole = SampleRole;
+			Heading->OriginRole = SampleRole;
 #if WITH_EDITORONLY_DATA
 			Heading->Weight = Sample.Weight * Weight;
 			Heading->DebugColor = Sample.DebugColor;
@@ -118,7 +127,7 @@ void UPoseSearchFeatureChannel_Trajectory::Finalize(UPoseSearchSchema* Schema)
 		}
 	}
 
-	Super::Finalize(Schema);
+	return Super::Finalize(Schema);
 }
 
 #if ENABLE_DRAW_DEBUG
@@ -159,11 +168,11 @@ void UPoseSearchFeatureChannel_Trajectory::DebugDraw(const UE::PoseSearch::FDebu
 			if (PrevTimeOffset * CurrTimeOffset < UE_KINDA_SMALL_NUMBER)
 			{
 				// we jumped from negative to positive time offset without having a zero time offset. so we add the zero
-				TrajSplinePos.Add(DrawParams.ExtractPosition(PoseVector, 0.f));
+				TrajSplinePos.Add(DrawParams.ExtractPosition(PoseVector, 0.f, RootSchemaBoneIdx, Positions[i]->OriginRole));
 				TrajSplineColor.Add(Color);
 			}
 
-			TrajSplinePos.Add(DrawParams.ExtractPosition(PoseVector, CurrTimeOffset));
+			TrajSplinePos.Add(DrawParams.ExtractPosition(PoseVector, CurrTimeOffset, RootSchemaBoneIdx, Positions[i]->OriginRole));
 			TrajSplineColor.Add(Color);
 
 			PrevTimeOffset = CurrTimeOffset;
diff --git a/Engine/Plugins/Experimental/Animation/PoseSearch/Source/Runtime/Private/PoseSearchFeatureChannel_Velocity.cpp b/Engine/Plugins/Experimental/Animation/PoseSearch/Source/Runtime/Private/PoseSearchFeatureChannel_Velocity.cpp
--- a/Engine/Plugins/Experimental/Animation/PoseSearch/Source/Runtime/Private/PoseSearchFeatureChannel_Velocity.cpp
+++ b/Engine/Plugins/Experimental/Animation/PoseSearch/Source/Runtime/Private/PoseSearchFeatureChannel_Velocity.cpp
@@ -1,6 +1,7 @@
 // Copyright Epic Games, Inc. All Rights Reserved.
 
 #include "PoseSearchFeatureChannel_Velocity.h"
+#include "Engine/BlueprintGeneratedClass.h"
 #include "PoseSearch/PoseSearchAssetIndexer.h"
 #include "PoseSearch/PoseSearchAssetSampler.h"
 #include "PoseSearch/PoseSearchContext.h"
@@ -8,28 +9,33 @@
 #include "PoseSearch/PoseSearchHistory.h"
 #include "PoseSearch/PoseSearchSchema.h"
 #include "PoseSearchFeatureChannel_Position.h"
-#include "Engine/BlueprintGeneratedClass.h"
+
+#if WITH_EDITOR
+#include "PropertyHandle.h"
+#endif // WITH_EDITOR
 
 UPoseSearchFeatureChannel_Velocity::UPoseSearchFeatureChannel_Velocity()
 {
 	bUseBlueprintQueryOverride = Cast<UBlueprintGeneratedClass>(GetClass()) != nullptr;
 }
 
-void UPoseSearchFeatureChannel_Velocity::Finalize(UPoseSearchSchema* Schema)
+bool UPoseSearchFeatureChannel_Velocity::Finalize(UPoseSearchSchema* Schema)
 {
 	ChannelDataOffset = Schema->SchemaCardinality;
 	ChannelCardinality = UE::PoseSearch::FFeatureVectorHelper::GetVectorCardinality(ComponentStripping);
 	Schema->SchemaCardinality += ChannelCardinality;
 
-	SchemaBoneIdx = Schema->AddBoneReference(Bone);
-	SchemaOriginBoneIdx = Schema->AddBoneReference(OriginBone);
+	SchemaBoneIdx = Schema->AddBoneReference(Bone, SampleRole);
+	SchemaOriginBoneIdx = Schema->AddBoneReference(OriginBone, OriginRole);
+
+	return SchemaBoneIdx >= 0 && SchemaOriginBoneIdx >= 0;
 }
 
 void UPoseSearchFeatureChannel_Velocity::AddDependentChannels(UPoseSearchSchema* Schema) const
 {
 	if (Schema->bInjectAdditionalDebugChannels)
 	{
-		UPoseSearchFeatureChannel_Position::FindOrAddToSchema(Schema, SampleTimeOffset, Bone.BoneName);
+		UPoseSearchFeatureChannel_Position::FindOrAddToSchema(Schema, SampleTimeOffset, Bone.BoneName, SampleRole);
 	}
 }
 
@@ -40,9 +46,8 @@ void UPoseSearchFeatureChannel_Velocity::BuildQuery(UE::PoseSearch::FSearchConte
 	const bool bIsRootBone = SchemaBoneIdx == RootSchemaBoneIdx;
 	if (bUseBlueprintQueryOverride)
 	{
-		const FVector LinearVelocityWorld = BP_GetWorldVelocity(SearchContext.GetAnimInstance());
-
-		FVector LinearVelocity = SearchContext.GetSampleVelocity(SampleTimeOffset, OriginTimeOffset, SchemaBoneIdx, SchemaOriginBoneIdx, bUseCharacterSpaceVelocities, EPermutationTimeType::UseSampleTime, &LinearVelocityWorld);
+		const FVector LinearVelocityWorld = BP_GetWorldVelocity(SearchContext.GetAnimInstance(SampleRole));
+		FVector LinearVelocity = SearchContext.GetSampleVelocity(SampleTimeOffset, OriginTimeOffset, SchemaBoneIdx, SchemaOriginBoneIdx, SampleRole, OriginRole, bUseCharacterSpaceVelocities, EPermutationTimeType::UseSampleTime, &LinearVelocityWorld);
 		if (bNormalize)
 		{
 			LinearVelocity = LinearVelocity.GetClampedToMaxSize(1.f);
@@ -56,16 +61,18 @@ void UPoseSearchFeatureChannel_Velocity::BuildQuery(UE::PoseSearch::FSearchConte
 	{
 		// composing a unique identifier to specify this channel with all the required properties to be able to share the query data with other channels of the same type
 		uint32 UniqueIdentifier = GetClass()->GetUniqueID();
-		UniqueIdentifier = HashCombineFast(UniqueIdentifier, ::GetTypeHash(SamplingAttributeId));
-		UniqueIdentifier = HashCombineFast(UniqueIdentifier, ::GetTypeHash(SampleTimeOffset));
-		UniqueIdentifier = HashCombineFast(UniqueIdentifier, ::GetTypeHash(OriginTimeOffset));
-		UniqueIdentifier = HashCombineFast(UniqueIdentifier, ::GetTypeHash(SchemaBoneIdx));
-		UniqueIdentifier = HashCombineFast(UniqueIdentifier, ::GetTypeHash(SchemaOriginBoneIdx));
-		UniqueIdentifier = HashCombineFast(UniqueIdentifier, ::GetTypeHash(InputQueryPose));
-		UniqueIdentifier = HashCombineFast(UniqueIdentifier, ::GetTypeHash(bUseCharacterSpaceVelocities));
-		UniqueIdentifier = HashCombineFast(UniqueIdentifier, ::GetTypeHash(bNormalize));
-		UniqueIdentifier = HashCombineFast(UniqueIdentifier, ::GetTypeHash(ComponentStripping));
-		UniqueIdentifier = HashCombineFast(UniqueIdentifier, ::GetTypeHash(PermutationTimeType));
+		UniqueIdentifier = HashCombineFast(UniqueIdentifier, GetTypeHash(SampleRole));
+		UniqueIdentifier = HashCombineFast(UniqueIdentifier, GetTypeHash(OriginRole));
+		UniqueIdentifier = HashCombineFast(UniqueIdentifier, GetTypeHash(SamplingAttributeId));
+		UniqueIdentifier = HashCombineFast(UniqueIdentifier, GetTypeHash(SampleTimeOffset));
+		UniqueIdentifier = HashCombineFast(UniqueIdentifier, GetTypeHash(OriginTimeOffset));
+		UniqueIdentifier = HashCombineFast(UniqueIdentifier, GetTypeHash(SchemaBoneIdx));
+		UniqueIdentifier = HashCombineFast(UniqueIdentifier, GetTypeHash(SchemaOriginBoneIdx));
+		UniqueIdentifier = HashCombineFast(UniqueIdentifier, GetTypeHash(InputQueryPose));
+		UniqueIdentifier = HashCombineFast(UniqueIdentifier, GetTypeHash(bUseCharacterSpaceVelocities));
+		UniqueIdentifier = HashCombineFast(UniqueIdentifier, GetTypeHash(bNormalize));
+		UniqueIdentifier = HashCombineFast(UniqueIdentifier, GetTypeHash(ComponentStripping));
+		UniqueIdentifier = HashCombineFast(UniqueIdentifier, GetTypeHash(PermutationTimeType));
 
 		TConstArrayView<float> CachedChannelData;
 		if (const UPoseSearchFeatureChannel* CachedChannel = SearchContext.GetCachedChannelData(UniqueIdentifier, this, CachedChannelData))
@@ -77,6 +84,8 @@ void UPoseSearchFeatureChannel_Velocity::BuildQuery(UE::PoseSearch::FSearchConte
 			check(CachedChannelData.Num() == ChannelCardinality);
 
 			// making sure there were no hash collisions
+			check(CachedVelocityChannel->SampleRole == SampleRole);
+			check(CachedVelocityChannel->OriginRole == OriginRole);
 			check(CachedVelocityChannel->SamplingAttributeId == SamplingAttributeId);
 			check(CachedVelocityChannel->SampleTimeOffset == SampleTimeOffset);
 			check(CachedVelocityChannel->OriginTimeOffset == OriginTimeOffset);
@@ -96,26 +105,26 @@ void UPoseSearchFeatureChannel_Velocity::BuildQuery(UE::PoseSearch::FSearchConte
 	}
 
 	const bool bCanUseCurrentResult = SearchContext.CanUseCurrentResult();
-	const bool bSkip = InputQueryPose != EInputQueryPose::UseCharacterPose && bCanUseCurrentResult;
-	if (bSkip || (!SearchContext.GetHistory() && !bIsRootBone))
-		{
+	const bool bSkip = InputQueryPose != EInputQueryPose::UseCharacterPose && bCanUseCurrentResult && SampleRole == OriginRole;
+	if (bSkip || (!SearchContext.ArePoseHistoriesValid() && !bIsRootBone))
+	{
 		if (bCanUseCurrentResult)
-			{
+		{
 			FFeatureVectorHelper::Copy(SearchContext.EditFeatureVector(), ChannelDataOffset, ChannelCardinality, SearchContext.GetCurrentResultPoseVector());
 			return;
-			}
+		}
 
-		// we leave the SearchContext.EditFeatureVector() set to zero since the SearchContext.History is invalid and it'll fail if we continue
-				UE_LOG(LogPoseSearch, Error, TEXT("UPoseSearchFeatureChannel_Velocity::BuildQuery - Failed because Pose History Node is missing."));
+		// we leave the SearchContext.EditFeatureVector() set to zero since the SearchContext.PoseHistory is invalid and it'll fail if we continue
+		UE_LOG(LogPoseSearch, Error, TEXT("UPoseSearchFeatureChannel_Velocity::BuildQuery - Failed because Pose History Node is missing."));
 		return;
-			}
+	}
 	
-			// calculating the LinearVelocity for the bone indexed by SchemaBoneIdx
-	FVector LinearVelocity = SearchContext.GetSampleVelocity(SampleTimeOffset, OriginTimeOffset, SchemaBoneIdx, SchemaOriginBoneIdx, bUseCharacterSpaceVelocities, PermutationTimeType);
-			if (bNormalize)
-			{
-				LinearVelocity = LinearVelocity.GetClampedToMaxSize(1.f);
-			}
+	// calculating the LinearVelocity for the bone indexed by SchemaBoneIdx
+	FVector LinearVelocity = SearchContext.GetSampleVelocity(SampleTimeOffset, OriginTimeOffset, SchemaBoneIdx, SchemaOriginBoneIdx, SampleRole, OriginRole, bUseCharacterSpaceVelocities, PermutationTimeType);
+	if (bNormalize)
+	{
+		LinearVelocity = LinearVelocity.GetClampedToMaxSize(1.f);
+	}
 
 	FFeatureVectorHelper::EncodeVector(SearchContext.EditFeatureVector(), ChannelDataOffset, LinearVelocity, ComponentStripping);
 }
@@ -134,9 +143,9 @@ void UPoseSearchFeatureChannel_Velocity::DebugDraw(const UE::PoseSearch::FDebugD
 
 	const float LinearVelocityScale = bNormalize ? 15.f : 0.08f;
 
-	const FVector LinearVelocity = DrawParams.GetRootTransform().TransformVector(FFeatureVectorHelper::DecodeVector(PoseVector, ChannelDataOffset, ComponentStripping));
+	const FVector LinearVelocity = DrawParams.GetRootTransform(SampleRole).TransformVector(FFeatureVectorHelper::DecodeVector(PoseVector, ChannelDataOffset, ComponentStripping));
 	const FVector BoneVelDirection = LinearVelocity.GetSafeNormal();
-	const FVector BonePos = DrawParams.ExtractPosition(PoseVector, SampleTimeOffset, SchemaBoneIdx, PermutationTimeType, SamplingAttributeId);
+	const FVector BonePos = DrawParams.ExtractPosition(PoseVector, SampleTimeOffset, SchemaBoneIdx, SampleRole, PermutationTimeType, SamplingAttributeId);
 
 	DrawParams.DrawLine(BonePos, BonePos + LinearVelocity * LinearVelocityScale, Color);
 }
@@ -158,7 +167,7 @@ bool UPoseSearchFeatureChannel_Velocity::IndexAsset(UE::PoseSearch::FAssetIndexe
 	FVector LinearVelocity;
 	for (int32 SampleIdx = Indexer.GetBeginSampleIdx(); SampleIdx != Indexer.GetEndSampleIdx(); ++SampleIdx)
 	{
-		if (Indexer.GetSampleVelocity(LinearVelocity, SampleTimeOffset, OriginTimeOffset, SampleIdx, SchemaBoneIdx, SchemaOriginBoneIdx, bUseCharacterSpaceVelocities, PermutationTimeType, SamplingAttributeId))
+		if (Indexer.GetSampleVelocity(LinearVelocity, SampleTimeOffset, OriginTimeOffset, SampleIdx, SchemaBoneIdx, SchemaOriginBoneIdx, SampleRole, OriginRole, bUseCharacterSpaceVelocities, PermutationTimeType, SamplingAttributeId))
 		{
 			if (bNormalize)
 			{
@@ -201,13 +210,27 @@ UE::PoseSearch::TLabelBuilder& UPoseSearchFeatureChannel_Velocity::GetLabel(UE::
 	if (SchemaBoneIdx != RootSchemaBoneIdx)
 	{
 		LabelBuilder.Append(TEXT("_"));
-		LabelBuilder.Append(Schema->BoneReferences[SchemaBoneIdx].BoneName.ToString());
+		LabelBuilder.Append(Schema->GetBoneReferences(SampleRole)[SchemaBoneIdx].BoneName.ToString());
+	}
+
+	if (SampleRole != DefaultRole)
+	{
+		LabelBuilder.Append(TEXT("["));
+		LabelBuilder.Append(SampleRole.ToString());
+		LabelBuilder.Append(TEXT("]"));
 	}
 
 	if (SchemaOriginBoneIdx != RootSchemaBoneIdx)
 	{
 		LabelBuilder.Append(TEXT("_"));
-		LabelBuilder.Append(Schema->BoneReferences[SchemaOriginBoneIdx].BoneName.ToString());
+		LabelBuilder.Append(Schema->GetBoneReferences(OriginRole)[SchemaOriginBoneIdx].BoneName.ToString());
+	}
+
+	if (OriginRole != DefaultRole)
+	{
+		LabelBuilder.Append(TEXT("["));
+		LabelBuilder.Append(OriginRole.ToString());
+		LabelBuilder.Append(TEXT("]"));
 	}
 
 	AppendLabelSeparator(LabelBuilder, LabelFormat, true);
@@ -221,4 +244,19 @@ UE::PoseSearch::TLabelBuilder& UPoseSearchFeatureChannel_Velocity::GetLabel(UE::
 
 	return LabelBuilder;
 }
+
+USkeleton* UPoseSearchFeatureChannel_Velocity::GetSkeleton(bool& bInvalidSkeletonIsError, const IPropertyHandle* PropertyHandle)
+{
+	// blueprint generated classes don't have a schema, until they're instanced by the schema
+	if (const UPoseSearchSchema* Schema = GetSchema())
+	{
+		bInvalidSkeletonIsError = false;
+		if (PropertyHandle && PropertyHandle->GetProperty()->GetFName() == GET_MEMBER_NAME_CHECKED(UPoseSearchFeatureChannel_Velocity, OriginBone))
+		{
+			return Schema->GetSkeleton(OriginRole);
+		}
+	}
+
+	return Super::GetSkeleton(bInvalidSkeletonIsError, PropertyHandle);
+}
 #endif
\ No newline at end of file
