motion matching - multicharacter WIP

diff --git a/Engine/Plugins/Experimental/Animation/PoseSearch/Source/Runtime/Private/PoseSearchHistory.cpp b/Engine/Plugins/Experimental/Animation/PoseSearch/Source/Runtime/Private/PoseSearchHistory.cpp
--- a/Engine/Plugins/Experimental/Animation/PoseSearch/Source/Runtime/Private/PoseSearchHistory.cpp
+++ b/Engine/Plugins/Experimental/Animation/PoseSearch/Source/Runtime/Private/PoseSearchHistory.cpp
@@ -69,6 +69,90 @@
 	return LowerBound(First, Last, Value, FIdentityFunctor(), SortPredicate);
 }
 
+static FBoneIndexType GetRemappedBoneIndexType(FBoneIndexType BoneIndexType, const USkeleton* BoneIndexSkeleton, const USkeleton* LastUpdateSkeleton)
+{
+	// remapping BoneIndexType in case the skeleton used to store history (LastUpdateSkeleton) is different from BoneIndexSkeleton
+	if (LastUpdateSkeleton != nullptr && LastUpdateSkeleton != BoneIndexSkeleton)
+	{
+		const FSkeletonRemapping& SkeletonRemapping = UE::Anim::FSkeletonRemappingRegistry::Get().GetRemapping(BoneIndexSkeleton, LastUpdateSkeleton);
+		if (SkeletonRemapping.IsValid())
+		{
+			BoneIndexType = SkeletonRemapping.GetTargetSkeletonBoneIndex(BoneIndexType);
+		}
+	}
+
+	return BoneIndexType;
+}
+
+static FComponentSpaceTransformIndex GetRemappedComponentSpaceTransformIndex(const USkeleton* BoneIndexSkeleton, const USkeleton* LastUpdateSkeleton, const FBoneToTransformMap& BoneToTransformMap, FBoneIndexType BoneIndexType, bool& bSuccess)
+{
+	check(BoneIndexType != WorldSpaceIndexType);
+
+	FComponentSpaceTransformIndex BoneTransformIndex = FComponentSpaceTransformIndex(BoneIndexType);
+	if (BoneIndexType != ComponentSpaceIndexType)
+	{
+		BoneIndexType = GetRemappedBoneIndexType(BoneIndexType, BoneIndexSkeleton, LastUpdateSkeleton);
+
+		if (!BoneToTransformMap.IsEmpty())
+		{
+			if (const FComponentSpaceTransformIndex* FoundBoneTransformIndex = BoneToTransformMap.Find(BoneTransformIndex))
+			{
+				BoneTransformIndex = *FoundBoneTransformIndex;
+			}
+			else
+			{
+				BoneTransformIndex = RootBoneIndexType;
+				bSuccess = false;
+			}
+		}
+	}
+	return BoneTransformIndex;
+}
+
+static bool LerpEntries(float Time, bool bExtrapolate, const FPoseHistoryEntry& PrevEntry, const FPoseHistoryEntry& NextEntry, const USkeleton* BoneIndexSkeleton, const USkeleton* LastUpdateSkeleton,
+	const FBoneToTransformMap& BoneToTransformMap, FBoneIndexType BoneIndexType, FBoneIndexType ReferenceBoneIndexType, FTransform& OutBoneTransform)
+{
+	bool bSuccess = true;
+
+	const float Denominator = NextEntry.AccumulatedSeconds - PrevEntry.AccumulatedSeconds;
+	float LerpValue = 0.f;
+	if (!FMath::IsNearlyZero(Denominator))
+	{
+		const float Numerator = Time - PrevEntry.AccumulatedSeconds;
+		LerpValue = bExtrapolate ? Numerator / Denominator : FMath::Clamp(Numerator / Denominator, 0.f, 1.f);
+	}
+
+	const FComponentSpaceTransformIndex BoneTransformIndex = GetRemappedComponentSpaceTransformIndex(BoneIndexSkeleton, LastUpdateSkeleton, BoneToTransformMap, BoneIndexType, bSuccess);
+	const FComponentSpaceTransformIndex ReferenceBoneTransformIndex = GetRemappedComponentSpaceTransformIndex(BoneIndexSkeleton, LastUpdateSkeleton, BoneToTransformMap, ReferenceBoneIndexType, bSuccess);
+
+	if (BoneTransformIndex != ComponentSpaceIndexType)
+	{
+		if (ReferenceBoneTransformIndex == ComponentSpaceIndexType)
+		{
+			OutBoneTransform.Blend(
+				PrevEntry.GetComponentSpaceTransform(BoneTransformIndex),
+				NextEntry.GetComponentSpaceTransform(BoneTransformIndex),
+				LerpValue);
+		}
+		else
+		{
+			OutBoneTransform.Blend(
+				PrevEntry.GetComponentSpaceTransform(BoneTransformIndex) * PrevEntry.GetComponentSpaceTransform(ReferenceBoneTransformIndex).Inverse(),
+				NextEntry.GetComponentSpaceTransform(BoneTransformIndex) * NextEntry.GetComponentSpaceTransform(ReferenceBoneTransformIndex).Inverse(),
+				LerpValue);
+		}
+	}
+	else
+	{
+		// @todo: implement if required
+		OutBoneTransform = FTransform::Identity;
+		bSuccess = false;
+		unimplemented();
+	}
+
+	return bSuccess;
+}
+
 //////////////////////////////////////////////////////////////////////////
 // FPoseHistoryEntry
 void FPoseHistoryEntry::Update(float Time, FCSPose<FCompactPose>& ComponentSpacePose, const FBoneToTransformMap& BoneToTransformMap, bool bStoreScales)
@@ -134,13 +218,22 @@
 	return FTransform(Quat, ComponentSpacePositions[Index], Scale);
 }
 
+FArchive& operator<<(FArchive& Ar, FPoseHistoryEntry& Entry)
+{
+	Ar << Entry.ComponentSpaceRotations;
+	Ar << Entry.ComponentSpacePositions;
+	Ar << Entry.ComponentSpaceScales;
+	Ar << Entry.AccumulatedSeconds;
+	return Ar;
+}
+
 //////////////////////////////////////////////////////////////////////////
 // IPoseHistory
 #if ENABLE_DRAW_DEBUG && ENABLE_ANIM_DEBUG
 void IPoseHistory::DebugDraw(FAnimInstanceProxy& AnimInstanceProxy, FColor Color, float Time, float PointSize, bool bExtrapolate) const
 {
 	const FBoneContainer& BoneContainer = AnimInstanceProxy.GetRequiredBones();
-	if (BoneContainer.IsValid())
+	if (Color.A > 0 && BoneContainer.IsValid())
 	{
 		const USkeleton* Skeleton = BoneContainer.GetSkeletonAsset();
 		FTransform OutBoneTransform;
@@ -174,120 +267,189 @@
 #endif // ENABLE_DRAW_DEBUG && ENABLE_ANIM_DEBUG
 
 //////////////////////////////////////////////////////////////////////////
+// FArchivedPoseHistory
+void FArchivedPoseHistory::InitFrom(const IPoseHistory* PoseHistory)
+{
+	Trajectory.Samples.Reset();
+	BoneToTransformMap.Reset();
+	Entries.Reset();
+
+	if (PoseHistory)
+	{
+		Trajectory = PoseHistory->GetTrajectory();
+		BoneToTransformMap = PoseHistory->GetBoneToTransformMap();
+			
+		const int32 NumEntries = PoseHistory->GetNumEntries();
+		Entries.SetNum(NumEntries);
+
+		for (int32 EntryIndex = 0; EntryIndex < NumEntries; ++EntryIndex)
+		{
+			Entries[EntryIndex] = PoseHistory->GetEntry(EntryIndex);
+			// validating input PoseHistory to have Entries properly sorted by time
+			check(EntryIndex == 0 || (Entries[EntryIndex - 1].AccumulatedSeconds <= Entries[EntryIndex].AccumulatedSeconds));
+		}
+	}
+}
+
+bool FArchivedPoseHistory::GetTransformAtTime(float Time, FTransform& OutBoneTransform, const USkeleton* BoneIndexSkeleton, FBoneIndexType BoneIndexType, FBoneIndexType ReferenceBoneIndexType, bool bExtrapolate) const
+{
+	static_assert(RootBoneIndexType == 0 && ComponentSpaceIndexType == FBoneIndexType(-1) && WorldSpaceIndexType == FBoneIndexType(-2)); // some assumptions
+	check(BoneIndexType != ComponentSpaceIndexType && BoneIndexType != WorldSpaceIndexType);
+	
+	bool bSuccess = false;
+	
+	const bool bApplyComponentToWorld = ReferenceBoneIndexType == WorldSpaceIndexType;
+	FTransform ComponentToWorld = FTransform::Identity;
+	if (bApplyComponentToWorld)
+	{
+		ComponentToWorld = Trajectory.GetSampleAtTime(Time, bExtrapolate).GetTransform();
+		ReferenceBoneIndexType = ComponentSpaceIndexType;
+	}
+
+	const int32 NumEntries = Entries.Num();
+	if (NumEntries > 0)
+	{
+		int32 NextIdx = 0;
+		int32 PrevIdx = 0;
+
+		if (NumEntries > 1)
+		{
+			const int32 LowerBoundIdx = Algo::LowerBound(Entries, Time, [](const FPoseHistoryEntry& Entry, float Value) { return Value > Entry.AccumulatedSeconds; });
+			NextIdx = FMath::Clamp(LowerBoundIdx, 1, NumEntries - 1);
+			PrevIdx = NextIdx - 1;
+		}
+	
+		const FPoseHistoryEntry& PrevEntry = Entries[PrevIdx];
+		const FPoseHistoryEntry& NextEntry = Entries[NextIdx];
+
+		bSuccess = LerpEntries(Time, bExtrapolate, PrevEntry, NextEntry, BoneIndexSkeleton, nullptr, BoneToTransformMap, BoneIndexType, ReferenceBoneIndexType, OutBoneTransform);
+		if (bApplyComponentToWorld)
+		{
+			OutBoneTransform *= ComponentToWorld;
+		}
+	}
+	else
+	{
+		OutBoneTransform = ComponentToWorld;
+	}
+	
+	return bSuccess;
+}
+
+#if ENABLE_DRAW_DEBUG && ENABLE_ANIM_DEBUG
+void FArchivedPoseHistory::DebugDraw(const UWorld* World, FColor Color) const
+{
+	if (Color.A > 0 && !Trajectory.Samples.IsEmpty())
+	{
+		TArray<FTransform, TInlineAllocator<128>> PrevGlobalTransforms;
+
+		for (int32 EntryIndex = 0; EntryIndex < Entries.Num(); ++EntryIndex)
+		{
+			const FPoseHistoryEntry& Entry = Entries[EntryIndex];
+
+			const int32 PrevGlobalTransformsNum = PrevGlobalTransforms.Num();
+			const int32 Max = FMath::Max(PrevGlobalTransformsNum, Entry.Num());
+
+			PrevGlobalTransforms.SetNum(Max, false);
+
+			for (int32 i = 0; i < Entry.Num(); ++i)
+			{
+				const FTransform RootTransform = Trajectory.GetSampleAtTime(Entry.AccumulatedSeconds).GetTransform();
+				const FTransform GlobalTransforms = Entry.GetComponentSpaceTransform(i) * RootTransform;
+
+				if (i < PrevGlobalTransformsNum)
+				{
+					DrawDebugLine(World, PrevGlobalTransforms[i].GetTranslation(), GlobalTransforms.GetTranslation(), Color, false, -1.f, SDPG_Foreground);
+				}
+
+				PrevGlobalTransforms[i] = GlobalTransforms;
+			}
+		}
+	}
+}
+#endif // ENABLE_DRAW_DEBUG && ENABLE_ANIM_DEBUG
+
+FArchive& operator<<(FArchive& Ar, FArchivedPoseHistory& Entry)
+{
+	Ar << Entry.BoneToTransformMap;
+	Ar << Entry.Entries;
+	Ar << Entry.Trajectory;
+	return Ar;
+}
+
+//////////////////////////////////////////////////////////////////////////
 // FPoseHistory
-void FPoseHistory::CacheBones_AnyThread(int32 InNumPoses, float InSamplingInterval, const TArray<FBoneIndexType>& RequiredBones)
+FPoseHistory::FPoseHistory(const FPoseHistory& Other)
+	: FPoseHistory()
 {
+	*this = Other;
+}
+
+FPoseHistory::FPoseHistory(FPoseHistory&& Other)
+	: FPoseHistory()
+{
+	*this = MoveTemp(Other);
+}
+
+FPoseHistory& FPoseHistory::operator=(const FPoseHistory& Other)
+{
+	CheckThreadSafetyWrite(ReadDataThreadSafeCounter);
+	CheckThreadSafetyWrite(WriteDataThreadSafeCounter);
+
+	ReadData = Other.ReadData;
+	WriteData = Other.WriteData;
+	return *this;
+}
+
+FPoseHistory& FPoseHistory::operator=(FPoseHistory&& Other)
+{
+	CheckThreadSafetyWrite(ReadDataThreadSafeCounter);
+	CheckThreadSafetyWrite(WriteDataThreadSafeCounter);
+
+	ReadData = MoveTemp(Other.ReadData);
+	WriteData = MoveTemp(Other.WriteData);
+	return *this;
+}
+
+void FPoseHistory::Initialize_AnyThread(int32 InNumPoses, float InSamplingInterval)
+{
+	CheckThreadSafetyWrite(WriteDataThreadSafeCounter);
 	check(InNumPoses >= 2 && InSamplingInterval > UE_KINDA_SMALL_NUMBER);
 
 	MaxNumPoses = InNumPoses;
 	SamplingInterval = InSamplingInterval;
+}
 
-	BoneToTransformMap.Reset();
+void FPoseHistory::CacheBones_AnyThread(const TArray<FBoneIndexType>& RequiredBones)
+{
+	CheckThreadSafetyWrite(WriteDataThreadSafeCounter);
+
+	check(MaxNumPoses >= 2 && SamplingInterval > UE_KINDA_SMALL_NUMBER);
+
+	WriteData.BoneToTransformMap.Reset();
 	if (!RequiredBones.IsEmpty())
 	{
 		// making sure we always collect the root bone transform (by construction BoneToTransformMap[0] = 0)
-		BoneToTransformMap.Add(RootBoneIndexType) = BoneToTransformMap.Num();
+		WriteData.BoneToTransformMap.Add(RootBoneIndexType) = WriteData.BoneToTransformMap.Num();
 
 		for (int32 i = 0; i < RequiredBones.Num(); ++i)
 		{
 			// adding only unique RequiredBones to avoid oversizing Entries::ComponentSpaceTransforms
-			if (!BoneToTransformMap.Find(RequiredBones[i]))
+			if (!WriteData.BoneToTransformMap.Find(RequiredBones[i]))
 			{
-				BoneToTransformMap.Add(RequiredBones[i]) = BoneToTransformMap.Num();
+				WriteData.BoneToTransformMap.Add(RequiredBones[i]) = WriteData.BoneToTransformMap.Num();
 			}
 		}
 	}
 
-	Entries.Reset();
-	Entries.Reserve(MaxNumPoses);
+	WriteData.Entries.Reset();
+	WriteData.Entries.Reserve(MaxNumPoses);
 }
 
-FBoneIndexType FPoseHistory::GetRemappedBoneIndexType(FBoneIndexType BoneIndexType, const USkeleton* BoneIndexSkeleton, const USkeleton* LastUpdateSkeleton)
-{
-	// remapping BoneIndexType in case the skeleton used to store history (LastUpdateSkeleton) is different from BoneIndexSkeleton
-	if (LastUpdateSkeleton != nullptr && LastUpdateSkeleton != BoneIndexSkeleton)
-	{
-		const FSkeletonRemapping& SkeletonRemapping = UE::Anim::FSkeletonRemappingRegistry::Get().GetRemapping(BoneIndexSkeleton, LastUpdateSkeleton);
-		if (SkeletonRemapping.IsValid())
-		{
-			BoneIndexType = SkeletonRemapping.GetTargetSkeletonBoneIndex(BoneIndexType);
-		}
-	}
-
-	return BoneIndexType;
-}
-
-FComponentSpaceTransformIndex FPoseHistory::GetRemappedComponentSpaceTransformIndex(const USkeleton* BoneIndexSkeleton, const USkeleton* LastUpdateSkeleton, const FBoneToTransformMap& BoneToTransformMap, FBoneIndexType BoneIndexType, bool& bSuccess)
-{
-	check(BoneIndexType != WorldSpaceIndexType);
-
-	FComponentSpaceTransformIndex BoneTransformIndex = FComponentSpaceTransformIndex(BoneIndexType);
-	if (BoneIndexType != ComponentSpaceIndexType)
-	{
-		BoneIndexType = GetRemappedBoneIndexType(BoneIndexType, BoneIndexSkeleton, LastUpdateSkeleton);
-
-		if (!BoneToTransformMap.IsEmpty())
-		{
-			if (const FComponentSpaceTransformIndex* FoundBoneTransformIndex = BoneToTransformMap.Find(BoneTransformIndex))
-			{
-				BoneTransformIndex = *FoundBoneTransformIndex;
-			}
-			else
-			{
-				BoneTransformIndex = RootBoneIndexType;
-				bSuccess = false;
-			}
-		}
-	}
-	return BoneTransformIndex;
-}
-
-bool FPoseHistory::LerpEntries(float Time, bool bExtrapolate, const FPoseHistoryEntry& PrevEntry, const FPoseHistoryEntry& NextEntry, const USkeleton* BoneIndexSkeleton, const USkeleton* LastUpdateSkeleton,
-	const FBoneToTransformMap& BoneToTransformMap, FBoneIndexType BoneIndexType, FBoneIndexType ReferenceBoneIndexType, FTransform& OutBoneTransform)
-{
-	bool bSuccess = true;
-
-	const float Denominator = NextEntry.AccumulatedSeconds - PrevEntry.AccumulatedSeconds;
-	float LerpValue = 0.f;
-	if (!FMath::IsNearlyZero(Denominator))
-	{
-		const float Numerator = Time - PrevEntry.AccumulatedSeconds;
-		LerpValue = bExtrapolate ? Numerator / Denominator : FMath::Clamp(Numerator / Denominator, 0.f, 1.f);
-	}
-
-	const FComponentSpaceTransformIndex BoneTransformIndex = GetRemappedComponentSpaceTransformIndex(BoneIndexSkeleton, LastUpdateSkeleton, BoneToTransformMap, BoneIndexType, bSuccess);
-	const FComponentSpaceTransformIndex ReferenceBoneTransformIndex = GetRemappedComponentSpaceTransformIndex(BoneIndexSkeleton, LastUpdateSkeleton, BoneToTransformMap, ReferenceBoneIndexType, bSuccess);
-
-	if (BoneTransformIndex != ComponentSpaceIndexType)
-	{
-		if (ReferenceBoneTransformIndex == ComponentSpaceIndexType)
-		{
-			OutBoneTransform.Blend(
-				PrevEntry.GetComponentSpaceTransform(BoneTransformIndex),
-				NextEntry.GetComponentSpaceTransform(BoneTransformIndex),
-				LerpValue);
-		}
-		else
-		{
-			OutBoneTransform.Blend(
-				PrevEntry.GetComponentSpaceTransform(BoneTransformIndex) * PrevEntry.GetComponentSpaceTransform(ReferenceBoneTransformIndex).Inverse(),
-				NextEntry.GetComponentSpaceTransform(BoneTransformIndex) * NextEntry.GetComponentSpaceTransform(ReferenceBoneTransformIndex).Inverse(),
-				LerpValue);
-		}
-	}
-	else
-	{
-		// @todo: implement if required
-		OutBoneTransform = FTransform::Identity;
-		bSuccess = false;
-		unimplemented();
-	}
-
-	return bSuccess;
-}
-
 bool FPoseHistory::GetTransformAtTime(float Time, FTransform& OutBoneTransform, const USkeleton* BoneIndexSkeleton, FBoneIndexType BoneIndexType, FBoneIndexType ReferenceBoneIndexType, bool bExtrapolate) const
 {
+	CheckThreadSafetyRead(ReadDataThreadSafeCounter);
+
 	static_assert(RootBoneIndexType == 0 && ComponentSpaceIndexType == FBoneIndexType(-1) && WorldSpaceIndexType == FBoneIndexType(-2)); // some assumptions
 	check(BoneIndexType != ComponentSpaceIndexType && BoneIndexType != WorldSpaceIndexType);
 	
@@ -301,7 +463,7 @@
 		ReferenceBoneIndexType = ComponentSpaceIndexType;
 	}
 
-	const int32 NumEntries = Entries.Num();
+	const int32 NumEntries = ReadData.Entries.Num();
 	if (NumEntries > 0)
 	{
 		int32 NextIdx = 0;
@@ -309,15 +471,15 @@
 
 		if (NumEntries > 1)
 		{
-			const int32 LowerBoundIdx = LowerBound(Entries.begin(), Entries.end(), Time, [](const FPoseHistoryEntry& Entry, float Value) { return Value > Entry.AccumulatedSeconds; });
+			const int32 LowerBoundIdx = LowerBound(ReadData.Entries.begin(), ReadData.Entries.end(), Time, [](const FPoseHistoryEntry& Entry, float Value) { return Value > Entry.AccumulatedSeconds; });
 			NextIdx = FMath::Clamp(LowerBoundIdx, 1, NumEntries - 1);
 			PrevIdx = NextIdx - 1;
 		}
 	
-		const FPoseHistoryEntry& PrevEntry = Entries[PrevIdx];
-		const FPoseHistoryEntry& NextEntry = Entries[NextIdx];
+		const FPoseHistoryEntry& PrevEntry = ReadData.Entries[PrevIdx];
+		const FPoseHistoryEntry& NextEntry = ReadData.Entries[NextIdx];
 
-		bSuccess = LerpEntries(Time, bExtrapolate, PrevEntry, NextEntry, BoneIndexSkeleton, GetLastUpdateSkeleton(), BoneToTransformMap, BoneIndexType, ReferenceBoneIndexType, OutBoneTransform);
+		bSuccess = LerpEntries(Time, bExtrapolate, PrevEntry, NextEntry, BoneIndexSkeleton, ReadData.LastUpdateSkeleton.Get(), ReadData.BoneToTransformMap, BoneIndexType, ReferenceBoneIndexType, OutBoneTransform);
 		if (bApplyComponentToWorld)
 		{
 			OutBoneTransform *= ComponentToWorld;
@@ -331,90 +493,147 @@
 	return bSuccess;
 }
 
-void FPoseHistory::Update_AnyThread(float DeltaTime, const FPoseSearchQueryTrajectory& InTrajectory, float InTrajectorySpeedMultiplier, bool bGenerateTrajectory, const FAnimInstanceProxy& AnimInstanceProxy, const FPoseSearchTrajectoryData& TrajectoryData, const FPoseSearchTrajectoryData::FSampling& TrajectoryDataSampling, bool bNeedsReset)
+const FPoseSearchQueryTrajectory& FPoseHistory::GetTrajectory() const
 {
+	CheckThreadSafetyRead(ReadDataThreadSafeCounter);
+	return Trajectory;
+}
+
+float FPoseHistory::GetTrajectorySpeedMultiplier() const
+{
+	CheckThreadSafetyRead(ReadDataThreadSafeCounter);
+	return TrajectorySpeedMultiplier;
+}
+
+bool FPoseHistory::IsEmpty() const
+{
+	CheckThreadSafetyRead(ReadDataThreadSafeCounter);
+	return ReadData.Entries.IsEmpty();
+}
+
+const FBoneToTransformMap& FPoseHistory::GetBoneToTransformMap() const
+{
+	CheckThreadSafetyRead(ReadDataThreadSafeCounter);
+	return ReadData.BoneToTransformMap;
+}
+
+int32 FPoseHistory::GetNumEntries() const
+{
+	CheckThreadSafetyRead(ReadDataThreadSafeCounter);
+	return ReadData.Entries.Num();
+}
+
+const FPoseHistoryEntry& FPoseHistory::GetEntry(int32 EntryIndex) const
+{
+	CheckThreadSafetyRead(ReadDataThreadSafeCounter);
+	return ReadData.Entries[EntryIndex];
+}
+
+void FPoseHistory::PreUpdate(const UAnimInstance* AnimInstance, float DeltaTime, bool bGenerateTrajectory, const FPoseSearchTrajectoryData& TrajectoryData, const FPoseSearchTrajectoryData::FSampling& TrajectoryDataSampling, bool bNeedsReset)
+{
+	// we're writing both ReadData, WriteData. checking for thread safety
+	CheckThreadSafetyWrite(ReadDataThreadSafeCounter);
+	CheckThreadSafetyWrite(WriteDataThreadSafeCounter);
+
+	check(IsInGameThread());
+
 	if (bNeedsReset)
 	{
-		Entries.Reset();
+		WriteData.Entries.Reset();
 	}
 
 	if (bGenerateTrajectory)
 	{
 		// @todo: Synchronize the FPoseSearchQueryTrajectorySample::AccumulatedSeconds of the generated trajectory with the FPoseHistoryEntry::AccumulatedSeconds of the captured poses
 		FPoseSearchTrajectoryData::FDerived TrajectoryDataDerived;
-		TrajectoryData.UpdateData(DeltaTime, AnimInstanceProxy, TrajectoryDataDerived, TrajectoryDataState);
-		FPoseSearchTrajectoryLibrary::InitTrajectorySamples(Trajectory, TrajectoryData, TrajectoryDataDerived, TrajectoryDataSampling);
-		FPoseSearchTrajectoryLibrary::UpdateHistory_TransformHistory(Trajectory, TrajectoryData, TrajectoryDataDerived, TrajectoryDataSampling, DeltaTime);
-		FPoseSearchTrajectoryLibrary::UpdatePrediction_SimulateCharacterMovement(Trajectory, TrajectoryData, TrajectoryDataDerived, TrajectoryDataSampling);
+		TrajectoryData.UpdateData(DeltaTime, AnimInstance, TrajectoryDataDerived, TrajectoryDataState);
+		UPoseSearchTrajectoryLibrary::InitTrajectorySamples(Trajectory, TrajectoryData, TrajectoryDataDerived, TrajectoryDataSampling);
+		UPoseSearchTrajectoryLibrary::UpdateHistory_TransformHistory(Trajectory, TrajectoryData, TrajectoryDataDerived, TrajectoryDataSampling, DeltaTime);
+		UPoseSearchTrajectoryLibrary::UpdatePrediction_SimulateCharacterMovement(Trajectory, TrajectoryData, TrajectoryDataDerived, TrajectoryDataSampling);
 
 		// @todo: support TrajectorySpeedMultiplier
-		TrajectorySpeedMultiplier = 1.f;
+		//TrajectorySpeedMultiplier = 1.f;
 	}
-	else
+	
+	ReadData = WriteData;
+}
+
+void FPoseHistory::Update_AnyThread(float DeltaTime, const FPoseSearchQueryTrajectory& InTrajectory, float InTrajectorySpeedMultiplier)
+{
+	if (!InTrajectory.Samples.IsEmpty())
 	{
-		TrajectoryDataState = FPoseSearchTrajectoryData::FState();
-	Trajectory = InTrajectory;
+		CheckThreadSafetyWrite(ReadDataThreadSafeCounter);
 
+		// @todo: THIS IS NOT THREAD SAFE! in the contex of multi character motion matching (CheckThreadSafetyWrite will assert in case of improper usage)
+		Trajectory = InTrajectory;
+
 		TrajectorySpeedMultiplier = InTrajectorySpeedMultiplier;
-	if (!FMath::IsNearlyEqual(TrajectorySpeedMultiplier, 1.f))
-	{
-		const float TrajectorySpeedMultiplierInv = FMath::IsNearlyZero(TrajectorySpeedMultiplier) ? 1.f : 1.f / TrajectorySpeedMultiplier;
-		for (FPoseSearchQueryTrajectorySample& Sample : Trajectory.Samples)
+
+		if (!FMath::IsNearlyEqual(TrajectorySpeedMultiplier, 1.f))
 		{
-			Sample.AccumulatedSeconds *= TrajectorySpeedMultiplierInv;
+			const float TrajectorySpeedMultiplierInv = FMath::IsNearlyZero(TrajectorySpeedMultiplier) ? 1.f : 1.f / TrajectorySpeedMultiplier;
+			for (FPoseSearchQueryTrajectorySample& Sample : Trajectory.Samples)
+			{
+				Sample.AccumulatedSeconds *= TrajectorySpeedMultiplierInv;
+			}
 		}
 	}
 }
-}
 
-void FPoseHistory::EvaluateComponentSpace_AnyThread(float DeltaTime, FCSPose<FCompactPose>& ComponentSpacePose, bool bGenerateTrajectory, bool bStoreScales)
+void FPoseHistory::EvaluateComponentSpace_AnyThread(float DeltaTime, FCSPose<FCompactPose>& ComponentSpacePose, bool bStoreScales)
 {
+	CheckThreadSafetyWrite(WriteDataThreadSafeCounter);
+
+	check(MaxNumPoses >= 2 && SamplingInterval > UE_KINDA_SMALL_NUMBER);
+
 	const USkeleton* Skeleton = ComponentSpacePose.GetPose().GetBoneContainer().GetSkeletonAsset();
-	if (LastUpdateSkeleton != Skeleton)
+	if (WriteData.LastUpdateSkeleton != Skeleton)
 	{
 		// @todo: support a different USkeleton per FPoseHistoryEntry if required
-		Entries.Reset();
-		LastUpdateSkeleton = Skeleton;
+		WriteData.Entries.Reset();
+		WriteData.LastUpdateSkeleton = Skeleton;
 	}
 
 	// Age our elapsed times
-	for (FPoseHistoryEntry& Entry : Entries)
+	for (FPoseHistoryEntry& Entry : WriteData.Entries)
 	{
 		Entry.AccumulatedSeconds -= DeltaTime;
 	}
 
-	if (Entries.Num() != MaxNumPoses)
+	if (WriteData.Entries.Num() != MaxNumPoses)
 	{
 		// Consume every pose until the queue is full
-		Entries.Emplace();
+		WriteData.Entries.Emplace();
 	}
-	else if (Entries[Entries.Num() - 2].AccumulatedSeconds <= -SamplingInterval)
+	else if (WriteData.Entries[WriteData.Entries.Num() - 2].AccumulatedSeconds <= -SamplingInterval)
 	{
-		FPoseHistoryEntry EntryTemp = MoveTemp(Entries.First());
-		Entries.PopFront();
-		Entries.Emplace(MoveTemp(EntryTemp));
+		FPoseHistoryEntry EntryTemp = MoveTemp(WriteData.Entries.First());
+		WriteData.Entries.PopFront();
+		WriteData.Entries.Emplace(MoveTemp(EntryTemp));
 	}
 
 	// Regardless of the retention policy, we always update the most recent Entry
-	Entries.Last().Update(0.f, ComponentSpacePose, BoneToTransformMap, bStoreScales);
+	WriteData.Entries.Last().Update(0.f, ComponentSpacePose, WriteData.BoneToTransformMap, bStoreScales);
 }
 
 #if ENABLE_DRAW_DEBUG && ENABLE_ANIM_DEBUG
 void FPoseHistory::DebugDraw(FAnimInstanceProxy& AnimInstanceProxy, FColor Color) const
 {
+	CheckThreadSafetyRead(ReadDataThreadSafeCounter);
+
 	if (CVarAnimPoseHistoryDebugDrawTrajectory.GetValueOnAnyThread())
 	{
 		Trajectory.DebugDrawTrajectory(AnimInstanceProxy);
 	}
 
 	if (Color.A > 0 && CVarAnimPoseHistoryDebugDrawPose.GetValueOnAnyThread())
-{
+	{
 		const bool bValidTrajectory = !Trajectory.Samples.IsEmpty();
 		TArray<FTransform, TInlineAllocator<128>> PrevGlobalTransforms;
 
-	for (int32 EntryIndex = 0; EntryIndex < Entries.Num(); ++EntryIndex)
-	{
-		const FPoseHistoryEntry& Entry = Entries[EntryIndex];
+		for (int32 EntryIndex = 0; EntryIndex < ReadData.Entries.Num(); ++EntryIndex)
+		{
+			const FPoseHistoryEntry& Entry = ReadData.Entries[EntryIndex];
 
 			const int32 PrevGlobalTransformsNum = PrevGlobalTransforms.Num();
 			const int32 Max = FMath::Max(PrevGlobalTransformsNum, Entry.Num());
@@ -439,19 +658,14 @@
 #endif // ENABLE_DRAW_DEBUG && ENABLE_ANIM_DEBUG
 
 //////////////////////////////////////////////////////////////////////////
-// FExtendedPoseHistory
-void FExtendedPoseHistory::Init(const IPoseHistory* InPoseHistory)
+// FMemStackPoseHistory
+void FMemStackPoseHistory::Init(const IPoseHistory* InPoseHistory)
 {
 	check(InPoseHistory);
 	PoseHistory = InPoseHistory;
 }
 
-bool FExtendedPoseHistory::IsInitialized() const
+bool FMemStackPoseHistory::GetTransformAtTime(float Time, FTransform& OutBoneTransform, const USkeleton* BoneIndexSkeleton, FBoneIndexType BoneIndexType, FBoneIndexType ReferenceBoneIndexType, bool bExtrapolate) const
 {
-	return PoseHistory != nullptr;
-}
-
-bool FExtendedPoseHistory::GetTransformAtTime(float Time, FTransform& OutBoneTransform, const USkeleton* BoneIndexSkeleton, FBoneIndexType BoneIndexType, FBoneIndexType ReferenceBoneIndexType, bool bExtrapolate) const
-{
 	check(PoseHistory);
 	if (Time > 0.f)
@@ -469,11 +683,10 @@
 
 			const int32 LowerBoundIdx = Algo::LowerBound(FutureEntries, Time, [](const FPoseHistoryEntry& Entry, float Value) { return Value > Entry.AccumulatedSeconds; });
 			const int32 NextIdx = FMath::Min(LowerBoundIdx, Num - 1);
-			const FPoseHistoryEntries& PastEntries = GetEntries();
 			const FPoseHistoryEntry& NextEntry = FutureEntries[NextIdx];
-			const FPoseHistoryEntry& PrevEntry = NextIdx > 0 ? FutureEntries[NextIdx - 1] : !PastEntries.IsEmpty() ? PastEntries.First() : NextEntry;
+			const FPoseHistoryEntry& PrevEntry = NextIdx > 0 ? FutureEntries[NextIdx - 1] : PoseHistory->GetNumEntries() > 0 ? PoseHistory->GetEntry(PoseHistory->GetNumEntries() - 1) : NextEntry;
 						
-			const bool bSuccess = FPoseHistory::LerpEntries(Time, bExtrapolate, PrevEntry, NextEntry, BoneIndexSkeleton, GetLastUpdateSkeleton(), GetBoneToTransformMap(), BoneIndexType, ReferenceBoneIndexType, OutBoneTransform);
+			const bool bSuccess = LerpEntries(Time, bExtrapolate, PrevEntry, NextEntry, BoneIndexSkeleton, nullptr, GetBoneToTransformMap(), BoneIndexType, ReferenceBoneIndexType, OutBoneTransform);
 			if (bApplyComponentToWorld)
 			{
 				OutBoneTransform *= ComponentToWorld;
@@ -485,7 +698,7 @@
 	return PoseHistory->GetTransformAtTime(Time, OutBoneTransform, BoneIndexSkeleton, BoneIndexType, ReferenceBoneIndexType, bExtrapolate);
 }
 
-void FExtendedPoseHistory::AddFutureRootBone(float Time, const FTransform& FutureRootBoneTransform, bool bStoreScales)
+void FMemStackPoseHistory::AddFutureRootBone(float Time, const FTransform& FutureRootBoneTransform, bool bStoreScales)
 {
 	// we don't allow to add "past" or "present" poses to FutureEntries
 	check(Time > 0.f);
@@ -497,7 +710,7 @@
 	FutureEntry.AccumulatedSeconds = Time;
 }
 
-void FExtendedPoseHistory::AddFuturePose(float Time, FCSPose<FCompactPose>& ComponentSpacePose)
+void FMemStackPoseHistory::AddFuturePose(float Time, FCSPose<FCompactPose>& ComponentSpacePose)
 {
 	// we don't allow to add "past" or "present" poses to FutureEntries
 	check(Time > 0.f);
@@ -506,8 +719,27 @@
 	FutureEntries.InsertDefaulted_GetRef(LowerBoundIdx).Update(Time, ComponentSpacePose, GetBoneToTransformMap(), true);
 }
 
+int32 FMemStackPoseHistory::GetNumEntries() const
+{
+	check(PoseHistory);
+	return PoseHistory->GetNumEntries() + FutureEntries.Num();
+}
+
+const FPoseHistoryEntry& FMemStackPoseHistory::GetEntry(int32 EntryIndex) const
+{
+	check(PoseHistory);
+
+	const int32 PoseHistoryNumEntries = PoseHistory->GetNumEntries();
+	if (EntryIndex < PoseHistoryNumEntries)
+	{
+		return PoseHistory->GetEntry(EntryIndex);
+	}
+	return FutureEntries[EntryIndex - PoseHistoryNumEntries];
+}
+
+
 #if ENABLE_DRAW_DEBUG && ENABLE_ANIM_DEBUG
-void FExtendedPoseHistory::DebugDraw(FAnimInstanceProxy& AnimInstanceProxy, FColor Color) const
+void FMemStackPoseHistory::DebugDraw(FAnimInstanceProxy& AnimInstanceProxy, FColor Color) const
 {
 	check(PoseHistory);
 
@@ -517,10 +749,8 @@
 		const bool bValidTrajectory = !Trajectory.Samples.IsEmpty();
 		TArray<FTransform, TInlineAllocator<128>> PrevGlobalTransforms;
 
-		const FPoseHistoryEntries& PastEntries = GetEntries();
-
 		int32 EntriesNum = FutureEntries.Num();
-		if (!PastEntries.IsEmpty())
+		if (PoseHistory->GetNumEntries() > 0)
 		{
 			// connecting the future entries with the past entries
 			++EntriesNum;
@@ -528,7 +758,7 @@
 
 		for (int32 EntryIndex = 0; EntryIndex < EntriesNum; ++EntryIndex)
 		{
-			const FPoseHistoryEntry& Entry = (EntryIndex == FutureEntries.Num()) ? PastEntries.Last() : FutureEntries[EntryIndex];
+			const FPoseHistoryEntry& Entry = (EntryIndex == FutureEntries.Num()) ? PoseHistory->GetEntry(PoseHistory->GetNumEntries() - 1) : FutureEntries[EntryIndex];
 
 			const int32 PrevGlobalTransformsNum = PrevGlobalTransforms.Num();
 			const int32 Max = FMath::Max(PrevGlobalTransformsNum, Entry.Num());
diff --git a/Engine/Plugins/Experimental/Animation/PoseSearch/Source/Runtime/Private/PoseSearchHistoryCollectorAnimNodeLibrary.cpp b/Engine/Plugins/Experimental/Animation/PoseSearch/Source/Runtime/Private/PoseSearchHistoryCollectorAnimNodeLibrary.cpp
new file mode 100644
--- /dev/null
+++ b/Engine/Plugins/Experimental/Animation/PoseSearch/Source/Runtime/Private/PoseSearchHistoryCollectorAnimNodeLibrary.cpp
@@ -0,0 +1,17 @@
+// Copyright Epic Games, Inc. All Rights Reserved.
+
+#include "PoseSearch/PoseSearchHistoryCollectorAnimNodeLibrary.h"
+#include "PoseSearch/AnimNode_PoseSearchHistoryCollector.h"
+
+FPoseSearchHistoryCollectorAnimNodeReference UPoseSearchHistoryCollectorAnimNodeLibrary::ConvertToPoseHistoryNode(const FAnimNodeReference& Node, EAnimNodeReferenceConversionResult& Result)
+{
+	return FAnimNodeReference::ConvertToType<FPoseSearchHistoryCollectorAnimNodeReference>(Node, Result);
+}
+
+void UPoseSearchHistoryCollectorAnimNodeLibrary::GetPoseHistoryNodeTrajectory(const FPoseSearchHistoryCollectorAnimNodeReference& PoseSearchHistoryCollectorNode, FPoseSearchQueryTrajectory& Trajectory)
+{
+	if (FAnimNode_PoseSearchHistoryCollector* PoseSearchHistoryCollectorNodePtr = PoseSearchHistoryCollectorNode.GetAnimNodePtr<FAnimNode_PoseSearchHistoryCollector>())
+	{
+		Trajectory = PoseSearchHistoryCollectorNodePtr->GetPoseHistory().GetTrajectory();
+	}
+}
\ No newline at end of file
diff --git a/Engine/Plugins/Experimental/Animation/PoseSearch/Source/Runtime/Private/PoseSearchIndex.cpp b/Engine/Plugins/Experimental/Animation/PoseSearch/Source/Runtime/Private/PoseSearchIndex.cpp
--- a/Engine/Plugins/Experimental/Animation/PoseSearch/Source/Runtime/Private/PoseSearchIndex.cpp
+++ b/Engine/Plugins/Experimental/Animation/PoseSearch/Source/Runtime/Private/PoseSearchIndex.cpp
@@ -391,6 +391,7 @@ FArchive& operator<<(FArchive& Ar, FSearchStats& Stats)
 const FSearchIndexAsset& FSearchIndexBase::GetAssetForPose(int32 PoseIdx) const
 {
 	const uint32 AssetIndex = PoseMetadata[PoseIdx].GetAssetIndex();
+	check(Assets[AssetIndex].IsPoseInRange(PoseIdx));
 	return Assets[AssetIndex];
 }
 
diff --git a/Engine/Plugins/Experimental/Animation/PoseSearch/Source/Runtime/Private/PoseSearchLibrary.cpp b/Engine/Plugins/Experimental/Animation/PoseSearch/Source/Runtime/Private/PoseSearchLibrary.cpp
--- a/Engine/Plugins/Experimental/Animation/PoseSearch/Source/Runtime/Private/PoseSearchLibrary.cpp
+++ b/Engine/Plugins/Experimental/Animation/PoseSearch/Source/Runtime/Private/PoseSearchLibrary.cpp
@@ -12,6 +12,7 @@
 #include "Animation/BlendSpace.h"
 #include "Animation/BuiltInAttributeTypes.h"
 #include "Animation/AnimTrace.h"
+#include "GameFramework/Character.h"
 #include "InstancedStruct.h"
 #include "PoseSearch/AnimNode_MotionMatching.h"
 #include "PoseSearch/AnimNode_PoseSearchHistoryCollector.h"
@@ -27,8 +28,8 @@
 #define LOCTEXT_NAMESPACE "PoseSearchLibrary"
 
 #if ENABLE_DRAW_DEBUG && ENABLE_ANIM_DEBUG
-TAutoConsoleVariable<int32> CVarAnimMotionMatchDrawQueryEnable(TEXT("a.MotionMatch.DrawQuery.Enable"), 0, TEXT("Enable / Disable MotionMatch Draw Query"));
-TAutoConsoleVariable<int32> CVarAnimMotionMatchDrawMatchEnable(TEXT("a.MotionMatch.DrawMatch.Enable"), 0, TEXT("Enable / Disable MotionMatch Draw Match"));
+TAutoConsoleVariable<bool> CVarAnimMotionMatchDrawQueryEnable(TEXT("a.MotionMatch.DrawQuery.Enable"), false, TEXT("Enable / Disable MotionMatch Draw Query"));
+TAutoConsoleVariable<bool> CVarAnimMotionMatchDrawMatchEnable(TEXT("a.MotionMatch.DrawMatch.Enable"), false, TEXT("Enable / Disable MotionMatch Draw Match"));
 #endif
 
 namespace UE::PoseSearch
@@ -49,7 +50,7 @@ namespace UE::PoseSearch
 			return true;
 
 		default:
-			unimplemented();
+			checkNoEntry();
 			return false;
 		}
 	}
@@ -70,7 +71,7 @@ namespace UE::PoseSearch
 			return true;
 
 		default:
-			unimplemented();
+			checkNoEntry();
 			return false;
 		}
 	}
@@ -228,7 +229,6 @@ void UPoseSearchLibrary::TraceMotionMatchingState(
 	const UE::PoseSearch::FSearchResult& CurrentResult,
 	float ElapsedPoseSearchTime,
 	const FTransform& RootMotionTransformDelta,
-	const UObject* AnimInstance,
 	int32 NodeId,
 	float DeltaTime,
 	bool bSearch,
@@ -236,11 +236,36 @@ void UPoseSearchLibrary::TraceMotionMatchingState(
 {
 	using namespace UE::PoseSearch;
 	
-	const int32 CurrentPoseIdx = bSearch && CurrentResult.PoseCost.IsValid() ? CurrentResult.PoseIdx : INDEX_NONE;
-	FTraceMotionMatchingState TraceState;
-	TraceState.DatabaseEntries.SetNum(SearchContext.GetBestPoseCandidatesMap().Num());
+	FTraceMotionMatchingStateMessage TraceState;
 	
+	const int32 AnimInstancesNum = SearchContext.GetAnimInstances().Num();
+	TraceState.SkeletalMeshComponentIds.SetNum(AnimInstancesNum);
+
+	for (int32 AnimInstanceIndex = 0; AnimInstanceIndex < AnimInstancesNum; ++AnimInstanceIndex)
+	{
+		const UAnimInstance* AnimInstance = SearchContext.GetAnimInstances()[AnimInstanceIndex];
+		const UObject* SkeletalMeshComponent = AnimInstance->GetOuter();
+
+		TRACE_OBJECT(AnimInstance);
+
+		TraceState.SkeletalMeshComponentIds[AnimInstanceIndex] = FObjectTrace::GetObjectId(SkeletalMeshComponent);
+	}
+
+	TraceState.Roles.SetNum(AnimInstancesNum);
+	for (const FRoleToIndexPair& RoleToIndexPair : SearchContext.GetRoleToIndex())
+	{
+		TraceState.Roles[RoleToIndexPair.Value] = RoleToIndexPair.Key;
+	}
+
+	TraceState.PoseHistories.SetNum(AnimInstancesNum);
+	for (int32 AnimInstanceIndex = 0; AnimInstanceIndex < AnimInstancesNum; ++AnimInstanceIndex)
+	{
+		TraceState.PoseHistories[AnimInstanceIndex].InitFrom(SearchContext.GetPoseHistories()[AnimInstanceIndex]);
+	}
+
 	int32 DbEntryIdx = 0;
+	const int32 CurrentPoseIdx = bSearch && CurrentResult.PoseCost.IsValid() ? CurrentResult.PoseIdx : INDEX_NONE;
+	TraceState.DatabaseEntries.SetNum(SearchContext.GetBestPoseCandidatesMap().Num());
 	for (TPair<const UPoseSearchDatabase*, FSearchContext::FBestPoseCandidates> DatabaseBestPoseCandidates : SearchContext.GetBestPoseCandidatesMap())
 	{
 		const UPoseSearchDatabase* Database = DatabaseBestPoseCandidates.Key;
@@ -250,7 +275,7 @@ void UPoseSearchLibrary::TraceMotionMatchingState(
 
 		// if throttling is on, the continuing pose can be valid, but no actual search occurred, so the query will not be cached, and we need to build it
 		DbEntry.QueryVector = SearchContext.GetOrBuildQuery(Database->Schema);
-		DbEntry.DatabaseId = FTraceMotionMatchingState::GetIdFromObject(Database);
+		DbEntry.DatabaseId = FTraceMotionMatchingStateMessage::GetIdFromObject(Database);
 
 		for (int32 CandidateIdx = 0; CandidateIdx < DatabaseBestPoseCandidates.Value.Num(); ++CandidateIdx)
 		{
@@ -283,11 +308,21 @@ void UPoseSearchLibrary::TraceMotionMatchingState(
 		// simulation
 		if (SearchContext.AnyCachedQuery())
 		{
-			const FTransform PrevRoot = SearchContext.GetWorldBoneTransformAtTime(-DeltaTime);
-			const FTransform CurrRoot = SearchContext.GetWorldBoneTransformAtTime(0.f);
-			const FTransform SimDelta = CurrRoot.GetRelativeTransform(PrevRoot);
-			TraceState.SimLinearVelocity = SimDelta.GetTranslation().Size() / DeltaTime;
-			TraceState.SimAngularVelocity = FMath::RadiansToDegrees(SimDelta.GetRotation().GetAngle()) / DeltaTime;
+			TraceState.SimLinearVelocity = 0.f;
+			TraceState.SimAngularVelocity = 0.f;
+
+			const int32 NumRoles = SearchContext.GetRoleToIndex().Num();
+			for (const FRoleToIndexPair& RoleToIndexPair : SearchContext.GetRoleToIndex())
+			{
+				const FRole& Role = RoleToIndexPair.Key;
+
+				const FTransform PrevRoot = SearchContext.GetWorldBoneTransformAtTime(-DeltaTime, Role, RootSchemaBoneIdx);
+				const FTransform CurrRoot = SearchContext.GetWorldBoneTransformAtTime(0.f, Role, RootSchemaBoneIdx);
+				
+				const FTransform SimDelta = CurrRoot.GetRelativeTransform(PrevRoot);
+				TraceState.SimLinearVelocity += SimDelta.GetTranslation().Size() / (DeltaTime * NumRoles);
+				TraceState.SimAngularVelocity += FMath::RadiansToDegrees(SimDelta.GetRotation().GetAngle()) / (DeltaTime * NumRoles);
+			}
 		}
 
 		// animation
@@ -304,16 +339,17 @@ void UPoseSearchLibrary::TraceMotionMatchingState(
 	TraceState.SearchBruteForceCost = CurrentResult.BruteForcePoseCost.GetTotalCost();
 	TraceState.SearchBestPosePos = CurrentResult.BestPosePos;
 
-	if (SearchContext.GetHistory())
-	{
-		TraceState.Trajectory = SearchContext.GetHistory()->GetTrajectory();
-	}
-	else
-	{
-		TraceState.Trajectory.Samples.Reset();
-	}
+	// @todo: clean up this code
+	const UAnimInstance* AnimInstance = SearchContext.GetAnimInstances()[0];
+	UObject* SkeletalMeshComponent = AnimInstance->GetOuter();
 
-	TraceState.Output(AnimInstance, NodeId);
+	TraceState.Cycle = FPlatformTime::Cycles64();
+	TraceState.AnimInstanceId = FObjectTrace::GetObjectId(AnimInstance);
+	TraceState.SkeletalMeshComponentId = FObjectTrace::GetObjectId(SkeletalMeshComponent);
+	TraceState.NodeId = NodeId;
+	TraceState.FrameCounter = FObjectTrace::GetObjectWorldTickCounter(AnimInstance);
+
+	TraceState.Output();
 }
 #endif // UE_POSE_SEARCH_TRACE_ENABLED
 
@@ -356,7 +392,7 @@ void UPoseSearchLibrary::UpdateMotionMatchingState(
 	InOutMotionMatchingState.bJumpedToPose = false;
 
 	// used when YawFromAnimationBlendRate is greater than zero, by setting a future (YawFromAnimationTrajectoryBlendTime seconds ahead) root bone to the skeleton default
-	FExtendedPoseHistory ExtendedPoseHistory;
+	FMemStackPoseHistory MemStackPoseHistory;
 	const IPoseHistory* PoseHistory = nullptr;
 	if (IPoseHistoryProvider* PoseHistoryProvider = Context.GetMessage<IPoseHistoryProvider>())
 	{
@@ -364,16 +400,16 @@ void UPoseSearchLibrary::UpdateMotionMatchingState(
 
 		if (YawFromAnimationBlendRate >= 0.f && YawFromAnimationTrajectoryBlendTime > 0.f)
 		{
-			ExtendedPoseHistory.Init(PoseHistory);
+			MemStackPoseHistory.Init(PoseHistory);
 
 			const FTransform& RefRootBoneTransform = Context.AnimInstanceProxy->GetSkeleton()->GetReferenceSkeleton().GetRefBonePose()[RootBoneIndexType];
-			ExtendedPoseHistory.AddFutureRootBone(YawFromAnimationTrajectoryBlendTime, RefRootBoneTransform, false);
+			MemStackPoseHistory.AddFutureRootBone(YawFromAnimationTrajectoryBlendTime, RefRootBoneTransform, false);
 
 #if ENABLE_DRAW_DEBUG && ENABLE_ANIM_DEBUG
-			ExtendedPoseHistory.DebugDraw(*Context.AnimInstanceProxy, FColor::Emerald);
+			MemStackPoseHistory.DebugDraw(*Context.AnimInstanceProxy, FColor::Emerald);
 #endif // ENABLE_DRAW_DEBUG && ENABLE_ANIM_DEBUG
 
-			PoseHistory = &ExtendedPoseHistory;
+			PoseHistory = &MemStackPoseHistory;
 		}
 	}
 
@@ -387,8 +423,8 @@ void UPoseSearchLibrary::UpdateMotionMatchingState(
 		InOutMotionMatchingState.CurrentSearchResult.Reset();
 	}
 
-	FSearchContext SearchContext(AnimInstance, PoseHistory, TConstArrayView<const UObject*>(), 0.f,
-		&InOutMotionMatchingState.PoseIndicesHistory, InOutMotionMatchingState.CurrentSearchResult, PoseJumpThresholdTime);
+	FSearchContext SearchContext(0.f, &InOutMotionMatchingState.PoseIndicesHistory, InOutMotionMatchingState.CurrentSearchResult, PoseJumpThresholdTime);
+	SearchContext.AddRole(DefaultRole, AnimInstance, PoseHistory);
 
 	const bool bCanAdvance = InOutMotionMatchingState.CurrentSearchResult.CanAdvance(DeltaTime);
 
@@ -504,14 +540,14 @@ void UPoseSearchLibrary::UpdateMotionMatchingState(
 	if (IsTracing(Context))
 	{
 		TraceMotionMatchingState(SearchContext, InOutMotionMatchingState.CurrentSearchResult, InOutMotionMatchingState.ElapsedPoseSearchTime,
-			InOutMotionMatchingState.RootMotionTransformDelta, Context.AnimInstanceProxy->GetAnimInstanceObject(), Context.GetCurrentNodeId(), DeltaTime, bSearch,
+			InOutMotionMatchingState.RootMotionTransformDelta, Context.GetCurrentNodeId(), DeltaTime, bSearch,
 			AnimInstance ? FObjectTrace::GetWorldElapsedTime(AnimInstance->GetWorld()) : 0.f);
 	}
 #endif // UE_POSE_SEARCH_TRACE_ENABLED
 
 #if WITH_EDITORONLY_DATA && ENABLE_ANIM_DEBUG
 	const FSearchResult& CurResult = InOutMotionMatchingState.CurrentSearchResult;
-	if ((bDebugDrawQuery || bDebugDrawCurResult) && CurResult.Database != nullptr)
+	if (bDebugDrawQuery || bDebugDrawCurResult)
 	{
 		const UPoseSearchDatabase* CurResultDatabase = CurResult.Database.Get();
 
@@ -519,16 +555,18 @@ void UPoseSearchLibrary::UpdateMotionMatchingState(
 		if (FAsyncPoseSearchDatabasesManagement::RequestAsyncBuildIndex(CurResultDatabase, ERequestAsyncBuildFlag::ContinueRequest))
 #endif // WITH_EDITOR
 		{
+			FAnimInstanceProxy* AnimInstanceProxy = Context.AnimInstanceProxy;
+			const TArrayView<FAnimInstanceProxy*> AnimInstanceProxies = MakeArrayView(&AnimInstanceProxy, 1);
+			
 			if (bDebugDrawCurResult)
 			{
-				UE::PoseSearch::FDebugDrawParams DrawParams(Context.AnimInstanceProxy, SearchContext.GetWorldBoneTransformAtTime(0.f), CurResultDatabase);
+				FDebugDrawParams DrawParams(AnimInstanceProxies, SearchContext.GetPoseHistories(), SearchContext.GetRoleToIndex(), CurResultDatabase);
 				DrawParams.DrawFeatureVector(CurResult.PoseIdx);
 			}
 
 			if (bDebugDrawQuery)
 			{
-				// @todo: use pose history to get the root bone!
-				UE::PoseSearch::FDebugDrawParams DrawParams(Context.AnimInstanceProxy, SearchContext.GetWorldBoneTransformAtTime(0.f), CurResultDatabase, EDebugDrawFlags::DrawQuery);
+				FDebugDrawParams DrawParams(AnimInstanceProxies, SearchContext.GetPoseHistories(), SearchContext.GetRoleToIndex(), CurResultDatabase, EDebugDrawFlags::DrawQuery);
 				DrawParams.DrawFeatureVector(SearchContext.GetOrBuildQuery(CurResultDatabase->Schema));
 			}
 		}
@@ -545,6 +583,56 @@ void UPoseSearchLibrary::MotionMatch(
 	float FutureAnimationStartTime,
 	float TimeToFutureAnimationStart,
 	const int32 DebugSessionUniqueIdentifier)
+{
+	using namespace UE::PoseSearch;
+
+	FMemMark Mark(FMemStack::Get());
+
+	TArray<UAnimInstance*, TInlineAllocator<PreallocatedRolesNum, TMemStackAllocator<>>> AnimInstances;
+	AnimInstances.Add(AnimInstance);
+
+	TArray<FName, TInlineAllocator<PreallocatedRolesNum, TMemStackAllocator<>>> Roles;
+	Roles.Add(UE::PoseSearch::DefaultRole);
+
+	MotionMatch(AnimInstances, Roles, Database, PoseHistoryName, Result, FutureAnimation, FutureAnimationStartTime, TimeToFutureAnimationStart, DebugSessionUniqueIdentifier);
+}
+
+void UPoseSearchLibrary::MotionMatchMulti(
+	TArray<ACharacter*> Characters,
+	TArray<FName> Roles,
+	const UPoseSearchDatabase* Database,
+	const FName PoseHistoryName,
+	FPoseSearchBlueprintResult& Result,
+	const int32 DebugSessionUniqueIdentifier)
+{
+	using namespace UE::PoseSearch;
+
+	FMemMark Mark(FMemStack::Get());
+
+	TArray<UAnimInstance*, TInlineAllocator<PreallocatedRolesNum, TMemStackAllocator<>>> AnimInstances;
+	for (ACharacter* Character : Characters)
+	{
+		UAnimInstance* AnimInstance = nullptr;
+		if (Character && Character->GetMesh())
+		{
+			AnimInstance = Character->GetMesh()->GetAnimInstance();
+		}
+		AnimInstances.Add(AnimInstance);
+	}
+
+	MotionMatch(AnimInstances, Roles, Database, PoseHistoryName, Result, nullptr, 0.f, 0.f, DebugSessionUniqueIdentifier);
+}
+
+void UPoseSearchLibrary::MotionMatch(
+	TArrayView<UAnimInstance*> AnimInstances,
+	TConstArrayView<FName> Roles,
+	const UPoseSearchDatabase* Database,
+	const FName PoseHistoryName,
+	FPoseSearchBlueprintResult& Result,
+	const UAnimationAsset* FutureAnimation,
+	float FutureAnimationStartTime,
+	float TimeToFutureAnimationStart,
+	const int32 DebugSessionUniqueIdentifier)
 {
 #if ENABLE_DRAW_DEBUG && ENABLE_ANIM_DEBUG
 	class UAnimInstanceProxyProvider : public UAnimInstance
@@ -572,126 +660,172 @@ void UPoseSearchLibrary::MotionMatch(
 	Result.SelectedDatabase = nullptr;
 	Result.SearchCost = MAX_flt;
 
-	FMemMark Mark(FMemStack::Get());
-	if (Database && AnimInstance && AnimInstance->CurrentSkeleton)
+	if (!Database)
+	{
+		UE_LOG(LogPoseSearch, Error, TEXT("UPoseSearchLibrary::MotionMatch - null Database"));
+		return;
+	}
+
+	if (AnimInstances.IsEmpty() || AnimInstances.Num() != Roles.Num())
 	{
+		UE_LOG(LogPoseSearch, Error, TEXT("UPoseSearchLibrary::MotionMatch - invalid input AnimInstances or Roles"));
+		return;
+	}
+	
+	for (UAnimInstance* AnimInstance : AnimInstances)
+	{
+		if (!AnimInstance)
+		{
+			UE_LOG(LogPoseSearch, Error, TEXT("UPoseSearchLibrary::MotionMatch - null AnimInstances"));
+			return;
+		}
+
+		if (!AnimInstance->CurrentSkeleton)
+		{
+			UE_LOG(LogPoseSearch, Error, TEXT("UPoseSearchLibrary::MotionMatch - null AnimInstances->CurrentSkeleton"));
+			return;
+		}
+	}
+
+	FMemMark Mark(FMemStack::Get());
+
 #if ENABLE_DRAW_DEBUG && ENABLE_ANIM_DEBUG
-		FColor HistoryCollectorColor = FColor::Red;
+	TArray<FColor, TInlineAllocator<PreallocatedRolesNum, TMemStackAllocator<>>> HistoryCollectorColors;
 #endif // ENABLE_DRAW_DEBUG && ENABLE_ANIM_DEBUG
 
-		// ExtendedPoseHistory will hold future poses to match AssetSamplerBase (at FutureAnimationStartTime) TimeToFutureAnimationStart seconds in the future
-		FExtendedPoseHistory ExtendedPoseHistory;
+	// MemStackPoseHistories will hold future poses to match AssetSamplerBase (at FutureAnimationStartTime) TimeToFutureAnimationStart seconds in the future
+	TArray<FMemStackPoseHistory, TInlineAllocator<PreallocatedRolesNum, TMemStackAllocator<>>> MemStackPoseHistories;
+	for (UAnimInstance* AnimInstance : AnimInstances)
+	{
 		if (IAnimClassInterface* AnimBlueprintClass = IAnimClassInterface::GetFromClass(AnimInstance->GetClass()))
 		{
 			if (const FAnimSubsystem_Tag* TagSubsystem = AnimBlueprintClass->FindSubsystem<FAnimSubsystem_Tag>())
 			{
 				if (const FAnimNode_PoseSearchHistoryCollector_Base* PoseHistoryNode = TagSubsystem->FindNodeByTag<FAnimNode_PoseSearchHistoryCollector_Base>(PoseHistoryName, AnimInstance))
 				{
-					ExtendedPoseHistory.Init(&PoseHistoryNode->GetPoseHistory());
+					MemStackPoseHistories.AddDefaulted_GetRef().Init(&PoseHistoryNode->GetPoseHistory());
 
-#if ENABLE_DRAW_DEBUG && ENABLE_ANIM_DEBUG && WITH_EDITORONLY_DATA
-					HistoryCollectorColor = PoseHistoryNode->DebugColor.ToFColor(true);
+#if ENABLE_DRAW_DEBUG && ENABLE_ANIM_DEBUG
+	#if WITH_EDITORONLY_DATA
+					HistoryCollectorColors.Emplace(PoseHistoryNode->DebugColor.ToFColor(true));
+	#else // WITH_EDITORONLY_DATA
+					HistoryCollectorColors.Emplace(FColor::Red);
+	#endif // WITH_EDITORONLY_DATA
 #endif // ENABLE_DRAW_DEBUG && ENABLE_ANIM_DEBUG && WITH_EDITORONLY_DATA
 				}
 			}
 		}
+	}
 
-		if (!ExtendedPoseHistory.IsInitialized())
+	if (MemStackPoseHistories.Num() != AnimInstances.Num())
+	{
+		UE_LOG(LogPoseSearch, Error, TEXT("UPoseSearchLibrary::MotionMatch - Couldn't find pose history with name '%s'"), *PoseHistoryName.ToString());
+		return;
+	}
+
+	if (FutureAnimation)
+	{
+		if (AnimInstances.Num() != 1)
 		{
-			if (FutureAnimation)
-			{
-				UE_LOG(LogPoseSearch, Error, TEXT("UPoseSearchLibrary::MotionMatch - Couldn't find pose history with name '%s'. FutureAnimation search will not be performed"), *PoseHistoryName.ToString());
-			}
-			else
-			{
-				UE_LOG(LogPoseSearch, Warning, TEXT("UPoseSearchLibrary::MotionMatch - Couldn't find pose history with name '%s'"), *PoseHistoryName.ToString());
-			}
+			UE_LOG(LogPoseSearch, Error, TEXT("UPoseSearchLibrary::MotionMatch - FutureAnimation not supported for more than one AnimInstances"));
+			return;
 		}
-		else if (FutureAnimation)
+		
+		if (FutureAnimationStartTime < FiniteDelta)
 		{
-			const FBoneContainer& BoneContainer = AnimInstance->GetRequiredBonesOnAnyThread();
-			// @todo... add input BlendParameters to support sampling FutureAnimation blendspaces
-			const FAnimationAssetSampler Sampler(FutureAnimation, FVector::ZeroVector);
-
-			if (FutureAnimationStartTime < FiniteDelta)
-			{
-				UE_LOG(LogPoseSearch, Warning, TEXT("UPoseSearchLibrary::MotionMatch - provided FutureAnimationStartTime (%f) is too small to be able to calculate velocities. Clamping it to minimum value of %f"), FutureAnimationStartTime, FiniteDelta);
-				FutureAnimationStartTime = FiniteDelta;
-			}
-
-			const float MinTimeToFutureAnimationStart = FiniteDelta + UE_KINDA_SMALL_NUMBER;
-			if (TimeToFutureAnimationStart < MinTimeToFutureAnimationStart)
-			{
-				UE_LOG(LogPoseSearch, Warning, TEXT("UPoseSearchLibrary::MotionMatch - provided TimeToFutureAnimationStart (%f) is too small. Clamping it to minimum value of %f"), TimeToFutureAnimationStart, MinTimeToFutureAnimationStart);
-				TimeToFutureAnimationStart = MinTimeToFutureAnimationStart;
-			}
-
-			// extracting 2 poses to be able to calculate velocities
-			for (int32 i = 0; i < 2; ++i)
-			{
-				const float ExtractionTime = FutureAnimationStartTime + (i - 1) * FiniteDelta;
-				const float FutureAnimationTime = TimeToFutureAnimationStart + (i - 1) * FiniteDelta;
-
-				FCompactPose Pose;
-				Pose.SetBoneContainer(&BoneContainer);
-				Sampler.ExtractPose(ExtractionTime, Pose);
+			UE_LOG(LogPoseSearch, Warning, TEXT("UPoseSearchLibrary::MotionMatch - provided FutureAnimationStartTime (%f) is too small to be able to calculate velocities. Clamping it to minimum value of %f"), FutureAnimationStartTime, FiniteDelta);
+			FutureAnimationStartTime = FiniteDelta;
+		}
 
-				FCSPose<FCompactPose> ComponentSpacePose;
-				ComponentSpacePose.InitPose(Pose);
+		const float MinTimeToFutureAnimationStart = FiniteDelta + UE_KINDA_SMALL_NUMBER;
+		if (TimeToFutureAnimationStart < MinTimeToFutureAnimationStart)
+		{
+			UE_LOG(LogPoseSearch, Warning, TEXT("UPoseSearchLibrary::MotionMatch - provided TimeToFutureAnimationStart (%f) is too small. Clamping it to minimum value of %f"), TimeToFutureAnimationStart, MinTimeToFutureAnimationStart);
+			TimeToFutureAnimationStart = MinTimeToFutureAnimationStart;
+		}
 
-				ExtendedPoseHistory.AddFuturePose(FutureAnimationTime, ComponentSpacePose);
-			}
+		// extracting 2 poses to be able to calculate velocities
+		FCSPose<FCompactPose> ComponentSpacePose;
+		FCompactPose Pose;
+		Pose.SetBoneContainer(&AnimInstances[0]->GetRequiredBonesOnAnyThread());
+		
+		// @todo... add input BlendParameters to support sampling FutureAnimation blendspaces and support for multi character
+		const FAnimationAssetSampler Sampler(FutureAnimation);
+		for (int32 i = 0; i < 2; ++i)
+		{
+			const float ExtractionTime = FutureAnimationStartTime + (i - 1) * FiniteDelta;
+			const float FutureAnimationTime = TimeToFutureAnimationStart + (i - 1) * FiniteDelta;
+				
+			Sampler.ExtractPose(ExtractionTime, Pose);
+			ComponentSpacePose.InitPose(Pose);
+			MemStackPoseHistories[0].AddFuturePose(FutureAnimationTime, ComponentSpacePose);
+		}
 
 #if ENABLE_DRAW_DEBUG && ENABLE_ANIM_DEBUG
-				if (FAnimInstanceProxy* AnimInstanceProxy = UAnimInstanceProxyProvider::GetAnimInstanceProxy(AnimInstance))
-				{
-				ExtendedPoseHistory.DebugDraw(*AnimInstanceProxy, HistoryCollectorColor);
-				}
-#endif // ENABLE_DRAW_DEBUG && ENABLE_ANIM_DEBUG
+		if (FAnimInstanceProxy* AnimInstanceProxy = UAnimInstanceProxyProvider::GetAnimInstanceProxy(AnimInstances[0]))
+		{
+			MemStackPoseHistories[0].DebugDraw(*AnimInstanceProxy, HistoryCollectorColors[0]);
 		}
+#endif // ENABLE_DRAW_DEBUG && ENABLE_ANIM_DEBUG
+	}
 
-		FSearchContext SearchContext(AnimInstance, ExtendedPoseHistory.IsInitialized() ? &ExtendedPoseHistory : nullptr, TConstArrayView<const UObject*>(), TimeToFutureAnimationStart);
+	FSearchContext SearchContext(TimeToFutureAnimationStart);
+	for (int32 RoleIndex = 0; RoleIndex < Roles.Num(); ++RoleIndex)
+	{
+		SearchContext.AddRole(Roles[RoleIndex], AnimInstances[RoleIndex], MemStackPoseHistories[RoleIndex].GetThisOrPoseHistory());
+	}
 
-		FSearchResult SearchResult = Database->Search(SearchContext);
-		if (SearchResult.IsValid())
+	FSearchResult SearchResult = Database->Search(SearchContext);
+	if (SearchResult.IsValid())
+	{
+		const FSearchIndexAsset* SearchIndexAsset = SearchResult.GetSearchIndexAsset();
+		if (const FPoseSearchDatabaseAnimationAssetBase* DatabaseAsset = SearchResult.Database->GetAnimationAssetBase(*SearchIndexAsset))
 		{
-			const FSearchIndexAsset* SearchIndexAsset = SearchResult.GetSearchIndexAsset();
-			if (const FPoseSearchDatabaseAnimationAssetBase* DatabaseAsset = SearchResult.Database->GetAnimationAssetBase(*SearchIndexAsset))
-			{
-				Result.SelectedAnimation = DatabaseAsset->GetAnimationAsset();
-				Result.SelectedTime = SearchResult.AssetTime;
-				Result.bLoop = SearchIndexAsset->IsLooping();
-				Result.bIsMirrored = SearchIndexAsset->IsMirrored();
-				Result.BlendParameters = SearchIndexAsset->GetBlendParameters();
-				Result.SelectedDatabase = Database;
-				Result.SearchCost = SearchResult.PoseCost.GetTotalCost();
-			}
+			Result.SelectedAnimation = DatabaseAsset->GetAnimationAsset();
+			Result.SelectedTime = SearchResult.AssetTime;
+			Result.bLoop = SearchIndexAsset->IsLooping();
+			Result.bIsMirrored = SearchIndexAsset->IsMirrored();
+			Result.BlendParameters = SearchIndexAsset->GetBlendParameters();
+			Result.SelectedDatabase = Database;
+			Result.SearchCost = SearchResult.PoseCost.GetTotalCost();
 		}
+	}
 
 #if ENABLE_DRAW_DEBUG && ENABLE_ANIM_DEBUG
-		if (SearchResult.IsValid())
+	if (SearchResult.IsValid())
+	{
+		const bool bDrawMatch = CVarAnimMotionMatchDrawMatchEnable.GetValueOnAnyThread();
+		const bool bDrawquery = CVarAnimMotionMatchDrawQueryEnable.GetValueOnAnyThread();
+
+		if (bDrawMatch || bDrawquery)
 		{
-			FAnimInstanceProxy* AnimInstanceProxy = UAnimInstanceProxyProvider::GetAnimInstanceProxy(AnimInstance);
-			if (CVarAnimMotionMatchDrawMatchEnable.GetValueOnAnyThread())
+			TArray<FAnimInstanceProxy*, TInlineAllocator<PreallocatedRolesNum, TMemStackAllocator<>>> AnimInstanceProxies;
+			AnimInstanceProxies.SetNum(Roles.Num());
+			
+			for (int32 RoleIndex = 0; RoleIndex < Roles.Num(); ++RoleIndex)
 			{
-				FDebugDrawParams DrawParams(AnimInstanceProxy, SearchContext.GetWorldBoneTransformAtTime(0.f), SearchResult.Database.Get());
+				AnimInstanceProxies[RoleIndex] = UAnimInstanceProxyProvider::GetAnimInstanceProxy(AnimInstances[RoleIndex]);
+			}
+
+			if (bDrawMatch)
+			{
+				FDebugDrawParams DrawParams(AnimInstanceProxies, SearchContext.GetPoseHistories(), SearchContext.GetRoleToIndex(), SearchResult.Database.Get());
 				DrawParams.DrawFeatureVector(SearchResult.PoseIdx);
 			}
 
-			if (CVarAnimMotionMatchDrawQueryEnable.GetValueOnAnyThread())
+			if (bDrawquery)
 			{
-				FDebugDrawParams DrawParams(AnimInstanceProxy, SearchContext.GetWorldBoneTransformAtTime(0.f), SearchResult.Database.Get(), EDebugDrawFlags::DrawQuery);
+				FDebugDrawParams DrawParams(AnimInstanceProxies, SearchContext.GetPoseHistories(), SearchContext.GetRoleToIndex(), SearchResult.Database.Get(), EDebugDrawFlags::DrawQuery);
 				DrawParams.DrawFeatureVector(SearchContext.GetOrBuildQuery(SearchResult.Database->Schema));
 			}
 		}
+	}
 #endif // ENABLE_DRAW_DEBUG && ENABLE_ANIM_DEBUG
 
 #if UE_POSE_SEARCH_TRACE_ENABLED
-		
-		TraceMotionMatchingState(SearchContext, SearchResult, 0.f, FTransform::Identity, AnimInstance, DebugSessionUniqueIdentifier,
-			AnimInstance->GetDeltaSeconds(), true, FObjectTrace::GetWorldElapsedTime(AnimInstance->GetWorld()));
+	TraceMotionMatchingState(SearchContext, SearchResult, 0.f, FTransform::Identity, DebugSessionUniqueIdentifier,
+		AnimInstances[0]->GetDeltaSeconds(), true, FObjectTrace::GetWorldElapsedTime(AnimInstances[0]->GetWorld()));
 #endif // UE_POSE_SEARCH_TRACE_ENABLED
-	}
 }
 
 UE::PoseSearch::FSearchResult UPoseSearchLibrary::MotionMatch(const FAnimationBaseContext& Context, 
@@ -701,10 +835,10 @@ UE::PoseSearch::FSearchResult UPoseSearchLibrary::MotionMatch(const FAnimationBa
 
 	FSearchResult SearchResult;
 
-	const IPoseHistory* History = nullptr;
+	const IPoseHistory* PoseHistory = nullptr;
 	if (IPoseHistoryProvider* PoseHistoryProvider = Context.GetMessage<IPoseHistoryProvider>())
 	{
-		History = &PoseHistoryProvider->GetPoseHistory();
+		PoseHistory = &PoseHistoryProvider->GetPoseHistory();
 	}
 
 	const UAnimInstance* AnimInstance = Cast<const UAnimInstance>(Context.AnimInstanceProxy->GetAnimInstanceObject());
@@ -712,15 +846,16 @@ UE::PoseSearch::FSearchResult UPoseSearchLibrary::MotionMatch(const FAnimationBa
 
 	FMemMark Mark(FMemStack::Get());
 	FSearchResult ReconstructedPreviousSearchResult;
-	FSearchContext SearchContext(AnimInstance, History, TConstArrayView<const UObject*>(), 0.f, nullptr, ReconstructedPreviousSearchResult);
+	FSearchContext SearchContext(0.f, nullptr, ReconstructedPreviousSearchResult);
+	SearchContext.AddRole(DefaultRole, AnimInstance, PoseHistory);
 
 	// budgeting some stack allocations for simple use cases. bigger requests of AnimationAssets contining 
 	// UAnimNotifyState_PoseSearchBranchIn referencing multiple datbases will default to slower heap allocations
 	enum { MAX_STACK_ALLOCATED_ANIMATIONS = 16 };
 	enum { MAX_STACK_ALLOCATED_SETS = 2 };
-	typedef	TArray<const UObject*, TInlineAllocator<MAX_STACK_ALLOCATED_ANIMATIONS>> TAssetsToSearch;
+	typedef	TArray<const UObject*, TInlineAllocator<MAX_STACK_ALLOCATED_ANIMATIONS, TMemStackAllocator<>>> TAssetsToSearch;
 	// an empty TAssetsToSearch associated to Database means we need to search ALL the assets
-	typedef TMap<const UPoseSearchDatabase*, TAssetsToSearch, TInlineSetAllocator<MAX_STACK_ALLOCATED_SETS>> TAssetsToSearchPerDatabaseMap;
+	typedef TMap<const UPoseSearchDatabase*, TAssetsToSearch, TInlineSetAllocator<MAX_STACK_ALLOCATED_SETS, TMemStackSetAllocator<>>> TAssetsToSearchPerDatabaseMap;
 	typedef TPair<const UPoseSearchDatabase*, TAssetsToSearch> TAssetsToSearchPerDatabasePair;
 	TAssetsToSearchPerDatabaseMap AssetsToSearchPerDatabaseMap;
 	
@@ -734,6 +869,13 @@ UE::PoseSearch::FSearchResult UPoseSearchLibrary::MotionMatch(const FAnimationBa
 				{
 					if (PoseSearchBranchIn->Database)
 					{
+#if WITH_EDITOR
+						if (!PoseSearchBranchIn->Database->Contains(SequenceBase))
+						{
+							UE_LOG(LogPoseSearch, Error, TEXT("improperly setup UAnimSequenceBase. Database %s doesn't contain UAnimSequenceBase %s"), *PoseSearchBranchIn->Database->GetName(), *SequenceBase->GetName());
+						}
+						else 
+#endif // WITH_EDITOR
 						if (TAssetsToSearch* AssetsToSearch = AssetsToSearchPerDatabaseMap.Find(PoseSearchBranchIn->Database))
 						{
 							// an empty TAssetsToSearch associated to Database means we need to search ALL the assets, so we don't need to add this SequenceBase
@@ -835,16 +977,26 @@ UE::PoseSearch::FSearchResult UPoseSearchLibrary::MotionMatch(const FAnimationBa
 #if ENABLE_DRAW_DEBUG && ENABLE_ANIM_DEBUG
 	if (SearchResult.IsValid())
 	{
-		if (CVarAnimMotionMatchDrawMatchEnable.GetValueOnAnyThread())
-		{
-			FDebugDrawParams DrawParams(Context.AnimInstanceProxy, SearchContext.GetWorldBoneTransformAtTime(0.f), SearchResult.Database.Get());
-			DrawParams.DrawFeatureVector(SearchResult.PoseIdx);
-		}
+		const bool bDrawMatch = CVarAnimMotionMatchDrawMatchEnable.GetValueOnAnyThread();
+		const bool bDrawquery = CVarAnimMotionMatchDrawQueryEnable.GetValueOnAnyThread();
 
-		if (CVarAnimMotionMatchDrawQueryEnable.GetValueOnAnyThread())
+		if (bDrawMatch || bDrawquery)
 		{
-			FDebugDrawParams DrawParams(Context.AnimInstanceProxy, SearchContext.GetWorldBoneTransformAtTime(0.f), SearchResult.Database.Get(), EDebugDrawFlags::DrawQuery);
-			DrawParams.DrawFeatureVector(SearchContext.GetOrBuildQuery(SearchResult.Database->Schema));
+			// @todo: populate the FDebugDrawParams properly
+			FAnimInstanceProxy* AnimInstanceProxy = Context.AnimInstanceProxy;
+			const TArrayView<FAnimInstanceProxy*> AnimInstanceProxies = MakeArrayView(&AnimInstanceProxy, 1);
+
+			if (bDrawMatch)
+			{
+				FDebugDrawParams DrawParams(AnimInstanceProxies, SearchContext.GetPoseHistories(), SearchContext.GetRoleToIndex(), SearchResult.Database.Get());
+				DrawParams.DrawFeatureVector(SearchResult.PoseIdx);
+			}
+
+			if (bDrawquery)
+			{
+				FDebugDrawParams DrawParams(AnimInstanceProxies, SearchContext.GetPoseHistories(), SearchContext.GetRoleToIndex(), SearchResult.Database.Get(), EDebugDrawFlags::DrawQuery);
+				DrawParams.DrawFeatureVector(SearchContext.GetOrBuildQuery(SearchResult.Database->Schema));
+			}
 		}
 	}
 #endif // ENABLE_DRAW_DEBUG && ENABLE_ANIM_DEBUG
@@ -852,7 +1004,7 @@ UE::PoseSearch::FSearchResult UPoseSearchLibrary::MotionMatch(const FAnimationBa
 #if UE_POSE_SEARCH_TRACE_ENABLED
 	const float SearchBestCost = SearchResult.PoseCost.GetTotalCost();
 	const float SearchBruteForceCost = SearchResult.BruteForcePoseCost.GetTotalCost();
-	TraceMotionMatchingState(SearchContext, SearchResult, 0.f, FTransform::Identity, AnimInstance, Context.GetCurrentNodeId(),
+	TraceMotionMatchingState(SearchContext, SearchResult, 0.f, FTransform::Identity, Context.GetCurrentNodeId(),
 		AnimInstance->GetDeltaSeconds(), true, FObjectTrace::GetWorldElapsedTime(AnimInstance->GetWorld()));
 #endif // UE_POSE_SEARCH_TRACE_ENABLED
 
diff --git a/Engine/Plugins/Experimental/Animation/PoseSearch/Source/Runtime/Private/PoseSearchMultiSequence.cpp b/Engine/Plugins/Experimental/Animation/PoseSearch/Source/Runtime/Private/PoseSearchMultiSequence.cpp
new file mode 100644
--- /dev/null
+++ b/Engine/Plugins/Experimental/Animation/PoseSearch/Source/Runtime/Private/PoseSearchMultiSequence.cpp
@@ -0,0 +1,121 @@
+// Copyright Epic Games, Inc. All Rights Reserved.
+
+#include "PoseSearch/PoseSearchMultiSequence.h"
+#include "Animation/AnimSequenceBase.h"
+
+bool UPoseSearchMultiSequence::IsLooping() const
+{
+	float CommonPlayLength = -1.f;
+	for (const FPoseSearchMultiSequenceItem& Item : Items)
+	{
+		if (Item.Sequence)
+		{
+			if (!Item.Sequence->bLoop)
+			{
+				return false;
+			}
+
+			if (CommonPlayLength < 0.f)
+			{
+				CommonPlayLength = Item.Sequence->GetPlayLength();
+			}
+			else if (!FMath::IsNearlyEqual(CommonPlayLength, Item.Sequence->GetPlayLength()))
+			{
+				return false;
+			}
+		}
+	}
+	return true;
+}
+
+const FString UPoseSearchMultiSequence::GetName() const
+{
+	FString Name;
+	Name.Reserve(256);
+	bool bAddComma = false;
+	for (const FPoseSearchMultiSequenceItem& Item : Items)
+	{
+		if (bAddComma)
+		{
+			Name += ", ";
+		}
+		else
+		{
+			bAddComma = true;
+		}
+
+		Name += "[";
+		Name += Item.Role.ToString();
+		Name += "] ";
+		Name += GetNameSafe(Item.Sequence);
+	}
+	return Name;
+}
+
+bool UPoseSearchMultiSequence::HasRootMotion() const
+{
+	bool bHasAtLeastOneValidItem = false;
+	bool bHasRootMotion = true;
+
+	for (const FPoseSearchMultiSequenceItem& Item : Items)
+	{
+		if (Item.Sequence)
+		{
+			bHasRootMotion &= Item.Sequence->HasRootMotion();
+			bHasAtLeastOneValidItem = true;
+		}
+	}
+
+	return bHasAtLeastOneValidItem && bHasRootMotion;
+}
+
+float UPoseSearchMultiSequence::GetPlayLength() const
+{
+	float PlayLength = 0.f;
+	for (const FPoseSearchMultiSequenceItem& Item : Items)
+	{
+		if (Item.Sequence)
+		{
+			PlayLength = FMath::Max(PlayLength, Item.Sequence->GetPlayLength());
+		}
+	}
+	return PlayLength;
+}
+
+#if WITH_EDITOR
+int32 UPoseSearchMultiSequence::GetFrameAtTime(float Time) const
+{
+	for (const FPoseSearchMultiSequenceItem& Item : Items)
+	{
+		if (Item.Sequence)
+		{
+			return Item.Sequence->GetFrameAtTime(Time);
+		}
+	}
+	return 0;
+}
+#endif // WITH_EDITOR
+
+UAnimSequenceBase* UPoseSearchMultiSequence::GetSequence(const UE::PoseSearch::FRole& Role) const
+{
+	for (const FPoseSearchMultiSequenceItem& Item : Items)
+	{
+		if (Item.Role == Role)
+		{
+			return Item.Sequence;
+		}
+	}
+	return nullptr;
+}
+
+const FTransform& UPoseSearchMultiSequence::GetOrigin(const UE::PoseSearch::FRole& Role) const
+{
+	for (const FPoseSearchMultiSequenceItem& Item : Items)
+	{
+		if (Item.Role == Role)
+		{
+			return Item.Origin;
+		}
+	}
+	return FTransform::Identity;
+}
\ No newline at end of file
diff --git a/Engine/Plugins/Experimental/Animation/PoseSearch/Source/Runtime/Private/PoseSearchSchema.cpp b/Engine/Plugins/Experimental/Animation/PoseSearch/Source/Runtime/Private/PoseSearchSchema.cpp
--- a/Engine/Plugins/Experimental/Animation/PoseSearch/Source/Runtime/Private/PoseSearchSchema.cpp
+++ b/Engine/Plugins/Experimental/Animation/PoseSearch/Source/Runtime/Private/PoseSearchSchema.cpp
@@ -11,11 +11,6 @@
 #include "PoseSearchFeatureChannel_Trajectory.h"
 #include "UObject/ObjectSaveContext.h"
 
-bool UPoseSearchSchema::IsValid() const
-{
-	return Skeleton != nullptr;
-}
-
 void UPoseSearchSchema::AddChannel(UPoseSearchFeatureChannel* Channel)
 {
 	Channels.Add(Channel);
@@ -41,10 +36,35 @@ TConstArrayView<float> UPoseSearchSchema::BuildQuery(UE::PoseSearch::FSearchCont
 	return SearchContext.EditFeatureVector();
 }
 
-FBoneIndexType UPoseSearchSchema::GetBoneIndexType(int8 SchemaBoneIdx) const
+void UPoseSearchSchema::AddSkeleton(USkeleton* Skeleton, UMirrorDataTable* MirrorDataTable, const UE::PoseSearch::FRole& Role)
 {
-	check(BoneReferences[SchemaBoneIdx].HasValidSetup());
-	return BoneReferences[SchemaBoneIdx].BoneIndex;
+	FPoseSearchRoledSkeleton& RoledSkeleton = Skeletons.AddDefaulted_GetRef();
+	RoledSkeleton.Skeleton = Skeleton;
+	RoledSkeleton.MirrorDataTable = MirrorDataTable;
+	RoledSkeleton.Role = Role;
+}
+
+bool UPoseSearchSchema::AreSkeletonsCompatible(const UPoseSearchSchema* Other) const
+{
+	if (Skeletons.Num() != Other->Skeletons.Num())
+	{
+		return false;
+	}
+
+	for (int32 SkeletonIndex = 0; SkeletonIndex < Skeletons.Num(); ++SkeletonIndex)
+	{
+		if (Skeletons[SkeletonIndex].Skeleton != Other->Skeletons[SkeletonIndex].Skeleton)
+		{
+			return false;
+		}
+
+		if (Skeletons[SkeletonIndex].Role != Other->Skeletons[SkeletonIndex].Role)
+		{
+			return false;
+		}
+	}
+
+	return true;
 }
 
 void UPoseSearchSchema::AddDefaultChannels()
@@ -54,13 +74,108 @@ void UPoseSearchSchema::AddDefaultChannels()
 	AddChannel(NewObject<UPoseSearchFeatureChannel_Pose>(this, NAME_None, RF_Transactional));
 }
 
-int8 UPoseSearchSchema::AddBoneReference(const FBoneReference& BoneReference)
+void UPoseSearchSchema::InitBoneContainersFromRoledSkeleton(TMap<FName, FBoneContainer>& RoledBoneContainers) const
+{
+	RoledBoneContainers.Reset();
+	RoledBoneContainers.Reserve(Skeletons.Num());
+
+	for (const FPoseSearchRoledSkeleton& RoledSkeleton : Skeletons)
+	{
+		RoledBoneContainers.Add(RoledSkeleton.Role).InitializeTo(RoledSkeleton.BoneIndicesWithParents, UE::Anim::FCurveFilterSettings(UE::Anim::ECurveFilterMode::DisallowAll), *RoledSkeleton.Skeleton);
+	}
+}
+
+bool UPoseSearchSchema::AllRoledSkeletonHaveMirrorDataTable() const
+{
+	for (const FPoseSearchRoledSkeleton& RoledSkeleton : Skeletons)
+	{
+		if (!RoledSkeleton.MirrorDataTable)
+		{
+			return false;
+		}
+	}
+	return true;
+}
+
+const FPoseSearchRoledSkeleton* UPoseSearchSchema::GetRoledSkeleton(const UE::PoseSearch::FRole& Role) const
+{
+	for (const FPoseSearchRoledSkeleton& RoledSkeleton : Skeletons)
+	{
+		if (RoledSkeleton.Role == Role)
+		{
+			return &RoledSkeleton;
+		}
+	}
+	return nullptr;
+}
+
+FPoseSearchRoledSkeleton* UPoseSearchSchema::GetRoledSkeleton(const UE::PoseSearch::FRole& Role)
+{
+	for (FPoseSearchRoledSkeleton& RoledSkeleton : Skeletons)
+	{
+		if (RoledSkeleton.Role == Role)
+		{
+			return &RoledSkeleton;
+		}
+	}
+	return nullptr;
+}
+
+#if WITH_EDITOR
+const UE::PoseSearch::FRole UPoseSearchSchema::GetDefaultRole() const
+{
+	if (!Skeletons.IsEmpty())
+	{
+		return Skeletons[0].Role;
+	}
+	return UE::PoseSearch::DefaultRole;
+}
+#endif // WITH_EDITOR
+
+USkeleton* UPoseSearchSchema::GetSkeleton(const UE::PoseSearch::FRole& Role) const
+{
+	if (const FPoseSearchRoledSkeleton* RoledSkeleton = GetRoledSkeleton(Role))
+	{
+		return RoledSkeleton->Skeleton.Get();
+	}
+	return nullptr;
+}
+
+UMirrorDataTable* UPoseSearchSchema::GetMirrorDataTable(const UE::PoseSearch::FRole& Role) const
+{
+	if (const FPoseSearchRoledSkeleton* RoledSkeleton = GetRoledSkeleton(Role))
+	{
+		return RoledSkeleton->MirrorDataTable.Get();
+	}
+	return nullptr;
+}
+
+TConstArrayView<FBoneReference> UPoseSearchSchema::GetBoneReferences(const UE::PoseSearch::FRole& Role) const
+{
+	const FPoseSearchRoledSkeleton* RoledSkeleton = GetRoledSkeleton(Role);
+	check(RoledSkeleton);
+	return RoledSkeleton->BoneReferences;
+}
+
+int8 UPoseSearchSchema::AddBoneReference(const FBoneReference& BoneReference, const UE::PoseSearch::FRole& Role)
 {
 	using namespace UE::PoseSearch;
 
+	FPoseSearchRoledSkeleton* RoledSkeleton = GetRoledSkeleton(Role);
+	if (!RoledSkeleton)
+	{
+		UE_LOG(LogPoseSearch, Error, TEXT("UPoseSearchSchema::AddBoneReference: couldn't find data for the requested Role '%s' in UPoseSearchSchema '%s'"), *Role.ToString(), *GetNameSafe(this));
+		return -1;
+	}
+
 	int32 SchemaBoneIdx = 0;
-	check(Skeleton);
-	
+	const USkeleton* Skeleton = RoledSkeleton->Skeleton;
+	if (!Skeleton)
+	{
+		UE_LOG(LogPoseSearch, Error, TEXT("UPoseSearchSchema::AddBoneReference: couldn't find Skeleton with Role '%s' in UPoseSearchSchema '%s'"), *Role.ToString(), *GetNameSafe(this));
+		return -1;
+	}
+
 	bool bDefaultToRootBone = true;
 	FBoneReference TempBoneReference = BoneReference;
 	if (TempBoneReference.BoneName != NAME_None)
@@ -68,52 +183,84 @@ int8 UPoseSearchSchema::AddBoneReference(const FBoneReference& BoneReference)
 		TempBoneReference.Initialize(Skeleton);
 		if (!TempBoneReference.HasValidSetup())
 		{
-			UE_LOG(LogPoseSearch, Error, TEXT("UPoseSearchSchema::AddBoneReference: couldn't initialize FBoneReference '%s' with Skeleton '%s' in UPoseSearchSchema '%s'. Defaulting to root bone instead"),
-				*TempBoneReference.BoneName.ToString(), *GetNameSafe(Skeleton), *GetNameSafe(this));
-		}
-		else
-		{
-			bDefaultToRootBone = false;
+			UE_LOG(LogPoseSearch, Error, TEXT("UPoseSearchSchema::AddBoneReference: couldn't initialize FBoneReference '%s' with Skeleton '%s' with Role '%s' in UPoseSearchSchema '%s'"),
+				*TempBoneReference.BoneName.ToString(), *GetNameSafe(Skeleton), *Role.ToString(), *GetNameSafe(this));
+			return -1;
 		}
 	}
-
-	if (bDefaultToRootBone)
+	else
 	{
 		TempBoneReference.BoneName = Skeleton->GetReferenceSkeleton().GetBoneName(int32(RootBoneIndexType));
 		TempBoneReference.Initialize(Skeleton);
 		check(TempBoneReference.HasValidSetup());
 	}
 
-	SchemaBoneIdx = BoneReferences.AddUnique(TempBoneReference);
+	SchemaBoneIdx = RoledSkeleton->BoneReferences.AddUnique(TempBoneReference);
 	check(SchemaBoneIdx >= 0 && SchemaBoneIdx < 128);
 	return int8(SchemaBoneIdx);
 }
 
-void UPoseSearchSchema::Finalize()
+void UPoseSearchSchema::ResetFinalize()
 {
-	using namespace UE::PoseSearch;
+	for (FPoseSearchRoledSkeleton& RoledSkeleton : Skeletons)
+	{
+		RoledSkeleton.BoneReferences.Reset();
+		RoledSkeleton.BoneIndicesWithParents.Reset();
+	}
 
-	BoneReferences.Reset();
-	BoneIndicesWithParents.Reset();
 	FinalizedChannels.Reset();
 	SchemaCardinality = 0;
+}
 
-	if (!Skeleton)
-	{
-		UE_LOG(LogPoseSearch, Error, TEXT("UPoseSearchSchema::Finalize: couldn't Finalize '%s' because Skeleton is null"), *GetNameSafe(this));
-		return;
-	}
+void UPoseSearchSchema::Finalize()
+{
+	using namespace UE::PoseSearch;
+
+	ResetFinalize();
 
 	// adding as first bone reference the root bone
-	const int8 SchemaBoneIdx = AddBoneReference(FBoneReference());
-	check(SchemaBoneIdx == RootSchemaBoneIdx);
+	for (int32 RoledSkeletonIndex = 0; RoledSkeletonIndex < Skeletons.Num(); ++RoledSkeletonIndex)
+	{
+		const FPoseSearchRoledSkeleton& RoledSkeleton = Skeletons[RoledSkeletonIndex];
+		for (int32 ComparisonIndex = RoledSkeletonIndex + 1; ComparisonIndex < Skeletons.Num(); ++ComparisonIndex)
+		{
+			if (Skeletons[ComparisonIndex].Role == RoledSkeleton.Role)
+			{
+				UE_LOG(LogPoseSearch, Error, TEXT("UPoseSearchSchema::Finalize: couldn't Finalize '%s' because of duplicate Role '%s' in Skeletons"), *GetNameSafe(this), *RoledSkeleton.Role.ToString());
+
+				ResetFinalize();
+				return;
+			}
+		}
+
+		const int8 SchemaBoneIdx = AddBoneReference(FBoneReference(), RoledSkeleton.Role);
+		if (SchemaBoneIdx != RootSchemaBoneIdx)
+		{
+			UE_LOG(LogPoseSearch, Error, TEXT("UPoseSearchSchema::Finalize: couldn't Finalize '%s' because couldn't initialize root bone properly"), *GetNameSafe(this));
+
+			ResetFinalize();
+			return;
+		}
+	}
 
 	for (const TObjectPtr<UPoseSearchFeatureChannel>& ChannelPtr : Channels)
 	{
 		if (ChannelPtr)
 		{
 			FinalizedChannels.Add(ChannelPtr);
-			ChannelPtr->Finalize(this);
+			if (!ChannelPtr->Finalize(this))
+			{
+				#if WITH_EDITOR
+				TLabelBuilder LabelBuilder;
+				FString Label = ChannelPtr->GetLabel(LabelBuilder).ToString();
+				UE_LOG(LogPoseSearch, Error, TEXT("UPoseSearchSchema::Finalize: couldn't Finalize '%s' because of Channel '%s'"), *GetNameSafe(this), *Label);
+				#else // WITH_EDITOR
+				UE_LOG(LogPoseSearch, Error, TEXT("UPoseSearchSchema::Finalize: couldn't Finalize '%s' because of Channel '%s'"), *GetNameSafe(this), *GetNameSafe(ChannelPtr));
+				#endif // WITH_EDITOR
+
+				ResetFinalize();
+				return;
+			}
 		}
 	}
 
@@ -148,16 +295,19 @@ void UPoseSearchSchema::Finalize()
 	}
 
 	// Initialize references to obtain bone indices and fill out bone index array
-	for (FBoneReference& BoneRef : BoneReferences)
+	for (FPoseSearchRoledSkeleton& RoledSkeleton : Skeletons)
 	{
-		check(BoneRef.HasValidSetup());
-		BoneIndicesWithParents.Add(BoneRef.BoneIndex);
-	}
+		for (FBoneReference& BoneRef : RoledSkeleton.BoneReferences)
+		{
+			check(BoneRef.HasValidSetup());
+			RoledSkeleton.BoneIndicesWithParents.Add(BoneRef.BoneIndex);
+		}
 
-	// Build separate index array with parent indices guaranteed to be present. Sort for EnsureParentsPresent.
-	check(!BoneIndicesWithParents.IsEmpty());
-	BoneIndicesWithParents.Sort();
-	FAnimationRuntime::EnsureParentsPresent(BoneIndicesWithParents, Skeleton->GetReferenceSkeleton());
+		// Build separate index array with parent indices guaranteed to be present. Sort for EnsureParentsPresent.
+		check(!RoledSkeleton.BoneIndicesWithParents.IsEmpty());
+		RoledSkeleton.BoneIndicesWithParents.Sort();
+		FAnimationRuntime::EnsureParentsPresent(RoledSkeleton.BoneIndicesWithParents, RoledSkeleton.Skeleton->GetReferenceSkeleton());
+	}
 }
 
 void UPoseSearchSchema::PreSave(FObjectPreSaveContext ObjectSaveContext)
@@ -169,13 +319,27 @@ void UPoseSearchSchema::PreSave(FObjectPreSaveContext ObjectSaveContext)
 void UPoseSearchSchema::PostLoad()
 {
 	Super::PostLoad();
-	Finalize();
-}
 
-USkeleton* UPoseSearchSchema::GetSkeleton(bool& bInvalidSkeletonIsError, const IPropertyHandle* PropertyHandle)
-{
-	bInvalidSkeletonIsError = false;
-	return Skeleton;
+	if (Skeleton_DEPRECATED)
+	{
+		Skeletons.AddDefaulted_GetRef().Skeleton = Skeleton_DEPRECATED;
+		Skeleton_DEPRECATED = nullptr;
+	}
+
+	if (MirrorDataTable_DEPRECATED)
+	{
+		if (Skeletons.IsEmpty())
+		{
+			Skeletons.AddDefaulted_GetRef().MirrorDataTable = MirrorDataTable_DEPRECATED;
+		}
+		else
+		{
+			Skeletons[0].MirrorDataTable = MirrorDataTable_DEPRECATED;
+		}
+		MirrorDataTable_DEPRECATED = nullptr;
+	}
+
+	Finalize();
 }
 
 #if WITH_EDITOR
diff --git a/Engine/Plugins/Experimental/Animation/PoseSearch/Source/Runtime/Private/PoseSearchTrajectoryLibrary.cpp b/Engine/Plugins/Experimental/Animation/PoseSearch/Source/Runtime/Private/PoseSearchTrajectoryLibrary.cpp
--- a/Engine/Plugins/Experimental/Animation/PoseSearch/Source/Runtime/Private/PoseSearchTrajectoryLibrary.cpp
+++ b/Engine/Plugins/Experimental/Animation/PoseSearch/Source/Runtime/Private/PoseSearchTrajectoryLibrary.cpp
@@ -8,14 +8,13 @@
 
 void FPoseSearchTrajectoryData::UpdateData(float DeltaTime, const FAnimInstanceProxy& AnimInstanceProxy, FDerived& TrajectoryDataDerived, FState& TrajectoryDataState) const
 {
-	// An AnimInstance might call this during an AnimBP recompile with 0 delta time.
-	if (DeltaTime <= 0.f)
-	{
-		return;
-	}
+	UpdateData(DeltaTime, Cast<const UAnimInstance>(AnimInstanceProxy.GetAnimInstanceObject()), TrajectoryDataDerived, TrajectoryDataState);
+}
 
-	const UAnimInstance* AnimInstance = Cast<const UAnimInstance>(AnimInstanceProxy.GetAnimInstanceObject());
-	if (!AnimInstance)
+void FPoseSearchTrajectoryData::UpdateData(float DeltaTime, const UAnimInstance* AnimInstance, FDerived& TrajectoryDataDerived, FState& TrajectoryDataState) const
+{
+	// An AnimInstance might call this during an AnimBP recompile with 0 delta time.
+	if (DeltaTime <= 0.f || !AnimInstance)
 	{
 		return;
 	}
@@ -137,7 +136,7 @@ FVector FPoseSearchTrajectoryData::StepCharacterMovementGroundPrediction(float D
 	return OutVelocity;
 }
 
-void FPoseSearchTrajectoryLibrary::InitTrajectorySamples(FPoseSearchQueryTrajectory& Trajectory, const FPoseSearchTrajectoryData& TrajectoryData, const FPoseSearchTrajectoryData::FDerived& TrajectoryDataDerived, const FPoseSearchTrajectoryData::FSampling& TrajectoryDataSampling)
+void UPoseSearchTrajectoryLibrary::InitTrajectorySamples(FPoseSearchQueryTrajectory& Trajectory, const FPoseSearchTrajectoryData& TrajectoryData, const FPoseSearchTrajectoryData::FDerived& TrajectoryDataDerived, const FPoseSearchTrajectoryData::FSampling& TrajectoryDataSampling)
 {
 	const int32 NumHistorySamples = TrajectoryDataSampling.NumHistorySamples;
 	const int32 NumPredictionSamples = TrajectoryDataSampling.NumPredictionSamples;
@@ -169,7 +168,7 @@ void FPoseSearchTrajectoryLibrary::InitTrajectorySamples(FPoseSearchQueryTraject
 	}
 }
 
-void FPoseSearchTrajectoryLibrary::UpdateHistory_TransformHistory(FPoseSearchQueryTrajectory& Trajectory, const FPoseSearchTrajectoryData& TrajectoryData, const FPoseSearchTrajectoryData::FDerived& TrajectoryDataDerived, const FPoseSearchTrajectoryData::FSampling& TrajectoryDataSampling, float DeltaTime)
+void UPoseSearchTrajectoryLibrary::UpdateHistory_TransformHistory(FPoseSearchQueryTrajectory& Trajectory, const FPoseSearchTrajectoryData& TrajectoryData, const FPoseSearchTrajectoryData::FDerived& TrajectoryDataDerived, const FPoseSearchTrajectoryData::FSampling& TrajectoryDataSampling, float DeltaTime)
 {
 	const int32 NumHistorySamples = TrajectoryDataSampling.NumHistorySamples;
 	const float SecondsPerHistorySample = TrajectoryDataSampling.SecondsPerHistorySample;
@@ -214,7 +213,7 @@ void FPoseSearchTrajectoryLibrary::UpdateHistory_TransformHistory(FPoseSearchQue
 	}
 }
 
-FVector FPoseSearchTrajectoryLibrary::RemapVectorMagnitudeWithCurve(const FVector& Vector, bool bUseCurve, const FRuntimeFloatCurve& Curve)
+FVector UPoseSearchTrajectoryLibrary::RemapVectorMagnitudeWithCurve(const FVector& Vector, bool bUseCurve, const FRuntimeFloatCurve& Curve)
 {
 	if (bUseCurve)
 	{
@@ -229,11 +228,11 @@ FVector FPoseSearchTrajectoryLibrary::RemapVectorMagnitudeWithCurve(const FVecto
 	return Vector;
 }
 
-void FPoseSearchTrajectoryLibrary::UpdatePrediction_SimulateCharacterMovement(FPoseSearchQueryTrajectory& Trajectory, const FPoseSearchTrajectoryData& TrajectoryData, const FPoseSearchTrajectoryData::FDerived& TrajectoryDataDerived, const FPoseSearchTrajectoryData::FSampling& TrajectoryDataSampling)
+void UPoseSearchTrajectoryLibrary::UpdatePrediction_SimulateCharacterMovement(FPoseSearchQueryTrajectory& Trajectory, const FPoseSearchTrajectoryData& TrajectoryData, const FPoseSearchTrajectoryData::FDerived& TrajectoryDataDerived, const FPoseSearchTrajectoryData::FSampling& TrajectoryDataSampling)
 {
 	FVector CurrentPositionWS = TrajectoryDataDerived.Position;
-	FVector CurrentVelocityWS = FPoseSearchTrajectoryLibrary::RemapVectorMagnitudeWithCurve(TrajectoryDataDerived.Velocity, TrajectoryData.bUseSpeedRemappingCurve, TrajectoryData.SpeedRemappingCurve);
-	FVector CurrentAccelerationWS = FPoseSearchTrajectoryLibrary::RemapVectorMagnitudeWithCurve(TrajectoryDataDerived.Acceleration, TrajectoryData.bUseAccelerationRemappingCurve, TrajectoryData.AccelerationRemappingCurve);
+	FVector CurrentVelocityWS = RemapVectorMagnitudeWithCurve(TrajectoryDataDerived.Velocity, TrajectoryData.bUseSpeedRemappingCurve, TrajectoryData.SpeedRemappingCurve);
+	FVector CurrentAccelerationWS = RemapVectorMagnitudeWithCurve(TrajectoryDataDerived.Acceleration, TrajectoryData.bUseAccelerationRemappingCurve, TrajectoryData.AccelerationRemappingCurve);
 
 	// bending CurrentVelocityWS towards CurrentAccelerationWS
 	if (TrajectoryData.BendVelocityTowardsAcceleration > UE_KINDA_SMALL_NUMBER && !CurrentAccelerationWS.IsNearlyZero())
@@ -287,12 +286,12 @@ void FPoseSearchTrajectoryLibrary::UpdatePrediction_SimulateCharacterMovement(FP
 
 			// Account for the controller (e.g. the camera) rotating.
 			CurrentFacingWS = ControllerRotationPerStep * CurrentFacingWS;
-			CurrentAccelerationWS = FPoseSearchTrajectoryLibrary::RemapVectorMagnitudeWithCurve(ControllerRotationPerStep * CurrentAccelerationWS,
+			CurrentAccelerationWS = RemapVectorMagnitudeWithCurve(ControllerRotationPerStep * CurrentAccelerationWS,
 				TrajectoryData.bUseAccelerationRemappingCurve, TrajectoryData.AccelerationRemappingCurve);
 
 			const FVector NewVelocityWS = TrajectoryData.StepCharacterMovementGroundPrediction(SecondsPerPredictionSample, CurrentVelocityWS, CurrentAccelerationWS, TrajectoryDataDerived);
 
-			CurrentVelocityWS = FPoseSearchTrajectoryLibrary::RemapVectorMagnitudeWithCurve(NewVelocityWS, TrajectoryData.bUseSpeedRemappingCurve, TrajectoryData.SpeedRemappingCurve);
+			CurrentVelocityWS = RemapVectorMagnitudeWithCurve(NewVelocityWS, TrajectoryData.bUseSpeedRemappingCurve, TrajectoryData.SpeedRemappingCurve);
 
 			if (TrajectoryDataDerived.bOrientRotationToMovement && !CurrentAccelerationWS.IsNearlyZero())
 			{
@@ -302,4 +301,29 @@ void FPoseSearchTrajectoryLibrary::UpdatePrediction_SimulateCharacterMovement(FP
 			}
 		}
 	}
+}
+
+void UPoseSearchTrajectoryLibrary::PoseSearchGenerateTrajectory(
+	const UAnimInstance* InAnimInstance, UPARAM(ref) const FPoseSearchTrajectoryData& InTrajectoryData, float InDeltaTime,
+	UPARAM(ref) FPoseSearchQueryTrajectory& InOutTrajectory, UPARAM(ref) float& InOutDesiredControllerYawLastUpdate, FPoseSearchQueryTrajectory& OutTrajectory,
+	float InHistorySamplingInterval, int32 InTrajectoryHistoryCount, float InPredictionSamplingInterval, int32 InTrajectoryPredictionCount)
+{
+	FPoseSearchTrajectoryData::FSampling TrajectoryDataSampling;
+	TrajectoryDataSampling.NumHistorySamples = InTrajectoryHistoryCount;
+	TrajectoryDataSampling.SecondsPerHistorySample = InHistorySamplingInterval;
+	TrajectoryDataSampling.NumPredictionSamples = InTrajectoryPredictionCount;
+	TrajectoryDataSampling.SecondsPerPredictionSample = InPredictionSamplingInterval;
+
+	FPoseSearchTrajectoryData::FState TrajectoryDataState;
+	TrajectoryDataState.DesiredControllerYawLastUpdate = InOutDesiredControllerYawLastUpdate;
+
+	FPoseSearchTrajectoryData::FDerived TrajectoryDataDerived;
+	InTrajectoryData.UpdateData(InDeltaTime, InAnimInstance, TrajectoryDataDerived, TrajectoryDataState);
+	InitTrajectorySamples(InOutTrajectory, InTrajectoryData, TrajectoryDataDerived, TrajectoryDataSampling);
+	UpdateHistory_TransformHistory(InOutTrajectory, InTrajectoryData, TrajectoryDataDerived, TrajectoryDataSampling, InDeltaTime);
+	UpdatePrediction_SimulateCharacterMovement(InOutTrajectory, InTrajectoryData, TrajectoryDataDerived, TrajectoryDataSampling);
+
+	InOutDesiredControllerYawLastUpdate = TrajectoryDataState.DesiredControllerYawLastUpdate;
+
+	OutTrajectory = InOutTrajectory;
 }
\ No newline at end of file
diff --git a/Engine/Plugins/Experimental/Animation/PoseSearch/Source/Runtime/Private/PoseSearchTrajectoryTypes.cpp b/Engine/Plugins/Experimental/Animation/PoseSearch/Source/Runtime/Private/PoseSearchTrajectoryTypes.cpp
--- a/Engine/Plugins/Experimental/Animation/PoseSearch/Source/Runtime/Private/PoseSearchTrajectoryTypes.cpp
+++ b/Engine/Plugins/Experimental/Animation/PoseSearch/Source/Runtime/Private/PoseSearchTrajectoryTypes.cpp
@@ -62,44 +62,44 @@ FPoseSearchQueryTrajectorySample FPoseSearchQueryTrajectory::GetSampleAtTime(flo
 #if ENABLE_ANIM_DEBUG
 void FPoseSearchQueryTrajectory::DebugDrawTrajectory(const UWorld* World, float HeightOffset) const
 {
-	FVector OffsetVector = FVector::UpVector * HeightOffset;
+	const FVector OffsetVector = FVector::UpVector * HeightOffset;
 
 	const int32 LastIndex = Samples.Num() - 1;
 	if (LastIndex >= 0)
 	{
 		for (int32 Index = 0; ; ++Index)
 		{
-			DrawDebugSphere(World, Samples[Index].Position + OffsetVector, 2.f /*Radius*/, 4 /*Segments*/, FColor::Black);
-			DrawDebugCoordinateSystem(World, Samples[Index].Position + OffsetVector, FRotator(Samples[Index].Facing), 12.f /*Scale*/);
+			DrawDebugSphere(World, Samples[Index].Position + OffsetVector, 2.f /*Radius*/, 4 /*Segments*/, FColor::Black, false /*bPersistentLines*/, -1.f /*LifeTime*/, SDPG_Foreground);
+			DrawDebugCoordinateSystem(World, Samples[Index].Position + OffsetVector, FRotator(Samples[Index].Facing), 12.f /*Scale*/, false /*bPersistentLines*/, -1.f /*LifeTime*/, SDPG_Foreground);
 
 			if (Index == LastIndex)
 			{
 				break;
 			}
 			
-			DrawDebugLine(World, Samples[Index].Position + OffsetVector, Samples[Index + 1].Position + OffsetVector, FColor::Black);
+			DrawDebugLine(World, Samples[Index].Position + OffsetVector, Samples[Index + 1].Position + OffsetVector, FColor::Black, false /*bPersistentLines*/, -1.f /*LifeTime*/, SDPG_Foreground);
 		}
 	}
 }
 
 void FPoseSearchQueryTrajectory::DebugDrawTrajectory(FAnimInstanceProxy& AnimInstanceProxy, float HeightOffset) const
 {
-	FVector OffsetVector = FVector::UpVector * HeightOffset;
+	const FVector OffsetVector = FVector::UpVector * HeightOffset;
 
 	const int32 LastIndex = Samples.Num() - 1;
 	if (LastIndex >= 0)
 	{
 		for (int32 Index = 0; ; ++Index)
 		{
-			AnimInstanceProxy.AnimDrawDebugSphere(Samples[Index].Position + OffsetVector, 2.f /*Radius*/, 4 /*Segments*/, FColor::Black);
-			AnimInstanceProxy.AnimDrawDebugCoordinateSystem(Samples[Index].Position + OffsetVector, FRotator(Samples[Index].Facing), 12.f /*Scale*/);
+			AnimInstanceProxy.AnimDrawDebugSphere(Samples[Index].Position + OffsetVector, 2.f /*Radius*/, 4 /*Segments*/, FColor::Black, false /*bPersistentLines*/, -1.f /*LifeTime*/, 0.f /*Thickness*/, SDPG_Foreground);
+			AnimInstanceProxy.AnimDrawDebugCoordinateSystem(Samples[Index].Position + OffsetVector, FRotator(Samples[Index].Facing), 12.f /*Scale*/, false /*bPersistentLines*/, -1.f /*LifeTime*/, 0.f /*Thickness*/, SDPG_Foreground);
 
 			if (Index == LastIndex)
 			{
 				break;
 			}
 			
-			AnimInstanceProxy.AnimDrawDebugLine(Samples[Index].Position + OffsetVector, Samples[Index + 1].Position + OffsetVector, FColor::Black);
+			AnimInstanceProxy.AnimDrawDebugLine(Samples[Index].Position + OffsetVector, Samples[Index + 1].Position + OffsetVector, FColor::Black, false /*bPersistentLines*/, -1.f /*LifeTime*/, 0.f /*Thickness*/, SDPG_Foreground);
 		}
 	}
 }
diff --git a/Engine/Plugins/Experimental/Animation/PoseSearch/Source/Runtime/Private/Trace/PoseSearchTraceLogger.cpp b/Engine/Plugins/Experimental/Animation/PoseSearch/Source/Runtime/Private/Trace/PoseSearchTraceLogger.cpp
--- a/Engine/Plugins/Experimental/Animation/PoseSearch/Source/Runtime/Private/Trace/PoseSearchTraceLogger.cpp
+++ b/Engine/Plugins/Experimental/Animation/PoseSearch/Source/Runtime/Private/Trace/PoseSearchTraceLogger.cpp
@@ -19,7 +19,7 @@ namespace UE::PoseSearch
 {
 
 const FName FTraceLogger::Name("PoseSearch");
-const FName FTraceMotionMatchingState::Name("MotionMatchingState");
+const FName FTraceMotionMatchingStateMessage::Name("MotionMatchingState");
 
 bool IsTracing(const FAnimationBaseContext& InContext)
 {
@@ -68,8 +68,10 @@ FArchive& operator<<(FArchive& Ar, FTraceMotionMatchingStateDatabaseEntry& Entry
 	return Ar;
 }
 
-FArchive& operator<<(FArchive& Ar, FTraceMotionMatchingState& State)
+FArchive& operator<<(FArchive& Ar, FTraceMotionMatchingStateMessage& State)
 {
+	Ar << static_cast<FTraceMessage&>(State);
+
 	Ar << State.ElapsedPoseSearchTime;
 	Ar << State.AssetPlayerTime;
 	Ar << State.DeltaTime;
@@ -81,37 +83,26 @@ FArchive& operator<<(FArchive& Ar, FTraceMotionMatchingState& State)
 	Ar << State.SearchBestCost;
 	Ar << State.SearchBruteForceCost;
 	Ar << State.SearchBestPosePos;
+	Ar << State.SkeletalMeshComponentIds;
+	Ar << State.Roles;
 	Ar << State.DatabaseEntries;
-	Ar << State.Trajectory;
+	Ar << State.PoseHistories;	
 	Ar << State.CurrentDbEntryIdx;
 	Ar << State.CurrentPoseEntryIdx;
 	return Ar;
 }
 
-void FTraceMotionMatchingState::Output(const UObject* AnimInstance, int32 NodeId)
+void FTraceMotionMatchingStateMessage::Output()
 {
 #if OBJECT_TRACE_ENABLED
 	TArray<uint8> ArchiveData;
 	FMemoryWriter Archive(ArchiveData);
-
-	TRACE_OBJECT(AnimInstance);
-	UObject* SkeletalMeshComponent = AnimInstance->GetOuter();
-
-	FTraceMessage TraceMessage;
-	TraceMessage.Cycle = FPlatformTime::Cycles64();
-	TraceMessage.AnimInstanceId = FObjectTrace::GetObjectId(AnimInstance);
-	TraceMessage.SkeletalMeshComponentId = FObjectTrace::GetObjectId(SkeletalMeshComponent);
-	TraceMessage.NodeId = NodeId;
-	TraceMessage.FrameCounter = FObjectTrace::GetObjectWorldTickCounter(AnimInstance);
-
-	Archive << TraceMessage;
 	Archive << *this;
-
 	UE_TRACE_LOG(PoseSearch, MotionMatchingState, PoseSearchChannel) << MotionMatchingState.Data(ArchiveData.GetData(), ArchiveData.Num());
 #endif
 }
 
-const UPoseSearchDatabase* FTraceMotionMatchingState::GetCurrentDatabase() const
+const UPoseSearchDatabase* FTraceMotionMatchingStateMessage::GetCurrentDatabase() const
 {
 	const UPoseSearchDatabase* Database = nullptr;
 	if (DatabaseEntries.IsValidIndex(CurrentDbEntryIdx))
@@ -121,7 +112,7 @@ const UPoseSearchDatabase* FTraceMotionMatchingState::GetCurrentDatabase() const
 	return Database;
 }
 
-int32 FTraceMotionMatchingState::GetCurrentDatabasePoseIndex() const
+int32 FTraceMotionMatchingStateMessage::GetCurrentDatabasePoseIndex() const
 {
 	if (const FTraceMotionMatchingStatePoseEntry* PoseEntry = GetCurrentPoseEntry())
 	{
@@ -130,7 +121,7 @@ int32 FTraceMotionMatchingState::GetCurrentDatabasePoseIndex() const
 	return INDEX_NONE;
 }
 
-const FTraceMotionMatchingStatePoseEntry* FTraceMotionMatchingState::GetCurrentPoseEntry() const
+const FTraceMotionMatchingStatePoseEntry* FTraceMotionMatchingStateMessage::GetCurrentPoseEntry() const
 {
 	if (DatabaseEntries.IsValidIndex(CurrentDbEntryIdx))
 	{
diff --git a/Engine/Plugins/Experimental/Animation/PoseSearch/Source/Runtime/Public/PoseSearch/AnimNode_PoseSearchHistoryCollector.h b/Engine/Plugins/Experimental/Animation/PoseSearch/Source/Runtime/Public/PoseSearch/AnimNode_PoseSearchHistoryCollector.h
--- a/Engine/Plugins/Experimental/Animation/PoseSearch/Source/Runtime/Public/PoseSearch/AnimNode_PoseSearchHistoryCollector.h
+++ b/Engine/Plugins/Experimental/Animation/PoseSearch/Source/Runtime/Public/PoseSearch/AnimNode_PoseSearchHistoryCollector.h
@@ -75,8 +75,11 @@ public:
 	FPoseSearchTrajectoryData TrajectoryData;
 
 	// FAnimNode_Base interface
+	virtual void Initialize_AnyThread(const FAnimationInitializeContext& Context) override;
 	virtual void CacheBones_AnyThread(const FAnimationCacheBonesContext& Context) override;
 	virtual void Update_AnyThread(const FAnimationUpdateContext& Context) override;
+	virtual bool HasPreUpdate() const { return true; }
+	virtual void PreUpdate(const UAnimInstance* InAnimInstance) override;
 	// End of FAnimNode_Base interface
 
 	const UE::PoseSearch::IPoseHistory& GetPoseHistory() const { return PoseHistory; }
