motion matching - multicharacter WIP

diff --git a/Engine/Plugins/Experimental/Animation/PoseSearch/Source/Runtime/Public/PoseSearch/PoseSearchAssetIndexer.h b/Engine/Plugins/Experimental/Animation/PoseSearch/Source/Runtime/Public/PoseSearch/PoseSearchAssetIndexer.h
--- a/Engine/Plugins/Experimental/Animation/PoseSearch/Source/Runtime/Public/PoseSearch/PoseSearchAssetIndexer.h
+++ b/Engine/Plugins/Experimental/Animation/PoseSearch/Source/Runtime/Public/PoseSearch/PoseSearchAssetIndexer.h
@@ -10,6 +10,7 @@
 #include "PoseSearch/PoseSearchFeatureChannel.h"
 #include "PoseSearch/PoseSearchSchema.h"
 
+class UAnimNotifyState_PoseSearchBase;
 class UPoseSearchDatabase;
 
 namespace UE::PoseSearch
@@ -19,12 +20,33 @@ struct FAnimationAssetSampler;
 struct FSearchIndexAsset;
 struct FPoseMetadata;
 
-struct FAssetSamplingContext : public FMirrorDataCache
+struct FAssetSamplingContext
 {
 	float BaseCostBias = 0.f;
 	float LoopingCostBias = 0.f;
 
-	FAssetSamplingContext(const UPoseSearchDatabase& Database, const FBoneContainer& BoneContainer);
+	FAssetSamplingContext(const UPoseSearchDatabase& Database);
+};
+
+struct FAnimationAssetSamplers
+{
+	void Reset();
+	int32 Num() const;
+
+	float GetPlayLength() const;
+	bool IsLoopable() const;
+	void ExtractPoseSearchNotifyStates(float Time, TFunction<bool(UAnimNotifyState_PoseSearchBase*)> ProcessPoseSearchBase) const;
+	bool ProcessAllAnimNotifyEvents(TFunction<bool(TConstArrayView<FAnimNotifyEvent>)> ProcessAnimNotifyEvents) const;
+	const FString GetAssetName() const;
+
+	FTransform ExtractRootTransform(float Time, int32 RoleIndex) const;
+	FTransform GetTotalRootTransform(int32 RoleIndex) const;
+	void ExtractPose(float Time, FCompactPose& OutPose, int32 RoleIndex) const;
+	FTransform MirrorTransform(const FTransform& InTransform, int32 RoleIndex) const;
+	void MirrorPose(FCompactPose& Pose, int32 RoleIndex) const;
+
+	TArray<const FAnimationAssetSampler*> AnimationAssetSamplers;
+	TArray<const FMirrorDataCache*> MirrorDataCaches;
 };
 
 class FAssetIndexer
@@ -39,8 +61,8 @@ public:
 		float MaxAcceleration = 0.f;
 	};
 
-	FAssetIndexer(const FBoneContainer& InBoneContainer, const FSearchIndexAsset& InSearchIndexAsset, const FAssetSamplingContext& InSamplingContext,
-		const UPoseSearchSchema& InSchema, const FAnimationAssetSampler& InAssetSampler, const FFloatInterval& InExtrapolationTimeInterval);
+	FAssetIndexer(const TConstArrayView<FBoneContainer> InBoneContainers, const FSearchIndexAsset& InSearchIndexAsset, const FAssetSamplingContext& InSamplingContext,
+		const UPoseSearchSchema& InSchema, const FAnimationAssetSamplers& InAssetSamplers, const FRoleToIndex& InRoleToIndex, const FFloatInterval& InExtrapolationTimeInterval);
 	void AssignWorkingData(int32 InStartPoseIdx, TArrayView<float> InOutFeatureVectorTable, TArrayView<FPoseMetadata> InOutPoseMetadata);
 	void Process(int32 AssetIdx);
 	const FStats& GetStats() const { return Stats; }
@@ -48,18 +70,18 @@ public:
 	// Returns OutSampleRotation as the rotation of the bone Schema.BoneReferences[SchemaSampleBoneIdx] at time CalculateSampleTime(SampleIdx) + SampleTimeOffset relative to the
 	// transform of the bone Schema.BoneReferences[SchemaOriginBoneIdx] at time CalculateSampleTime(SampleIdx) + OriginTimeOffset 
 	// Times will be processed by GetPermutationTimeOffsets(PermutationTimeType, ...)
-	bool GetSampleRotation(FQuat& OutSampleRotation, float SampleTimeOffset, float OriginTimeOffset, int32 SampleIdx, int8 SchemaSampleBoneIdx = RootSchemaBoneIdx, int8 SchemaOriginBoneIdx = RootSchemaBoneIdx, EPermutationTimeType PermutationTimeType = EPermutationTimeType::UseSampleTime, int32 SamplingAttributeId = INDEX_NONE);
+	bool GetSampleRotation(FQuat& OutSampleRotation, float SampleTimeOffset, float OriginTimeOffset, int32 SampleIdx, int8 SchemaSampleBoneIdx, int8 SchemaOriginBoneIdx, const FRole& SampleRole, const FRole& OriginRole, EPermutationTimeType PermutationTimeType = EPermutationTimeType::UseSampleTime, int32 SamplingAttributeId = INDEX_NONE);
 
 	// Returns OutSamplePosition as the position of the bone Schema.BoneReferences[SchemaSampleBoneIdx] at time CalculateSampleTime(SampleIdx) + SampleTimeOffset relative to the
 	// transform of the bone Schema.BoneReferences[SchemaOriginBoneIdx] at time CalculateSampleTime(SampleIdx) + OriginTimeOffset.
 	// Times will be processed by GetPermutationTimeOffsets(PermutationTimeType, ...)
-	bool GetSamplePosition(FVector& OutSamplePosition, float SampleTimeOffset, float OriginTimeOffset, int32 SampleIdx, int8 SchemaSampleBoneIdx = RootSchemaBoneIdx, int8 SchemaOriginBoneIdx = RootSchemaBoneIdx, EPermutationTimeType PermutationTimeType = EPermutationTimeType::UseSampleTime, int32 SamplingAttributeId = INDEX_NONE);
+	bool GetSamplePosition(FVector& OutSamplePosition, float SampleTimeOffset, float OriginTimeOffset, int32 SampleIdx, int8 SchemaSampleBoneIdx, int8 SchemaOriginBoneIdx, const FRole& SampleRole, const FRole& OriginRole, EPermutationTimeType PermutationTimeType = EPermutationTimeType::UseSampleTime, int32 SamplingAttributeId = INDEX_NONE);
 
 	// Returns OutSampleVelocity as the delta velocity of the bone Schema.BoneReferences[SchemaSampleBoneIdx] velocity at time CalculateSampleTime(SampleIdx) + SampleTimeOffset minus
 	// the velocity of the bone Schema.BoneReferences[SchemaOriginBoneIdx] at time CalculateSampleTime(SampleIdx) + OriginTimeOffset.
 	// Times will be processed by GetPermutationTimeOffsets(PermutationTimeType, ...)
 	// if bUseCharacterSpaceVelocities is true, velocities will be computed in root bone space, rather than animation (world) space
-	bool GetSampleVelocity(FVector& OutSampleVelocity, float SampleTimeOffset, float OriginTimeOffset, int32 SampleIdx, int8 SchemaSampleBoneIdx = RootSchemaBoneIdx, int8 SchemaOriginBoneIdx = RootSchemaBoneIdx, bool bUseCharacterSpaceVelocities = true, EPermutationTimeType PermutationTimeType = EPermutationTimeType::UseSampleTime, int32 SamplingAttributeId = INDEX_NONE);
+	bool GetSampleVelocity(FVector& OutSampleVelocity, float SampleTimeOffset, float OriginTimeOffset, int32 SampleIdx, int8 SchemaSampleBoneIdx, int8 SchemaOriginBoneIdx, const FRole& SampleRole, const FRole& OriginRole, bool bUseCharacterSpaceVelocities = true, EPermutationTimeType PermutationTimeType = EPermutationTimeType::UseSampleTime, int32 SamplingAttributeId = INDEX_NONE);
 
 	int32 GetBeginSampleIdx() const;
 	int32 GetEndSampleIdx() const;
@@ -83,48 +105,45 @@ private:
 
 	// Returns the animation (world) space transform of the bone Schema.BoneReferences[SchemaBoneIdx] at time SampleTime
 	// bClamped will be true if SampleTime is outside the animation duration boundaries
-	FTransform GetTransform(float SampleTime, bool& bClamped, int8 SchemaBoneIdx = RootSchemaBoneIdx);
-	FTransform GetTransform(float SampleTime, bool& bClamped, const FBoneReference& BoneReference);
+	FTransform GetTransform(float SampleTime, const FRole& Role, bool& bClamped, int8 SchemaBoneIdx = RootSchemaBoneIdx);
+	FTransform GetTransform(float SampleTime, int32 RoleIndex, bool& bClamped, const FBoneReference& BoneReference);
 
 	// Returns the component space transform of the bone Schema.BoneReferences[SchemaBoneIdx] at time SampleTime
 	// bClamped will be true if SampleTime is outside the animation duration boundaries
-	FTransform GetComponentSpaceTransform(float SampleTime, bool& bClamped, int8 SchemaBoneIdx = RootSchemaBoneIdx);
+	FTransform GetComponentSpaceTransform(float SampleTime, const FRole& Role, bool& bClamped, int8 SchemaBoneIdx = RootSchemaBoneIdx);
 
 	// Returns OutSamplePosition as the position of the bone Schema.BoneReferences[SchemaSampleBoneIdx] at time SampleTime relative to the
 	// transform of the bone Schema.BoneReferences[SchemaOriginBoneIdx] at time OriginTime.
 	// bClamped will be true if SampleTime or OriginTime are outside the animation duration boundaries
-	bool GetSamplePositionInternal(FVector& OutSamplePosition, float SampleTime, float OriginTime, bool& bClamped, int8 SchemaSampleBoneIdx, int8 SchemaOriginBoneIdx, int32 SamplingAttributeId);
+	bool GetSamplePositionInternal(FVector& OutSamplePosition, float SampleTime, float OriginTime, bool& bClamped, int8 SchemaSampleBoneIdx, int8 SchemaOriginBoneIdx, const FRole& SampleRole, const FRole& OriginRole, int32 SamplingAttributeId);
 
-	struct FSampleInfo
-	{
-		FTransform RootTransform;
-		float ClipTime = 0.f;
-		bool bClamped = false;
-	};
-
-	struct CachedEntry
+	struct FCachedEntry
 	{
 		float SampleTime = 0.f;
 		bool bClamped = false;
 
-		FTransform RootTransform;
-		FCSPose<FCompactHeapPose> ComponentSpacePose;
+		TArray<FTransform> RootTransform;
+		TArray<FCSPose<FCompactHeapPose>> ComponentSpacePose;
 	};
 
-	FSampleInfo GetSampleInfo(float SampleTime) const;
-	FTransform MirrorTransform(const FTransform& Transform) const;
-	CachedEntry& GetEntry(float SampleTime);
-	FTransform CalculateComponentSpaceTransform(CachedEntry& Entry, int8 SchemaBoneIdx);
-	FTransform CalculateComponentSpaceTransform(CachedEntry& Entry, const FBoneReference& BoneReference);
+	void GetSampleInfo(float SampleTime, int32 RoleIndex, FTransform& OutRootTransform, float& OutClipTime, bool& bOutClamped) const;
+	FTransform MirrorTransform(const FTransform& Transform, int32 RoleIndex) const;
+	FCachedEntry& GetEntry(float SampleTime);
+	FTransform CalculateComponentSpaceTransform(FCachedEntry& Entry, const FBoneReference& BoneReference, int32 RoleIndex);
 
 	void ComputeStats();
 
-	FBoneContainer BoneContainer;
-	TMap<float, CachedEntry> CachedEntries;
+	TArray<FBoneContainer> BoneContainers;
+	TMap<float, FCachedEntry> CachedEntries;
 	const FSearchIndexAsset& SearchIndexAsset;
 	const FAssetSamplingContext& SamplingContext;
 	const UPoseSearchSchema& Schema;
-	const FAnimationAssetSampler& AssetSampler;
+	const FAnimationAssetSamplers AssetSamplers;
+
+	// NoTe: mapping Role to the index of the associated asset that this FAssetIndexer is indexing.
+	// NOT the index of the UPoseSearchSchema::Skeletons! Use UPoseSearchSchema::GetRoledSkeleton API to resolve that Role to FPoseSearchRoledSkeleton 
+	const FRoleToIndex RoleToIndex;
+
 	FFloatInterval ExtrapolationTimeInterval = FFloatInterval(-UE_BIG_NUMBER, UE_BIG_NUMBER);
 	
 	int32 StartPoseIdx = 0;
diff --git a/Engine/Plugins/Experimental/Animation/PoseSearch/Source/Runtime/Public/PoseSearch/PoseSearchAssetSampler.h b/Engine/Plugins/Experimental/Animation/PoseSearch/Source/Runtime/Public/PoseSearch/PoseSearchAssetSampler.h
--- a/Engine/Plugins/Experimental/Animation/PoseSearch/Source/Runtime/Public/PoseSearch/PoseSearchAssetSampler.h
+++ b/Engine/Plugins/Experimental/Animation/PoseSearch/Source/Runtime/Public/PoseSearch/PoseSearchAssetSampler.h
@@ -17,8 +17,8 @@ namespace UE::PoseSearch
  */
 struct POSESEARCH_API FAnimationAssetSampler
 {
-	FAnimationAssetSampler(TObjectPtr<const UAnimationAsset> InAnimationAsset = nullptr, const FVector& InBlendParameters = FVector::ZeroVector, int32 InRootTransformSamplingRate = 30);
-	void Init(TObjectPtr<const UAnimationAsset> InAnimationAsset, const FVector& InBlendParameters = FVector::ZeroVector, int32 InRootTransformSamplingRate = 30);
+	FAnimationAssetSampler(TObjectPtr<const UAnimationAsset> InAnimationAsset = nullptr, const FTransform& InRootTransformOrigin = FTransform::Identity, const FVector& InBlendParameters = FVector::ZeroVector, int32 InRootTransformSamplingRate = 30);
+	void Init(TObjectPtr<const UAnimationAsset> InAnimationAsset, const FTransform& InRootTransformOrigin = FTransform::Identity, const FVector& InBlendParameters = FVector::ZeroVector, int32 InRootTransformSamplingRate = 30);
 
 	bool IsInitialized() const;
 	float GetPlayLength() const;
@@ -47,12 +47,13 @@ struct POSESEARCH_API FAnimationAssetSampler
 
 	const UAnimationAsset* GetAsset() const;
 
-	void Process(const FBoneContainer& BoneContainer);
+	void Process();
 
 	static float GetPlayLength(const UAnimationAsset* AnimAsset, const FVector& BlendParameters = FVector::ZeroVector);
 
 protected:
 	TWeakObjectPtr<const UAnimationAsset> AnimationAssetPtr;
+	FTransform RootTransformOrigin = FTransform::Identity;
 
 	// members used to sample blend spaces only!
 	FVector BlendParameters = FVector::ZeroVector;
diff --git a/Engine/Plugins/Experimental/Animation/PoseSearch/Source/Runtime/Public/PoseSearch/PoseSearchContext.h b/Engine/Plugins/Experimental/Animation/PoseSearch/Source/Runtime/Public/PoseSearch/PoseSearchContext.h
--- a/Engine/Plugins/Experimental/Animation/PoseSearch/Source/Runtime/Public/PoseSearch/PoseSearchContext.h
+++ b/Engine/Plugins/Experimental/Animation/PoseSearch/Source/Runtime/Public/PoseSearch/PoseSearchContext.h
@@ -3,6 +3,7 @@
 #pragma once
 
 #include "DrawDebugHelpers.h"
+#include "Animation/AnimInstance.h"
 #include "PoseSearch/PoseSearchDefines.h"
 #include "PoseSearch/PoseSearchFeatureChannel.h"
 #include "PoseSearch/PoseSearchIndex.h"
@@ -22,8 +23,8 @@ enum class EDebugDrawFlags : uint32
 {
 	None = 0,
 
-	// Draw using Query colors form the Schema
-	DrawQuery = 1 << 1,
+	// used to differenciate channels debug drawing of the query
+	DrawQuery = 1 << 0,
 };
 ENUM_CLASS_FLAGS(EDebugDrawFlags);
 
@@ -48,77 +49,21 @@ enum class EPoseCandidateFlags : uint32
 };
 ENUM_CLASS_FLAGS(EPoseCandidateFlags);
 
-template<typename FTransformType>
-struct POSESEARCH_API FCachedTransform
-{
-	FCachedTransform()
-		: SampleTime(0.f)
-		, BoneIndexType(RootBoneIndexType)
-		, Transform(FTransformType::Identity)
-	{
-	}
-
-	FCachedTransform(float InSampleTime, FBoneIndexType InBoneIndexType, const FTransformType& InTransform)
-		: SampleTime(InSampleTime)
-		, BoneIndexType(InBoneIndexType)
-		, Transform(InTransform)
-	{
-	}
-
-	float SampleTime = 0.f;
-	
-	FBoneIndexType BoneIndexType = RootBoneIndexType;
-
-	// associated transform to BoneIndexType in ComponentSpace (except for the root bone stored in global space)
-	FTransformType Transform = FTransformType::Identity;
-};
-
-template<typename FTransformType>
-struct FCachedTransforms
-{
-	const FCachedTransform<FTransformType>* Find(float SampleTime, FBoneIndexType BoneIndexType) const
-	{
-		// @todo: use an hashmap if we end up having too many entries
-		return CachedTransforms.FindByPredicate([SampleTime, BoneIndexType](const FCachedTransform<FTransformType>& CachedTransform)
-			{
-				return CachedTransform.SampleTime == SampleTime && CachedTransform.BoneIndexType == BoneIndexType;
-			});
-	}
-
-	void Add(float SampleTime, FBoneIndexType BoneIndexType, const FTransformType& Transform)
-	{
-		CachedTransforms.Emplace(SampleTime, BoneIndexType, Transform);
-	}
-
-	void Reset()
-	{
-		CachedTransforms.Reset();
-	}
-
-	bool IsEmpty() const
-	{
-		return CachedTransforms.IsEmpty();
-	}
-
-private:
-	TArray<FCachedTransform<FTransformType>, TInlineAllocator<64, TMemStackAllocator<>>> CachedTransforms;
-};
-
 #if ENABLE_DRAW_DEBUG
 struct POSESEARCH_API FDebugDrawParams
 {
-	FDebugDrawParams(FAnimInstanceProxy* InAnimInstanceProxy, const FTransform& InRootMotionTransform, const UPoseSearchDatabase* InDatabase, EDebugDrawFlags InFlags = EDebugDrawFlags::None);
-	FDebugDrawParams(const UWorld* InWorld, const USkinnedMeshComponent* InMesh, const FTransform& InRootMotionTransform, const UPoseSearchDatabase* InDatabase, EDebugDrawFlags InFlags = EDebugDrawFlags::None);
+	FDebugDrawParams(TArrayView<FAnimInstanceProxy*> InAnimInstanceProxies, TConstArrayView<const IPoseHistory*> InPoseHistories, const FRoleToIndex& InRoleToIndex, const UPoseSearchDatabase* InDatabase, EDebugDrawFlags InFlags = EDebugDrawFlags::None);
+	FDebugDrawParams(TArrayView<const USkinnedMeshComponent*> InMeshes, TConstArrayView<const IPoseHistory*> InPoseHistories, const FRoleToIndex& InRoleToIndex, const UPoseSearchDatabase* InDatabase, EDebugDrawFlags InFlags = EDebugDrawFlags::None);
 
 	const FSearchIndex* GetSearchIndex() const;
 	const UPoseSearchSchema* GetSchema() const;
 
 	FVector ExtractPosition(TConstArrayView<float> PoseVector, const UPoseSearchFeatureChannel_Position* Position) const;
-	FVector ExtractPosition(TConstArrayView<float> PoseVector, float SampleTimeOffset, int8 SchemaBoneIdx = RootSchemaBoneIdx, EPermutationTimeType PermutationTimeType = EPermutationTimeType::UseSampleTime, int32 SamplingAttributeId = INDEX_NONE) const;
+	FVector ExtractPosition(TConstArrayView<float> PoseVector, float SampleTimeOffset, int8 SchemaBoneIdx, const FRole& Role, EPermutationTimeType PermutationTimeType = EPermutationTimeType::UseSampleTime, int32 SamplingAttributeId = INDEX_NONE) const;
 
-	FQuat ExtractRotation(TConstArrayView<float> PoseVector, float SampleTimeOffset, int8 SchemaBoneIdx = RootSchemaBoneIdx) const;
+	FQuat ExtractRotation(TConstArrayView<float> PoseVector, float SampleTimeOffset, int8 SchemaBoneIdx, const FRole& Role, EPermutationTimeType PermutationTimeType = EPermutationTimeType::UseSampleTime, int32 SamplingAttributeId = INDEX_NONE) const;
 
-	const FTransform& GetRootTransform() const;
+	FTransform GetRootTransform(const FRole& Role) const;
 
 	void DrawLine(const FVector& LineStart, const FVector& LineEnd, const FColor& Color, float Thickness = 0.f) const;
 	void DrawPoint(const FVector& Position, const FColor& Color, float Thickness = 6.f) const;
@@ -131,10 +76,14 @@ struct POSESEARCH_API FDebugDrawParams
 private:
 	bool CanDraw() const;
 
-	FAnimInstanceProxy* AnimInstanceProxy = nullptr;
-	const UWorld* World = nullptr;
-	const USkinnedMeshComponent* Mesh = nullptr;
-	const FTransform* RootMotionTransform = nullptr;
+	const TArrayView<FAnimInstanceProxy*> AnimInstanceProxies;
+	const TArrayView<const USkinnedMeshComponent*> Meshes;
+	const TConstArrayView<const IPoseHistory*> PoseHistories;
+
+	// NoTe: mapping Role to the index of the associated asset that this FDebugDrawParams is drawing.
+	// NOT the index of the UPoseSearchSchema::Skeletons! Use UPoseSearchSchema::GetRoledSkeleton API to resolve that Role to FPoseSearchRoledSkeleton 
+	const FRoleToIndex& RoleToIndex;
+
 	const UPoseSearchDatabase* Database = nullptr;
 	EDebugDrawFlags Flags = EDebugDrawFlags::None;
 };
@@ -169,25 +118,32 @@ struct FCachedChannel
 
 struct POSESEARCH_API FSearchContext
 {
-	FSearchContext(const UAnimInstance* InAnimInstance, const IPoseHistory* InHistory, TConstArrayView<const UObject*> InAssetsToConsider = TConstArrayView<const UObject*>(),
-		float InDesiredPermutationTimeOffset = 0.f, const FPoseIndicesHistory* InPoseIndicesHistory = nullptr,
+	FSearchContext(float InDesiredPermutationTimeOffset = 0.f, const FPoseIndicesHistory* InPoseIndicesHistory = nullptr,
 		const FSearchResult& InCurrentResult = FSearchResult(), const FFloatInterval& InPoseJumpThresholdTime = FFloatInterval(0.f, 0.f), bool bInUseCachedChannelData = false);
 
+	// deleting copies and moves since members could reference other members (e.g.: CurrentResultPoseVector could point to CurrentResultPoseVectorData, so it'll require proper copies and movers implementations)
+	FSearchContext(const FSearchContext& Other) = delete;
+	FSearchContext(FSearchContext&& Other) = delete;
+	FSearchContext& operator=(const FSearchContext& Other) = delete;
+	FSearchContext& operator=(FSearchContext&& Other) = delete;
+
+	void AddRole(const FRole& Role, const UAnimInstance* AnimInstance, const IPoseHistory* PoseHistory);
+
 	// Returns the rotation of the bone Schema.BoneReferences[SchemaSampleBoneIdx] at an offset time of SampleTimeOffset relative to the
 	// transform of the bone Schema.BoneReferences[SchemaOriginBoneIdx] at an offset time of time OriginTimeOffset 
 	// Times will be processed by GetPermutationTimeOffsets(PermutationTimeType, ...)
-	FQuat GetSampleRotation(float SampleTimeOffset, float OriginTimeOffset, int8 SchemaSampleBoneIdx = RootSchemaBoneIdx, int8 SchemaOriginBoneIdx = RootSchemaBoneIdx, EPermutationTimeType PermutationTimeType = EPermutationTimeType::UseSampleTime, const FQuat* SampleBoneRotationWorldOverride = nullptr);
+	FQuat GetSampleRotation(float SampleTimeOffset, float OriginTimeOffset, int8 SchemaSampleBoneIdx, int8 SchemaOriginBoneIdx, const FRole& SampleRole, const FRole& OriginRole, EPermutationTimeType PermutationTimeType = EPermutationTimeType::UseSampleTime, const FQuat* SampleBoneRotationWorldOverride = nullptr);
 	
 	// Returns the position of the bone Schema.BoneReferences[SchemaSampleBoneIdx] at an offset time of SampleTimeOffset relative to the
 	// transform of the bone Schema.BoneReferences[SchemaOriginBoneIdx] at an offset time of time OriginTimeOffset 
 	// Times will be processed by GetPermutationTimeOffsets(PermutationTimeType, ...)
-	FVector GetSamplePosition(float SampleTimeOffset, float OriginTimeOffset, int8 SchemaSampleBoneIdx = RootSchemaBoneIdx, int8 SchemaOriginBoneIdx = RootSchemaBoneIdx, EPermutationTimeType PermutationTimeType = EPermutationTimeType::UseSampleTime, const FVector* SampleBonePositionWorldOverride = nullptr);
+	FVector GetSamplePosition(float SampleTimeOffset, float OriginTimeOffset, int8 SchemaSampleBoneIdx, int8 SchemaOriginBoneIdx, const FRole& SampleRole, const FRole& OriginRole, EPermutationTimeType PermutationTimeType = EPermutationTimeType::UseSampleTime, const FVector* SampleBonePositionWorldOverride = nullptr);
 	
 	// Returns the delta velocity of the velocity of the bone Schema.BoneReferences[SchemaSampleBoneIdx] at an offset time of SampleTimeOffset minus
 	// the velocity of the bone Schema.BoneReferences[SchemaOriginBoneIdx] at an offset time of time OriginTimeOffset 
 	// Times will be processed by GetPermutationTimeOffsets(PermutationTimeType, ...)
 	// if bUseCharacterSpaceVelocities is true, velocities will be computed in root bone space, rather than world space
-	FVector GetSampleVelocity(float SampleTimeOffset, float OriginTimeOffset, int8 SchemaSampleBoneIdx = RootSchemaBoneIdx, int8 SchemaOriginBoneIdx = RootSchemaBoneIdx, bool bUseCharacterSpaceVelocities = true, EPermutationTimeType PermutationTimeType = EPermutationTimeType::UseSampleTime, const FVector* SampleBoneVelocityWorldOverride = nullptr);
+	FVector GetSampleVelocity(float SampleTimeOffset, float OriginTimeOffset, int8 SchemaSampleBoneIdx, int8 SchemaOriginBoneIdx, const FRole& SampleRole, const FRole& OriginRole, bool bUseCharacterSpaceVelocities = true, EPermutationTimeType PermutationTimeType = EPermutationTimeType::UseSampleTime, const FVector* SampleBoneVelocityWorldOverride = nullptr);
 
 	void ResetCurrentBestCost();
 	void UpdateCurrentBestCost(const FPoseSearchCost& PoseSearchCost);
@@ -205,16 +161,23 @@ struct POSESEARCH_API FSearchContext
 	const FSearchResult& GetCurrentResult() const { return CurrentResult; }
 	const FFloatInterval& GetPoseJumpThresholdTime() const { return PoseJumpThresholdTime; }
 	const FPoseIndicesHistory* GetPoseIndicesHistory() const { return PoseIndicesHistory; }
-	const IPoseHistory* GetHistory() const { return History; }
+	
+	bool ArePoseHistoriesValid() const;
+	const TArray<const IPoseHistory*, TInlineAllocator<PreallocatedRolesNum, TMemStackAllocator<>>>& GetPoseHistories() const { return PoseHistories; }
+	const IPoseHistory* GetPoseHistory(const FRole& Role) const { return PoseHistories[RoleToIndex[Role]]; }
+
 	float GetDesiredPermutationTimeOffset() const { return DesiredPermutationTimeOffset; }
-	const UAnimInstance* GetAnimInstance() const { return AnimInstance; }
+	const UAnimInstance* GetAnimInstance(const FRole& Role) const { return AnimInstances[RoleToIndex[Role]]; }
+	const TArray<const UAnimInstance*, TInlineAllocator<PreallocatedRolesNum, TMemStackAllocator<>>>& GetAnimInstances() const { return AnimInstances; }
+
+	const FRoleToIndex& GetRoleToIndex() const { return RoleToIndex; }
 
 	void SetAssetsToConsider(TConstArrayView<const UObject*> InAssetsToConsider) { AssetsToConsider = InAssetsToConsider; }
 	TConstArrayView<const UObject*> GetAssetsToConsider() const { return AssetsToConsider; }
 	
 	// returns the world space transform of the bone SchemaBoneIdx at time SampleTime
-	FTransform GetWorldBoneTransformAtTime(float SampleTime, int8 SchemaBoneIdx = RootSchemaBoneIdx);
-	
+	FTransform GetWorldBoneTransformAtTime(float SampleTime, const FRole& SampleRole, int8 SchemaBoneIdx);
+
 #if WITH_EDITOR
 	void SetAsyncBuildIndexInProgress() { bAsyncBuildIndexInProgress = true; }
 	void ResetAsyncBuildIndexInProgress() { bAsyncBuildIndexInProgress = false; }
@@ -230,13 +193,14 @@ struct POSESEARCH_API FSearchContext
 	void SetUseCachedChannelData(bool bInUseCachedChannelData) { bUseCachedChannelData = bInUseCachedChannelData; }
 
 private:
-	FVector GetSamplePositionInternal(float SampleTime, float OriginTime, int8 SchemaSampleBoneIdx = RootSchemaBoneIdx, int8 SchemaOriginBoneIdx = RootSchemaBoneIdx, const FVector* SampleBonePositionWorldOverride = nullptr);
-	FQuat GetSampleRotationInternal(float SampleTime, float OriginTime, int8 SchemaSampleBoneIdx = RootSchemaBoneIdx, int8 SchemaOriginBoneIdx = RootSchemaBoneIdx, const FQuat* SampleBoneRotationWorldOverride = nullptr);
-	FTransform GetWorldRootBoneTransformAtTime(float SampleTime) const;
+	FVector GetSamplePositionInternal(float SampleTime, float OriginTime, int8 SchemaSampleBoneIdx, int8 SchemaOriginBoneIdx, const FRole& SampleRole, const FRole& OriginRole, const FVector* SampleBonePositionWorldOverride = nullptr);
+	FQuat GetSampleRotationInternal(float SampleTime, float OriginTime, int8 SchemaSampleBoneIdx, int8 SchemaOriginBoneIdx, const FRole& SampleRole, const FRole& OriginRole, const FQuat* SampleBoneRotationWorldOverride = nullptr);
+	FTransform GetWorldRootBoneTransformAtTime(float SampleTime, const FRole& SampleRole) const;
+	
+	TArray<const UAnimInstance*, TInlineAllocator<PreallocatedRolesNum, TMemStackAllocator<>>> AnimInstances;
+	TArray<const IPoseHistory*, TInlineAllocator<PreallocatedRolesNum, TMemStackAllocator<>>> PoseHistories;
+	FRoleToIndex RoleToIndex;
 	
-	const UAnimInstance* AnimInstance = nullptr;
-	const IPoseHistory* History = nullptr;
-
 	// if AssetsToConsider is not empty, we'll search only for poses from UObject(s) that are in the AssetsToConsider
 	TConstArrayView<const UObject*> AssetsToConsider;
 
@@ -252,13 +216,12 @@ private:
 	TStackAlignedArray<float> CurrentResultPoseVectorData;
 
 	// transforms cached in world space
-	FCachedTransforms<FTransform> CachedTransforms;
+	TMap<uint32, FTransform, TInlineSetAllocator<64, TMemStackSetAllocator<>>> CachedTransforms;
 
 	TArray<FCachedQuery, TInlineAllocator<PreallocatedCachedQueriesNum, TMemStackAllocator<>>> CachedQueries;
 
-	// @todo add an overflow TMemStackAllocator
 	// mapping channel unique identifier (hash) to FCachedChannel
-	TMap<uint32, FCachedChannel, TInlineSetAllocator<PreallocatedCachedChannelDataNum>> CachedChannels;
+	TMap<uint32, FCachedChannel, TInlineSetAllocator<PreallocatedCachedChannelDataNum, TMemStackSetAllocator<>>> CachedChannels;
 
 	float CurrentBestTotalCost = MAX_flt;
 	
@@ -293,7 +256,6 @@ public:
 		}
 	}
 
-
 	struct FBestPoseCandidates
 	{
 		FBestPoseCandidates()
diff --git a/Engine/Plugins/Experimental/Animation/PoseSearch/Source/Runtime/Public/PoseSearch/PoseSearchDatabase.h b/Engine/Plugins/Experimental/Animation/PoseSearch/Source/Runtime/Public/PoseSearch/PoseSearchDatabase.h
--- a/Engine/Plugins/Experimental/Animation/PoseSearch/Source/Runtime/Public/PoseSearch/PoseSearchDatabase.h
+++ b/Engine/Plugins/Experimental/Animation/PoseSearch/Source/Runtime/Public/PoseSearch/PoseSearchDatabase.h
@@ -7,6 +7,7 @@
 #include "PoseSearch/PoseSearchCost.h"
 #include "PoseSearch/PoseSearchIndex.h"
 #include "PoseSearch/PoseSearchResult.h"
+#include "PoseSearch/PoseSearchRole.h"
 #include "PoseSearchDatabase.generated.h"
 
 struct FInstancedStruct;
@@ -14,6 +15,7 @@ class UAnimationAsset;
 class UAnimComposite;
 class UAnimMontage;
 class UBlendSpace;
+class UPoseSearchMultiSequence;
 
 #if WITH_EDITORONLY_DATA
 class UPoseSearchNormalizationSet;
@@ -52,6 +54,14 @@ struct POSESEARCH_API FPoseSearchDatabaseAnimationAssetBase
 	virtual ~FPoseSearchDatabaseAnimationAssetBase() = default;
 	virtual UObject* GetAnimationAsset() const { return nullptr; }
 	virtual float GetPlayLength() const;
+	virtual int32 GetNumRoles() const { return 1; }
+	virtual UE::PoseSearch::FRole GetRole(int32 RoleIndex) const { return UE::PoseSearch::DefaultRole; }
+	virtual UAnimationAsset* GetAnimationAssetForRole(const UE::PoseSearch::FRole& Role) const;
+	virtual const FTransform& GetRootTransformOriginForRole(const UE::PoseSearch::FRole& Role) const;
+
+#if WITH_EDITOR
+	virtual int32 GetFrameAtTime(float Time) const;
+#endif // WITH_EDITOR
 
 #if WITH_EDITORONLY_DATA
 	virtual bool IsDisableReselection() const { return bDisableReselection; }
@@ -108,14 +118,14 @@ struct POSESEARCH_API FPoseSearchDatabaseSequence : public FPoseSearchDatabaseAn
 	UPROPERTY(EditAnywhere, Category="Settings", meta = (DisplayPriority = 2))
 	FFloatInterval SamplingRange = FFloatInterval(0.f, 0.f);
 
-	UClass* GetAnimationAssetStaticClass() const override;
-	bool IsLooping() const override;
-	const FString GetName() const override;
-	bool IsRootMotionEnabled() const override;
-	FFloatInterval GetSamplingRange() const override { return SamplingRange; }
+	virtual UClass* GetAnimationAssetStaticClass() const override;
+	virtual bool IsLooping() const override;
+	virtual const FString GetName() const override;
+	virtual bool IsRootMotionEnabled() const override;
+	virtual FFloatInterval GetSamplingRange() const override { return SamplingRange; }
 #endif // WITH_EDITORONLY_DATA
 	
-	UObject* GetAnimationAsset() const override;
+	virtual UObject* GetAnimationAsset() const override;
 };
 
 /** An blend space entry in a UPoseSearchDatabase. */
@@ -154,16 +164,20 @@ struct POSESEARCH_API FPoseSearchDatabaseBlendSpace : public FPoseSearchDatabase
 	UPROPERTY(EditAnywhere, Category = "Settings", meta = (EditCondition = "bUseSingleSample", EditConditionHides, DisplayPriority = 9))
 	float BlendParamY = 0.f;
 
-	UClass* GetAnimationAssetStaticClass() const override;
-	bool IsLooping() const override;
-	const FString GetName() const override;
-	bool IsRootMotionEnabled() const override;
+	virtual UClass* GetAnimationAssetStaticClass() const override;
+	virtual bool IsLooping() const override;
+	virtual const FString GetName() const override;
+	virtual bool IsRootMotionEnabled() const override;
 
 	void GetBlendSpaceParameterSampleRanges(int32& HorizontalBlendNum, int32& VerticalBlendNum) const;
 	FVector BlendParameterForSampleRanges(int32 HorizontalBlendIndex, int32 VerticalBlendIndex) const;
 #endif // WITH_EDITORONLY_DATA
 
-	UObject* GetAnimationAsset() const override;
+	virtual UObject* GetAnimationAsset() const override;
+
+#if WITH_EDITOR
+	virtual int32 GetFrameAtTime(float Time) const override;
+#endif // WITH_EDITOR
 };
 
 /** An entry in a UPoseSearchDatabase. */
@@ -183,14 +197,14 @@ struct POSESEARCH_API FPoseSearchDatabaseAnimComposite : public FPoseSearchDatab
 	UPROPERTY(EditAnywhere, Category = "Settings", meta = (DisplayPriority = 3))
 	FFloatInterval SamplingRange = FFloatInterval(0.f, 0.f);
 
-	UClass* GetAnimationAssetStaticClass() const override;
-	bool IsLooping() const override;
-	const FString GetName() const override;
-	bool IsRootMotionEnabled() const override;
-	FFloatInterval GetSamplingRange() const override { return SamplingRange; }
+	virtual UClass* GetAnimationAssetStaticClass() const override;
+	virtual bool IsLooping() const override;
+	virtual const FString GetName() const override;
+	virtual bool IsRootMotionEnabled() const override;
+	virtual FFloatInterval GetSamplingRange() const override { return SamplingRange; }
 #endif // WITH_EDITORONLY_DATA
 
-	UObject* GetAnimationAsset() const override;
+	virtual UObject* GetAnimationAsset() const override;
 };
 
 /** An anim montage entry in a UPoseSearchDatabase. */
@@ -210,14 +224,50 @@ struct POSESEARCH_API FPoseSearchDatabaseAnimMontage : public FPoseSearchDatabas
 	UPROPERTY(EditAnywhere, Category="Settings", meta = (DisplayPriority = 2))
 	FFloatInterval SamplingRange = FFloatInterval(0.f, 0.f);
 
-	UClass* GetAnimationAssetStaticClass() const override;
-	bool IsLooping() const override;
-	const FString GetName() const override;
-	bool IsRootMotionEnabled() const override;
-	FFloatInterval GetSamplingRange() const override { return SamplingRange; }
+	virtual UClass* GetAnimationAssetStaticClass() const override;
+	virtual bool IsLooping() const override;
+	virtual const FString GetName() const override;
+	virtual bool IsRootMotionEnabled() const override;
+	virtual FFloatInterval GetSamplingRange() const override { return SamplingRange; }
+#endif // WITH_EDITORONLY_DATA
+
+	virtual UObject* GetAnimationAsset() const override;
+};
+
+USTRUCT(BlueprintType, Category = "Animation|Pose Search")
+struct POSESEARCH_API FPoseSearchDatabaseMultiSequence : public FPoseSearchDatabaseAnimationAssetBase
+{
+	GENERATED_BODY()
+	virtual ~FPoseSearchDatabaseMultiSequence() = default;
+
+	UPROPERTY(EditAnywhere, Category="Settings", meta = (DisplayPriority = 0))
+	TObjectPtr<UPoseSearchMultiSequence> MultiSequence;
+
+#if WITH_EDITORONLY_DATA
+	// It allows users to set a time range to an individual animation sequence in the database. 
+	// This is effectively trimming the beginning and end of the animation in the database (not in the original sequence).
+	// If set to [0, 0] it will be the entire frame range of the original sequence.
+	UPROPERTY(EditAnywhere, Category="Settings", meta = (DisplayPriority = 2))
+	FFloatInterval SamplingRange = FFloatInterval(0.f, 0.f);
+
+	virtual UClass* GetAnimationAssetStaticClass() const override;
+	virtual bool IsLooping() const override;
+	virtual const FString GetName() const override;
+	virtual bool IsRootMotionEnabled() const override;
+	virtual FFloatInterval GetSamplingRange() const override { return SamplingRange; }
 #endif // WITH_EDITORONLY_DATA
 
-	UObject* GetAnimationAsset() const override;
+	virtual UObject* GetAnimationAsset() const override;
+	virtual float GetPlayLength() const override;
+
+	virtual int32 GetNumRoles() const override;
+	virtual UE::PoseSearch::FRole GetRole(int32 RoleIndex) const override;
+	virtual UAnimationAsset* GetAnimationAssetForRole(const UE::PoseSearch::FRole& Role) const override;
+	virtual const FTransform& GetRootTransformOriginForRole(const UE::PoseSearch::FRole& Role) const override;
+
+#if WITH_EDITOR
+	virtual int32 GetFrameAtTime(float Time) const override;
+#endif // WITH_EDITOR
 };
 
 /** A data asset for indexing a collection of animation sequences. */
@@ -341,6 +391,10 @@ public:
 
 	int32 GetPoseIndexFromTime(float AssetTime, const UE::PoseSearch::FSearchIndexAsset& SearchIndexAsset) const;
 
+	void AddAnimationAsset(FInstancedStruct AnimationAsset);
+	void RemoveAnimationAssetAt(int32 AnimationAssetIndex);
+
+	const TArray<FInstancedStruct>& GetAnimationAssets() const { return AnimationAssets; }
 	const FInstancedStruct& GetAnimationAssetStruct(int32 AnimationAssetIndex) const;
 	const FInstancedStruct& GetAnimationAssetStruct(const UE::PoseSearch::FSearchIndexAsset& SearchIndexAsset) const;
 	FInstancedStruct& GetMutableAnimationAssetStruct(int32 AnimationAssetIndex);
@@ -380,8 +434,14 @@ public:
 
 	void SynchronizeWithExternalDependencies();
 	void SynchronizeWithExternalDependencies(TConstArrayView<UAnimSequenceBase*> SequencesBase);
+
+	bool Contains(const UObject* Object) const;
 #endif // WITH_EDITOR
 
+#if WITH_EDITOR && ENABLE_ANIM_DEBUG
+	void TestSynchronizeWithExternalDependencies();
+#endif // WITH_EDITOR && ENABLE_ANIM_DEBUG
+
 private:
 	UE::PoseSearch::FSearchResult SearchPCAKDTree(UE::PoseSearch::FSearchContext& SearchContext) const;
 	UE::PoseSearch::FSearchResult SearchVPTree(UE::PoseSearch::FSearchContext& SearchContext) const;
diff --git a/Engine/Plugins/Experimental/Animation/PoseSearch/Source/Runtime/Public/PoseSearch/PoseSearchDefines.h b/Engine/Plugins/Experimental/Animation/PoseSearch/Source/Runtime/Public/PoseSearch/PoseSearchDefines.h
--- a/Engine/Plugins/Experimental/Animation/PoseSearch/Source/Runtime/Public/PoseSearch/PoseSearchDefines.h
+++ b/Engine/Plugins/Experimental/Animation/PoseSearch/Source/Runtime/Public/PoseSearch/PoseSearchDefines.h
@@ -28,15 +28,80 @@ static constexpr FBoneIndexType RootBoneIndexType = 0;
 static constexpr FBoneIndexType ComponentSpaceIndexType = FBoneIndexType(-1);
 static constexpr FBoneIndexType WorldSpaceIndexType = FBoneIndexType(-2);
 
+static constexpr int32 PreallocatedRolesNum = 4;
 static constexpr int32 PreallocatedCachedQueriesNum = 8;
 static constexpr int32 PreallocatedCachedChannelDataNum = 64;
 static constexpr float FiniteDelta = 1 / 60.0f; // Time delta used for computing pose derivatives
 static constexpr int32 MaxNumberOfCollectedPoseCandidatesPerDatabase = 200;
 
+template<uint32 Alignment = DEFAULT_ALIGNMENT> class TMemStackSetAllocator : public TSetAllocator<TSparseArrayAllocator<TMemStackAllocator<Alignment>, TMemStackAllocator<Alignment>>, TMemStackAllocator<Alignment>> {};
+
 constexpr int32 TMax(int32 A, int32 B) { return (A > B ? A : B); }
 template<typename ElementType> constexpr int32 TAlignOf() { return TMax(alignof(ElementType), 16); }
 template<typename ElementType> using TAlignedArray = TArray<ElementType, TAlignedHeapAllocator<TAlignOf<ElementType>()>>;
 template<typename ElementType> using TStackAlignedArray = TArray<ElementType, TMemStackAllocator<TAlignOf<ElementType>()>>;
+
+#if ENABLE_ANIM_DEBUG
+
+	template<bool bIsRead = 0>
+	struct FDebugThreadSafetyChecker
+	{
+		// if any read is in progress Counter > 0
+		// if any write is in progress Counter < 0
+		explicit FDebugThreadSafetyChecker(FThreadSafeCounter& InCounter)
+			: Counter(InCounter)
+		{
+			if (bIsRead)
+			{
+				const int32 PreviousCounter = Counter.Add(1);
+				if (PreviousCounter < 0)
+				{
+					UE_LOG(LogPoseSearch, Error, TEXT("Non thread safe call! Read detected a Write in progress"));
+				}
+			}
+			else
+			{
+				const int32 PreviousCounter = Counter.Set(-1);
+				if (PreviousCounter > 0)
+				{
+					UE_LOG(LogPoseSearch, Error, TEXT("Non thread safe call! Write detected %d Read(s) in progress"), PreviousCounter);
+				}
+			}
+		}
+
+		~FDebugThreadSafetyChecker()
+		{
+			if (bIsRead)
+			{
+				const int32 PreviousCounter = Counter.Add(-1);
+				if (PreviousCounter <= 0)
+				{
+					UE_LOG(LogPoseSearch, Error, TEXT("Non thread safe call! Read detected a Write in progress"));
+				}
+			}
+			else
+			{
+				const int32 PreviousCounter = Counter.Set(0);
+				if (PreviousCounter >= 0)
+				{
+					UE_LOG(LogPoseSearch, Error, TEXT("Non thread safe call! Write detected %d Read(s) in progress"), PreviousCounter);
+				}
+			}
+		}
+
+		FThreadSafeCounter& Counter;
+	};
+
+#define CheckThreadSafetyRead(ThreadSafeCounter) FDebugThreadSafetyChecker<1> DebugThreadSafetyChecker##ThreadSafeCounter(ThreadSafeCounter)
+#define CheckThreadSafetyWrite(ThreadSafeCounter) FDebugThreadSafetyChecker DebugThreadSafetyChecker##ThreadSafeCounter(ThreadSafeCounter)
+
+#else // ENABLE_ANIM_DEBUG
+
+	#define CheckThreadSafetyRead(ThreadSafeCounter)
+	#define CheckThreadSafetyWrite(ThreadSafeCounter)
+
+#endif // ENABLE_ANIM_DEBUG
+
 }
 
 
diff --git a/Engine/Plugins/Experimental/Animation/PoseSearch/Source/Runtime/Public/PoseSearch/PoseSearchFeatureChannel.h b/Engine/Plugins/Experimental/Animation/PoseSearch/Source/Runtime/Public/PoseSearch/PoseSearchFeatureChannel.h
--- a/Engine/Plugins/Experimental/Animation/PoseSearch/Source/Runtime/Public/PoseSearch/PoseSearchFeatureChannel.h
+++ b/Engine/Plugins/Experimental/Animation/PoseSearch/Source/Runtime/Public/PoseSearch/PoseSearchFeatureChannel.h
@@ -5,6 +5,7 @@
 #include "IO/IoHash.h"
 #include "Interfaces/Interface_BoneReferenceSkeletonProvider.h"
 #include "DrawDebugHelpers.h"
+#include "PoseSearch/PoseSearchRole.h"
 #include "PoseSearchFeatureChannel.generated.h"
 
 class UPoseSearchSchema;
@@ -120,7 +121,7 @@ public:
 	int32 GetChannelDataOffset() const { checkSlow(ChannelDataOffset >= 0); return ChannelDataOffset; }
 
 	// Called during UPoseSearchSchema::Finalize to prepare the schema for this channel
-	virtual void Finalize(UPoseSearchSchema* Schema) PURE_VIRTUAL(UPoseSearchFeatureChannel::Finalize, );
+	virtual bool Finalize(UPoseSearchSchema* Schema) PURE_VIRTUAL(UPoseSearchFeatureChannel::Finalize, return false;);
 	
 	// Called at runtime to add this channel's data to the query pose vector
 	virtual void BuildQuery(UE::PoseSearch::FSearchContext& SearchContext) const PURE_VIRTUAL(UPoseSearchFeatureChannel::BuildQuery, );
@@ -154,9 +155,10 @@ public:
 
 	virtual bool CanBeNormalizedWith(const UPoseSearchFeatureChannel* Other) const;
 	const UPoseSearchSchema* GetSchema() const;
+
+	virtual const UE::PoseSearch::FRole GetDefaultRole() const;
 #endif
 
-private:
 	// IBoneReferenceSkeletonProvider interface
 	// Note this function is exclusively for FBoneReference details customization
 	USkeleton* GetSkeleton(bool& bInvalidSkeletonIsError, const IPropertyHandle* PropertyHandle) override;
diff --git a/Engine/Plugins/Experimental/Animation/PoseSearch/Source/Runtime/Public/PoseSearch/PoseSearchHistory.h b/Engine/Plugins/Experimental/Animation/PoseSearch/Source/Runtime/Public/PoseSearch/PoseSearchHistory.h
--- a/Engine/Plugins/Experimental/Animation/PoseSearch/Source/Runtime/Public/PoseSearch/PoseSearchHistory.h
+++ b/Engine/Plugins/Experimental/Animation/PoseSearch/Source/Runtime/Public/PoseSearch/PoseSearchHistory.h
@@ -43,9 +43,7 @@ struct FPoseHistoryEntry
 	void SetComponentSpaceTransform(int32 Index, const FTransform& Transform);
 	FTransform GetComponentSpaceTransform(int32 Index) const;
 };
-
-typedef TRingBuffer<FPoseHistoryEntry> FPoseHistoryEntries;
-typedef TArray<FPoseHistoryEntry> FPoseHistoryFutureEntries;
+FArchive& operator<<(FArchive& Ar, FPoseHistoryEntry& Entry);
 
 struct POSESEARCH_API IPoseHistory
 {
@@ -64,91 +62,141 @@ public:
 	virtual bool IsEmpty() const = 0;
 
 	virtual const FBoneToTransformMap& GetBoneToTransformMap() const = 0;
-	virtual const FPoseHistoryEntries& GetEntries() const = 0;
-	virtual const USkeleton* GetLastUpdateSkeleton() const = 0;
+	virtual int32 GetNumEntries() const = 0;
+	virtual const FPoseHistoryEntry& GetEntry(int32 EntryIndex) const = 0;
 
 #if ENABLE_DRAW_DEBUG && ENABLE_ANIM_DEBUG
+	virtual void DebugDraw(const UWorld* World, FColor Color) const = 0;
 	virtual void DebugDraw(FAnimInstanceProxy& AnimInstanceProxy, FColor Color) const = 0;
 	virtual void DebugDraw(FAnimInstanceProxy& AnimInstanceProxy, FColor Color, float Time, float PointSize = 6.f, bool bExtrapolate = true) const;
 #endif
 };
 
-struct FPoseHistory : public IPoseHistory
+struct POSESEARCH_API FArchivedPoseHistory : public IPoseHistory
 {
-	void CacheBones_AnyThread(int32 InNumPoses, float InSamplingInterval, const TArray<FBoneIndexType>& RequiredBones);
-
-	// if bGenerateTrajectory is true, this method will generate trajectory
-	// NoTe: InTrajectory.Samples[i].AccumulatedSeconds == 0 is the sample of the previous frame of simulation (since MM works by matching the previous character pose)
-	void Update_AnyThread(float DeltaTime, const FPoseSearchQueryTrajectory& InTrajectory, float InTrajectorySpeedMultiplier, bool bGenerateTrajectory, const FAnimInstanceProxy& AnimInstanceProxy, const FPoseSearchTrajectoryData& TrajectoryData, const FPoseSearchTrajectoryData::FSampling& TrajectoryDataSampling, bool bNeedsReset);
-
-	void EvaluateComponentSpace_AnyThread(float DeltaTime, FCSPose<FCompactPose>& ComponentSpacePose, bool bGenerateTrajectory, bool bStoreScales);
-
-#if ENABLE_DRAW_DEBUG && ENABLE_ANIM_DEBUG
-	virtual void DebugDraw(FAnimInstanceProxy& AnimInstanceProxy, FColor Color) const override;
-#endif
+	void InitFrom(const IPoseHistory* PoseHistory);
 
 	// IPoseHistory interface
 	virtual bool GetTransformAtTime(float Time, FTransform& OutBoneTransform, const USkeleton* BoneIndexSkeleton = nullptr, FBoneIndexType BoneIndexType = RootBoneIndexType, FBoneIndexType ReferenceBoneIndexType = ComponentSpaceIndexType, bool bExtrapolate = true) const override;
 	virtual const FPoseSearchQueryTrajectory& GetTrajectory() const override { return Trajectory; }
-	virtual float GetTrajectorySpeedMultiplier() const override { return TrajectorySpeedMultiplier; }
+	virtual float GetTrajectorySpeedMultiplier() const override { return 1.f; }
 	virtual bool IsEmpty() const override { return Entries.IsEmpty(); }
 	virtual const FBoneToTransformMap& GetBoneToTransformMap() const override { return BoneToTransformMap; }
-	virtual const FPoseHistoryEntries& GetEntries() const override { return Entries; }
-	virtual const USkeleton* GetLastUpdateSkeleton() const override { return LastUpdateSkeleton.Get(); }
+	
+	virtual int32 GetNumEntries() const override { return Entries.Num(); }
+	virtual const FPoseHistoryEntry& GetEntry(int32 EntryIndex) const override { return Entries[EntryIndex]; }
+
+#if ENABLE_DRAW_DEBUG && ENABLE_ANIM_DEBUG
+	virtual void DebugDraw(const UWorld* World, FColor Color) const override;
+	virtual void DebugDraw(FAnimInstanceProxy& AnimInstanceProxy, FColor Color) const override { unimplemented(); }
+#endif
 	// End of IPoseHistory interface
 
-	static FBoneIndexType GetRemappedBoneIndexType(FBoneIndexType BoneIndexType, const USkeleton* BoneIndexSkeleton, const USkeleton* LastUpdateSkeleton);
-	static FComponentSpaceTransformIndex GetRemappedComponentSpaceTransformIndex(const USkeleton* BoneIndexSkeleton, const USkeleton* LastUpdateSkeleton, const FBoneToTransformMap& BoneToTransformMap, FBoneIndexType BoneIndexType, bool& bSuccess);
-	static bool LerpEntries(float Time, bool bExtrapolate, const FPoseHistoryEntry& PrevEntry, const FPoseHistoryEntry& NextEntry, const USkeleton* BoneIndexSkeleton, const USkeleton* LastUpdateSkeleton, const FBoneToTransformMap& BoneToTransformMap, FBoneIndexType BoneIndexType, FBoneIndexType ReferenceBoneIndexType, FTransform& OutBoneTransform);
+	FBoneToTransformMap BoneToTransformMap;
+	TArray<FPoseHistoryEntry> Entries;
+	FPoseSearchQueryTrajectory Trajectory;
+};
+FArchive& operator<<(FArchive& Ar, FArchivedPoseHistory& Entry);
 
-private:
-	// skeleton from the last Update, to keep tracking skeleton changes, and support compatible skeletons
-	TWeakObjectPtr<const USkeleton> LastUpdateSkeleton;
+struct FPoseHistory : public IPoseHistory
+{
+	FPoseHistory() = default;
+	FPoseHistory(const FPoseHistory& Other);
+	FPoseHistory(FPoseHistory&& Other);
+	FPoseHistory& operator=(const FPoseHistory& Other);
+	FPoseHistory& operator=(FPoseHistory&& Other);
 
-	// map of FBoneIndexType(s) to collect. If Empty all the bones get collected
-	FBoneToTransformMap BoneToTransformMap;
+	void PreUpdate(const UAnimInstance* AnimInstance, float DeltaTime, bool bGenerateTrajectory, const FPoseSearchTrajectoryData& TrajectoryData, const FPoseSearchTrajectoryData::FSampling& TrajectoryDataSampling, bool bNeedsReset);
+
+	void Initialize_AnyThread(int32 InNumPoses, float InSamplingInterval);
+	void CacheBones_AnyThread(const TArray<FBoneIndexType>& RequiredBones);
+
+	// if bGenerateTrajectory is true, this method will generate trajectory
+	// NoTe: InTrajectory.Samples[i].AccumulatedSeconds == 0 is the sample of the previous frame of simulation (since MM works by matching the previous character pose)
+	void Update_AnyThread(float DeltaTime, const FPoseSearchQueryTrajectory& InTrajectory, float InTrajectorySpeedMultiplier);
 
-	// ring buffer of collected bones
-	FPoseHistoryEntries Entries;
+	void EvaluateComponentSpace_AnyThread(float DeltaTime, FCSPose<FCompactPose>& ComponentSpacePose, bool bStoreScales);
+
+	// IPoseHistory interface
+	virtual bool GetTransformAtTime(float Time, FTransform& OutBoneTransform, const USkeleton* BoneIndexSkeleton = nullptr, FBoneIndexType BoneIndexType = RootBoneIndexType, FBoneIndexType ReferenceBoneIndexType = ComponentSpaceIndexType, bool bExtrapolate = true) const override;
+	virtual const FPoseSearchQueryTrajectory& GetTrajectory() const override;
+	virtual float GetTrajectorySpeedMultiplier() const override;
+	virtual bool IsEmpty() const override;
+	virtual const FBoneToTransformMap& GetBoneToTransformMap() const override;
 	
-	// caching MaxNumPoses, since Entries.Max() is a padded number
+	virtual int32 GetNumEntries() const override;
+	virtual const FPoseHistoryEntry& GetEntry(int32 EntryIndex) const override;
+
+#if ENABLE_DRAW_DEBUG && ENABLE_ANIM_DEBUG
+	virtual void DebugDraw(const UWorld* World, FColor Color) const override { unimplemented(); }
+	virtual void DebugDraw(FAnimInstanceProxy& AnimInstanceProxy, FColor Color) const override;
+#endif
+	// End of IPoseHistory interface
+	
+private:
+
+	// caching MaxNumPoses, since FData::Entries.Max() is a padded number
 	int32 MaxNumPoses = 0;
 
 	float SamplingInterval = 0.f;
 
 	FPoseSearchQueryTrajectory Trajectory;
 	FPoseSearchTrajectoryData::FState TrajectoryDataState;
-
 	// @todo: deprecate this member and expose it via blue print logic or as global query scaling multiplier
 	float TrajectorySpeedMultiplier = 1.f;
+
+	struct FData
+	{
+		// skeleton from the last Update, to keep tracking skeleton changes, and support compatible skeletons
+		TWeakObjectPtr<const USkeleton> LastUpdateSkeleton;
+
+		// map of FBoneIndexType(s) to collect. If Empty all the bones get collected
+		FBoneToTransformMap BoneToTransformMap;
+
+		// ring buffer of collected bones
+		TRingBuffer<FPoseHistoryEntry> Entries;
+	};
+
+	FData ReadData;
+	FData WriteData;
+
+#if ENABLE_ANIM_DEBUG
+	
+	// used to analyze thread safety
+	mutable FThreadSafeCounter ReadDataThreadSafeCounter = 0;
+	mutable FThreadSafeCounter WriteDataThreadSafeCounter = 0;
+
+#endif // ENABLE_ANIM_DEBUG
 };
 
-struct FExtendedPoseHistory : public IPoseHistory
+struct FMemStackPoseHistory : public IPoseHistory
 {
 	void Init(const IPoseHistory* InPoseHistory);
 
-	bool IsInitialized() const;
-
 	// IPoseHistory interface
 	virtual bool GetTransformAtTime(float Time, FTransform& OutBoneTransform, const USkeleton* BoneIndexSkeleton = nullptr, FBoneIndexType BoneIndexType = RootBoneIndexType, FBoneIndexType ReferenceBoneIndexType = ComponentSpaceIndexType, bool bExtrapolate = true) const override;
 	virtual const FPoseSearchQueryTrajectory& GetTrajectory() const override { check(PoseHistory); return PoseHistory->GetTrajectory(); }
 	virtual float GetTrajectorySpeedMultiplier() const override { check(PoseHistory); return PoseHistory->GetTrajectorySpeedMultiplier(); }
 	virtual bool IsEmpty() const override { check(PoseHistory); return PoseHistory->IsEmpty() && FutureEntries.IsEmpty(); }
 	virtual const FBoneToTransformMap& GetBoneToTransformMap() const override { check(PoseHistory); return PoseHistory->GetBoneToTransformMap(); }
-	virtual const FPoseHistoryEntries& GetEntries() const override { check(PoseHistory); return PoseHistory->GetEntries(); }
-	virtual const USkeleton* GetLastUpdateSkeleton() const override { check(PoseHistory); return PoseHistory->GetLastUpdateSkeleton(); }
-	// End of IPoseHistory interface
+	
+	virtual int32 GetNumEntries() const override;
+	virtual const FPoseHistoryEntry& GetEntry(int32 EntryIndex) const override;
 	
 #if ENABLE_DRAW_DEBUG && ENABLE_ANIM_DEBUG
+	virtual void DebugDraw(const UWorld* World, FColor Color) const override { unimplemented(); }
 	virtual void DebugDraw(FAnimInstanceProxy& AnimInstanceProxy, FColor Color) const override;
 #endif
+	// End of IPoseHistory interface
 
 	void AddFutureRootBone(float Time, const FTransform& FutureRootBoneTransform, bool bStoreScales);
 	void AddFuturePose(float Time, FCSPose<FCompactPose>& ComponentSpacePose);
 
+	const IPoseHistory* GetThisOrPoseHistory() const { return FutureEntries.IsEmpty() ? PoseHistory : this; }
+
 private:
 	const IPoseHistory* PoseHistory = nullptr;
-	FPoseHistoryFutureEntries FutureEntries;
+	TArray<FPoseHistoryEntry, TInlineAllocator<4, TMemStackAllocator<>>> FutureEntries;
 };
 
 class IPoseHistoryProvider : public UE::Anim::IGraphMessage
diff --git a/Engine/Plugins/Experimental/Animation/PoseSearch/Source/Runtime/Public/PoseSearch/PoseSearchHistoryCollectorAnimNodeLibrary.h b/Engine/Plugins/Experimental/Animation/PoseSearch/Source/Runtime/Public/PoseSearch/PoseSearchHistoryCollectorAnimNodeLibrary.h
new file mode 100644
--- /dev/null
+++ b/Engine/Plugins/Experimental/Animation/PoseSearch/Source/Runtime/Public/PoseSearch/PoseSearchHistoryCollectorAnimNodeLibrary.h
@@ -0,0 +1,43 @@
+// Copyright Epic Games, Inc. All Rights Reserved.
+
+#pragma once
+
+#include "Animation/AnimExecutionContext.h"
+#include "Animation/AnimNodeReference.h"
+#include "Kismet/BlueprintFunctionLibrary.h"
+#include "PoseSearchHistoryCollectorAnimNodeLibrary.generated.h"
+
+struct FAnimNode_PoseSearchHistoryCollector;
+
+USTRUCT(BlueprintType)
+struct FPoseSearchHistoryCollectorAnimNodeReference : public FAnimNodeReference
+{
+	GENERATED_BODY()
+
+	typedef FAnimNode_PoseSearchHistoryCollector FInternalNodeType;
+};
+
+// Exposes operations that can be run on a Pose History node via Anim Node Functions such as "On Become Relevant" and "On Update".
+UCLASS()
+class POSESEARCH_API UPoseSearchHistoryCollectorAnimNodeLibrary : public UBlueprintFunctionLibrary
+{
+	GENERATED_BODY()
+
+public:
+	/** Get a Pose History node context from an anim node context */
+	UFUNCTION(BlueprintCallable, Category = "Animation|PoseHistory", meta = (BlueprintThreadSafe, ExpandEnumAsExecs = "Result"))
+	static FPoseSearchHistoryCollectorAnimNodeReference ConvertToPoseHistoryNode(const FAnimNodeReference& Node, EAnimNodeReferenceConversionResult& Result);
+
+	/** Get a Pose History node context from an anim node context (pure) */
+	UFUNCTION(BlueprintPure, Category = "Animation|PoseHistory", meta = (BlueprintThreadSafe, DisplayName = "Convert to Pose History Node"))
+	static void ConvertToPoseHistoryNodePure(const FAnimNodeReference& Node, FPoseSearchHistoryCollectorAnimNodeReference& PoseSearchHistoryCollectorNode, bool& Result)
+	{
+		EAnimNodeReferenceConversionResult ConversionResult;
+		PoseSearchHistoryCollectorNode = ConvertToPoseHistoryNode(Node, ConversionResult);
+		Result = (ConversionResult == EAnimNodeReferenceConversionResult::Succeeded);
+	}
+
+	UFUNCTION(BlueprintPure, Category = "Animation|PoseHistory", meta = (BlueprintThreadSafe, DisplayName = "Get Pose History Node Trajectory"))
+	static void GetPoseHistoryNodeTrajectory(const FPoseSearchHistoryCollectorAnimNodeReference& PoseSearchHistoryCollectorNode, FPoseSearchQueryTrajectory& Trajectory);
+
+};
\ No newline at end of file
diff --git a/Engine/Plugins/Experimental/Animation/PoseSearch/Source/Runtime/Public/PoseSearch/PoseSearchLibrary.h b/Engine/Plugins/Experimental/Animation/PoseSearch/Source/Runtime/Public/PoseSearch/PoseSearchLibrary.h
--- a/Engine/Plugins/Experimental/Animation/PoseSearch/Source/Runtime/Public/PoseSearch/PoseSearchLibrary.h
+++ b/Engine/Plugins/Experimental/Animation/PoseSearch/Source/Runtime/Public/PoseSearch/PoseSearchLibrary.h
@@ -2,14 +2,13 @@
 
 #pragma once
 
-#include "AlphaBlend.h"
-#include "AnimNodes/AnimNode_SequenceEvaluator.h"
 #include "Animation/AnimExecutionContext.h"
 #include "Animation/AnimNodeReference.h"
 #include "Kismet/BlueprintFunctionLibrary.h"
 #include "PoseSearch/PoseSearchDefines.h"
 #include "PoseSearch/PoseSearchHistory.h"
 #include "PoseSearch/PoseSearchResult.h"
+#include "PoseSearch/PoseSearchRole.h"
 #include "SequenceEvaluatorLibrary.h"
 #include "SequencePlayerLibrary.h"
 #include "PoseSearchLibrary.generated.h"
@@ -100,7 +99,6 @@ class POSESEARCH_API UPoseSearchLibrary : public UBlueprintFunctionLibrary
 		const UE::PoseSearch::FSearchResult& CurrentResult,
 		float ElapsedPoseSearchTime,
 		const FTransform& RootMotionTransformDelta,
-		const UObject* AnimInstance,
 		int32 NodeId,
 		float DeltaTime,
 		bool bSearch,
@@ -171,6 +169,26 @@ public:
 		float TimeToFutureAnimationStart = 0.f,
 		const int32 DebugSessionUniqueIdentifier = 6174);
 
+	UFUNCTION(BlueprintPure, Category = "Animation|Pose Search", meta = (BlueprintThreadSafe, Keywords = "PoseMatch"))
+	static void MotionMatchMulti(
+		TArray<ACharacter*> AnimInstances,
+		TArray<FName> Roles,
+		const UPoseSearchDatabase* Database,
+		const FName PoseHistoryName,
+		FPoseSearchBlueprintResult& Result,
+		const int32 DebugSessionUniqueIdentifier = 6174);
+
+	static void MotionMatch(
+		TArrayView<UAnimInstance*> AnimInstances,
+		TConstArrayView<FName> Roles,
+		const UPoseSearchDatabase* Database,
+		const FName PoseHistoryName,
+		FPoseSearchBlueprintResult& Result,
+		const UAnimationAsset* FutureAnimation,
+		float FutureAnimationStartTime,
+		float TimeToFutureAnimationStart,
+		const int32 DebugSessionUniqueIdentifier);
+
 	static UE::PoseSearch::FSearchResult MotionMatch(const FAnimationBaseContext& Context, TConstArrayView<UObject*> AssetsToSearch,
 		const UObject* PlayingAsset = nullptr, float PlayingAssetAccumulatedTime = 0.f);
 };
diff --git a/Engine/Plugins/Experimental/Animation/PoseSearch/Source/Runtime/Public/PoseSearch/PoseSearchMirrorDataCache.h b/Engine/Plugins/Experimental/Animation/PoseSearch/Source/Runtime/Public/PoseSearch/PoseSearchMirrorDataCache.h
--- a/Engine/Plugins/Experimental/Animation/PoseSearch/Source/Runtime/Public/PoseSearch/PoseSearchMirrorDataCache.h
+++ b/Engine/Plugins/Experimental/Animation/PoseSearch/Source/Runtime/Public/PoseSearch/PoseSearchMirrorDataCache.h
@@ -11,8 +11,6 @@ struct FCompactPose;
 namespace UE::PoseSearch
 {
 
-//POSESEARCH_API FTransform MirrorTransform(const FTransform& InTransform, EAxis::Type MirrorAxis, const FQuat& ReferenceRotation);
-
 struct POSESEARCH_API FMirrorDataCache
 {
 	FMirrorDataCache();
diff --git a/Engine/Plugins/Experimental/Animation/PoseSearch/Source/Runtime/Public/PoseSearch/PoseSearchMultiSequence.h b/Engine/Plugins/Experimental/Animation/PoseSearch/Source/Runtime/Public/PoseSearch/PoseSearchMultiSequence.h
new file mode 100644
--- /dev/null
+++ b/Engine/Plugins/Experimental/Animation/PoseSearch/Source/Runtime/Public/PoseSearch/PoseSearchMultiSequence.h
@@ -0,0 +1,51 @@
+// Copyright Epic Games, Inc. All Rights Reserved.
+
+#pragma once
+
+#include "Engine/DataAsset.h"
+#include "PoseSearch/PoseSearchRole.h"
+#include "PoseSearchMultiSequence.generated.h"
+
+class UAnimSequenceBase;
+
+USTRUCT(Experimental)
+struct POSESEARCH_API FPoseSearchMultiSequenceItem
+{
+	GENERATED_BODY()
+
+	UPROPERTY(EditAnywhere, Category = "Settings")
+	TObjectPtr<UAnimSequenceBase> Sequence;
+
+	UPROPERTY(EditAnywhere, Category = "Settings")
+	FName Role;
+
+	UPROPERTY(EditAnywhere, Category = "Settings")
+	FTransform Origin = FTransform::Identity;
+};
+
+UCLASS(Experimental, BlueprintType, Category = "Animation|Pose Search")
+class UPoseSearchMultiSequence : public UDataAsset
+{
+	GENERATED_BODY()
+
+	UPROPERTY(EditAnywhere, Category = "Settings")
+	TArray<FPoseSearchMultiSequenceItem> Items;
+
+public:
+	bool IsLooping() const;
+	const FString GetName() const;
+	bool HasRootMotion() const;
+	float GetPlayLength() const;
+
+#if WITH_EDITOR
+	int32 GetFrameAtTime(float Time) const;
+#endif // WITH_EDITOR
+
+	int32 GetNumRoles() const { return Items.Num(); }
+	const UE::PoseSearch::FRole& GetRole(int32 RoleIndex) const { return Items[RoleIndex].Role; }
+	
+	UFUNCTION(BlueprintPure, Category = "Animation", meta=(BlueprintThreadSafe))
+	UAnimSequenceBase* GetSequence(const FName /*UE::PoseSearch::FRole*/& Role) const;
+
+	const FTransform& GetOrigin(const UE::PoseSearch::FRole& Role) const;
+};
\ No newline at end of file
diff --git a/Engine/Plugins/Experimental/Animation/PoseSearch/Source/Runtime/Public/PoseSearch/PoseSearchSchema.h b/Engine/Plugins/Experimental/Animation/PoseSearch/Source/Runtime/Public/PoseSearch/PoseSearchSchema.h
--- a/Engine/Plugins/Experimental/Animation/PoseSearch/Source/Runtime/Public/PoseSearch/PoseSearchSchema.h
+++ b/Engine/Plugins/Experimental/Animation/PoseSearch/Source/Runtime/Public/PoseSearch/PoseSearchSchema.h
@@ -4,7 +4,6 @@
 
 #include "BoneContainer.h"
 #include "Engine/DataAsset.h"
-#include "Interfaces/Interface_BoneReferenceSkeletonProvider.h"
 #include "PoseSearch/PoseSearchFeatureChannel.h"
 #include "PoseSearchSchema.generated.h"
 
@@ -24,24 +23,50 @@ enum class EPoseSearchDataPreprocessor : int32
 	NormalizeOnlyByDeviation
 };
 
+
+USTRUCT()
+struct POSESEARCH_API FPoseSearchRoledSkeleton
+{
+	GENERATED_BODY()
+
+	// Skeleton Reference for Motion Matching Database assets. Must be set to a compatible skeleton to the animation data in the database.
+	UPROPERTY(EditAnywhere, Category = "Schema")
+	TObjectPtr<USkeleton> Skeleton;
+
+	// Setting up and assigning a mirror data table will allow all your assets in your database to access the mirrored version of the data. This is required for mirroring to work with Motion Matching.
+	UPROPERTY(EditAnywhere, Category = "Schema")
+	TObjectPtr<UMirrorDataTable> MirrorDataTable;
+
+	UPROPERTY(EditAnywhere, Category = "Schema")
+	FName Role;
+
+	UPROPERTY(Transient)
+	TArray<FBoneReference> BoneReferences;
+
+	UPROPERTY(Transient)
+	TArray<uint16> BoneIndicesWithParents;
+};
+
 /**
 * Specifies the format of a pose search index. At runtime, queries are built according to the schema for searching.
 */
 UCLASS(BlueprintType, Category = "Animation|Pose Search", meta = (DisplayName = "Pose Search Schema"), CollapseCategories)
-class POSESEARCH_API UPoseSearchSchema : public UDataAsset, public IBoneReferenceSkeletonProvider
+class POSESEARCH_API UPoseSearchSchema : public UDataAsset
 {
 	GENERATED_BODY()
 
 public:
-	// Skeleton Reference for Motion Matching Database assets. Must be set to a compatible skeleton to the animation data in the database.
-	UPROPERTY(EditAnywhere, Category = "Schema", meta = (DisplayPriority = 0))
-	TObjectPtr<USkeleton> Skeleton;
-
+	UPROPERTY()
+	TObjectPtr<USkeleton> Skeleton_DEPRECATED;
+	
 	// The update rate at which we sample the animation data in the database. The higher the SampleRate the more refined your searches will be, but the more memory will be required
 	UPROPERTY(EditAnywhere, Category = "Schema", meta = (DisplayPriority = 3, ClampMin = "1", ClampMax = "240"))
 	int32 SampleRate = 30;
 
 private:
+	UPROPERTY(EditAnywhere, Category = "Schema", meta = (DisplayPriority = 0))
+	TArray<FPoseSearchRoledSkeleton> Skeletons;
+
 	// Channels itemize the cost breakdown of the Schema in simpler parts such as position or velocity of a bones, or phase of limbs. The total cost of a query against an indexed database pose will be the sum of the combined channel costs
 	UPROPERTY(EditAnywhere, Instanced, Category = "Schema")
 	TArray<TObjectPtr<UPoseSearchFeatureChannel>> Channels;
@@ -51,9 +76,8 @@ private:
 	TArray<TObjectPtr<UPoseSearchFeatureChannel>> FinalizedChannels;
 
 public:
-	// Setting up and assigning a mirror data table will allow all your assets in your database to access the mirrored version of the data. This is required for mirroring to work with Motion Matching.
-	UPROPERTY(EditAnywhere, Category = "Schema", meta = (DisplayPriority = 1))
-	TObjectPtr<UMirrorDataTable> MirrorDataTable;
+	UPROPERTY()
+	TObjectPtr<UMirrorDataTable> MirrorDataTable_DEPRECATED;
 
 #if WITH_EDITORONLY_DATA
 	// Type of operation performed to the full pose features dataset
@@ -64,12 +88,6 @@ public:
 	UPROPERTY(Transient)
 	int32 SchemaCardinality = 0;
 
-	UPROPERTY(Transient)
-	TArray<FBoneReference> BoneReferences;
-
-	UPROPERTY(Transient)
-	TArray<uint16> BoneIndicesWithParents;
-
 #if WITH_EDITORONLY_DATA
 	// How many times the animation assets of the database using this schema will be indexed.
 	UPROPERTY(EditAnywhere, Category = "Permutations", meta = (ClampMin = "1"))
@@ -95,7 +113,7 @@ public:
 	UPROPERTY(EditAnywhere, Category = "Debug")
 	bool bInjectAdditionalDebugChannels;
 
-	bool IsValid () const;
+	//bool IsValid () const;
 
 	TConstArrayView<TObjectPtr<UPoseSearchFeatureChannel>> GetChannels() const { return FinalizedChannels; }
 
@@ -118,21 +136,31 @@ public:
 	virtual void PreSave(FObjectPreSaveContext ObjectSaveContext) override;
 	virtual void PostLoad() override;
 
-	int8 AddBoneReference(const FBoneReference& BoneReference);
-
-	// IBoneReferenceSkeletonProvider
-	USkeleton* GetSkeleton(bool& bInvalidSkeletonIsError, const IPropertyHandle* PropertyHandle) override;
+	int8 AddBoneReference(const FBoneReference& BoneReference, const UE::PoseSearch::FRole& Role);
 
 #if WITH_EDITOR
 	virtual void PostEditChangeProperty(struct FPropertyChangedEvent& PropertyChangedEvent) override;
-#endif
+	const UE::PoseSearch::FRole GetDefaultRole() const;
+#endif // WITH_EDITOR
 
 	TConstArrayView<float> BuildQuery(UE::PoseSearch::FSearchContext& SearchContext) const;
 
-	FBoneIndexType GetBoneIndexType(int8 SchemaBoneIdx) const;
-
+	void AddSkeleton(USkeleton* Skeleton, UMirrorDataTable* MirrorDataTable = nullptr, const UE::PoseSearch::FRole& Role = UE::PoseSearch::DefaultRole);
+	bool AreSkeletonsCompatible(const UPoseSearchSchema* Other) const;
 	void AddDefaultChannels();
 
+	// we avoid exposing GetRoledSkeleton(RoleIndex) and all others method referencing Skeletons by index
+	// to avoid confusing the Role to index mapping in the Schema versus the Role to index mapping in other assets
+	// please use GetRoledSkeleton(Role)
+	void InitBoneContainersFromRoledSkeleton(TMap<FName, FBoneContainer>& RoledBoneContainers) const;
+	bool AllRoledSkeletonHaveMirrorDataTable() const;
+	const FPoseSearchRoledSkeleton* GetRoledSkeleton(const UE::PoseSearch::FRole& Role) const;
+	FPoseSearchRoledSkeleton* GetRoledSkeleton(const UE::PoseSearch::FRole& Role);
+
+	USkeleton* GetSkeleton(const UE::PoseSearch::FRole& Role) const;
+	UMirrorDataTable* GetMirrorDataTable(const UE::PoseSearch::FRole& Role) const;
+	TConstArrayView<FBoneReference> GetBoneReferences(const UE::PoseSearch::FRole& Role) const;
+
 private:
 	template <typename FindPredicateType>
 	static const UPoseSearchFeatureChannel* FindChannelRecursive(TConstArrayView<TObjectPtr<UPoseSearchFeatureChannel>> Channels, FindPredicateType FindPredicate)
@@ -156,4 +184,5 @@ private:
 	}
 
 	void Finalize();
+	void ResetFinalize();
 };
\ No newline at end of file
diff --git a/Engine/Plugins/Experimental/Animation/PoseSearch/Source/Runtime/Public/PoseSearch/PoseSearchTrajectoryLibrary.h b/Engine/Plugins/Experimental/Animation/PoseSearch/Source/Runtime/Public/PoseSearch/PoseSearchTrajectoryLibrary.h
--- a/Engine/Plugins/Experimental/Animation/PoseSearch/Source/Runtime/Public/PoseSearch/PoseSearchTrajectoryLibrary.h
+++ b/Engine/Plugins/Experimental/Animation/PoseSearch/Source/Runtime/Public/PoseSearch/PoseSearchTrajectoryLibrary.h
@@ -6,6 +6,8 @@
 #include "PoseSearch/PoseSearchTrajectoryTypes.h"
 #include "PoseSearchTrajectoryLibrary.generated.h"
 
+class UAnimInstance;
+
 USTRUCT(BlueprintType)
 struct POSESEARCH_API FPoseSearchTrajectoryData
 {
@@ -44,6 +46,7 @@ public:
 	};
 
 	void UpdateData(float DeltaTime, const FAnimInstanceProxy& AnimInstanceProxy, FDerived& TrajectoryDataDerived, FState& TrajectoryDataState) const;
+	void UpdateData(float DeltaTime, const UAnimInstance* AnimInstance, FDerived& TrajectoryDataDerived, FState& TrajectoryDataState) const;
 	FVector StepCharacterMovementGroundPrediction(float DeltaTime, const FVector& InVelocity, const FVector& InAcceleration, const FDerived& TrajectoryDataDerived) const;
 	
 	// If the character is forward facing (i.e. bOrientRotationToMovement is true), this controls how quickly the trajectory will rotate
@@ -77,10 +80,12 @@ public:
 
 /**
  * Set of functions to help populate a FPoseSearchQueryTrajectory for motion matching.
- * UCharacterTrajectoryComponent uses these functions, but they can also be used by a UAnimInstance to avoid the component.
  */
-struct POSESEARCH_API FPoseSearchTrajectoryLibrary
+UCLASS()
+class POSESEARCH_API UPoseSearchTrajectoryLibrary : public UBlueprintFunctionLibrary
 {
+	GENERATED_BODY()
+
 public:
 	static void InitTrajectorySamples(FPoseSearchQueryTrajectory& Trajectory, const FPoseSearchTrajectoryData& TrajectoryData, const FPoseSearchTrajectoryData::FDerived& TrajectoryDataDerived, const FPoseSearchTrajectoryData::FSampling& TrajectoryDataSampling);
 
@@ -92,6 +97,13 @@ public:
 	// Update prediction by simulating the movement math for ground locomotion from UCharacterMovementComponent.
 	static void UpdatePrediction_SimulateCharacterMovement(FPoseSearchQueryTrajectory& Trajectory, const FPoseSearchTrajectoryData& TrajectoryData, const FPoseSearchTrajectoryData::FDerived& TrajectoryDataDerived, const FPoseSearchTrajectoryData::FSampling& TrajectoryDataSampling);
 
+	/** Get a Pose History node context from an anim node context (pure) */
+	UFUNCTION(BlueprintPure, Category = "Animation|PoseSearch", meta = (BlueprintThreadSafe, DisplayName = "Pose Search Generate Trajectory"))
+	static void PoseSearchGenerateTrajectory(
+		const UAnimInstance* InAnimInstance, UPARAM(ref) const FPoseSearchTrajectoryData& InTrajectoryData, float InDeltaTime,
+		UPARAM(ref) FPoseSearchQueryTrajectory& InOutTrajectory, UPARAM(ref) float& InOutDesiredControllerYawLastUpdate, FPoseSearchQueryTrajectory& OutTrajectory,
+		float InHistorySamplingInterval = 0.04f, int32 InTrajectoryHistoryCount = 10, float InPredictionSamplingInterval = 0.4f, int32 InTrajectoryPredictionCount = 8);
+
 private:
 	static FVector RemapVectorMagnitudeWithCurve(const FVector& Vector, bool bUseCurve, const FRuntimeFloatCurve& Curve);
 };
\ No newline at end of file
diff --git a/Engine/Plugins/Experimental/Animation/PoseSearch/Source/Runtime/Public/PoseSearch/Trace/PoseSearchTraceLogger.h b/Engine/Plugins/Experimental/Animation/PoseSearch/Source/Runtime/Public/PoseSearch/Trace/PoseSearchTraceLogger.h
--- a/Engine/Plugins/Experimental/Animation/PoseSearch/Source/Runtime/Public/PoseSearch/Trace/PoseSearchTraceLogger.h
+++ b/Engine/Plugins/Experimental/Animation/PoseSearch/Source/Runtime/Public/PoseSearch/Trace/PoseSearchTraceLogger.h
@@ -3,6 +3,7 @@
 #pragma once
 
 #include "PoseSearch/PoseSearchContext.h"
+#include "PoseSearch/PoseSearchHistory.h"
 #include "PoseSearch/PoseSearchTrajectoryTypes.h"
 
 UE_TRACE_CHANNEL_EXTERN(PoseSearchChannel, POSESEARCH_API);
@@ -63,7 +64,7 @@ POSESEARCH_API FArchive& operator<<(FArchive& Ar, FTraceMotionMatchingStateDatab
 /**
  * Used to trace motion matching state data via the logger, which is then placed into a timeline
  */
-struct POSESEARCH_API FTraceMotionMatchingState
+struct POSESEARCH_API FTraceMotionMatchingStateMessage : public FTraceMessage
 {
 	/** Amount of time since the last pose switch */
 	float ElapsedPoseSearchTime = 0.f;
@@ -80,9 +81,13 @@ struct POSESEARCH_API FTraceMotionMatchingState
 	float SearchBruteForceCost = 0.f;
 	int32 SearchBestPosePos = 0;
 
+	TArray<uint64> SkeletalMeshComponentIds;
+	
+	TArray<FRole> Roles;
+
 	TArray<FTraceMotionMatchingStateDatabaseEntry> DatabaseEntries;
 
-	FPoseSearchQueryTrajectory Trajectory;
+	TArray<UE::PoseSearch::FArchivedPoseHistory> PoseHistories;
 
 	/** Index of the current database in DatabaseEntries */
 	int32 CurrentDbEntryIdx = INDEX_NONE;
@@ -91,7 +96,7 @@ struct POSESEARCH_API FTraceMotionMatchingState
 	int32 CurrentPoseEntryIdx = INDEX_NONE;
 
 	/** Output the current state info to the logger */
-	void Output(const UObject* AnimInstance, int32 NodeId);
+	void Output();
 
 	const UPoseSearchDatabase* GetCurrentDatabase() const;
 	int32 GetCurrentDatabasePoseIndex() const;
@@ -126,6 +131,6 @@ struct POSESEARCH_API FTraceMotionMatchingState
 	static const FName Name;
 };
 
-POSESEARCH_API FArchive& operator<<(FArchive& Ar, FTraceMotionMatchingState& State);
+POSESEARCH_API FArchive& operator<<(FArchive& Ar, FTraceMotionMatchingStateMessage& State);
 
 } // namespace UE::PoseSearch
