motion matching - fix for UPoseSearchDatabase::SynchronizeWithExternalDependencies

diff --git a/Engine/Plugins/Experimental/Animation/PoseSearch/Source/Runtime/Private/PoseSearchDatabase.cpp b/Engine/Plugins/Experimental/Animation/PoseSearch/Source/Runtime/Private/PoseSearchDatabase.cpp
--- a/Engine/Plugins/Experimental/Animation/PoseSearch/Source/Runtime/Private/PoseSearchDatabase.cpp
+++ b/Engine/Plugins/Experimental/Animation/PoseSearch/Source/Runtime/Private/PoseSearchDatabase.cpp
@@ -811,22 +811,23 @@ void UPoseSearchDatabase::SynchronizeWithExternalDependencies()
 
 void UPoseSearchDatabase::SynchronizeWithExternalDependencies(TConstArrayView<UAnimSequenceBase*> SequencesBase)
 {
-	// @todo: use a TSet to speed up this method
-	TArray<FInstancedStruct> AnimationAssetsCopy = AnimationAssets;
+	// cannot use TSet since FInstancedStruct doesn't implement GetTypeHash
+	TArray<FInstancedStruct> NewAnimationAssets;
 
-	// removing all the SequencesBase references from AnimationAssetsCopy
-	for (int32 AnimationAssetIndex = AnimationAssetsCopy.Num() - 1; AnimationAssetIndex >= 0; --AnimationAssetIndex)
+	// collecting all the database AnimationAsset(s) that don't require synchronization
+	for (FInstancedStruct& AnimationAsset : AnimationAssets)
 	{
-		if (const FPoseSearchDatabaseAnimationAssetBase* AnimationAssetBase = AnimationAssetsCopy[AnimationAssetIndex].GetPtr<FPoseSearchDatabaseAnimationAssetBase>())
+		if (const FPoseSearchDatabaseAnimationAssetBase* AnimationAssetBase = AnimationAsset.GetPtr<FPoseSearchDatabaseAnimationAssetBase>())
 		{
-			if (AnimationAssetBase->bSynchronizeWithExternalDependency && SequencesBase.Contains(AnimationAssetBase->GetAnimationAsset()))
+			const bool bRequiresSynchronization = AnimationAssetBase->bSynchronizeWithExternalDependency && SequencesBase.Contains(AnimationAssetBase->GetAnimationAsset());
+			if (!bRequiresSynchronization)
 			{
-				AnimationAssetsCopy.RemoveAt(AnimationAssetIndex);
+				NewAnimationAssets.Add(AnimationAsset);
 			}
 		}
 	}
 
-	// readding the required / updated references
+	// collecting all the SequencesBase(s) requiring synchronization
 	for (UAnimSequenceBase* SequenceBase : SequencesBase)
 	{
 		if (SequenceBase)
@@ -843,7 +844,7 @@ void UPoseSearchDatabase::SynchronizeWithExternalDependencies(TConstArrayView<UA
 							DatabaseSequence.Sequence = Sequence;
 							DatabaseSequence.SamplingRange = FFloatInterval(NotifyEvent.GetTime(), NotifyEvent.GetTime() + NotifyEvent.GetDuration());
 							DatabaseSequence.bSynchronizeWithExternalDependency = true;
-							AnimationAssetsCopy.Add(FInstancedStruct::Make(DatabaseSequence));
+							NewAnimationAssets.Add(FInstancedStruct::Make(DatabaseSequence));
 						}
 						else if (UAnimComposite* AnimComposite = Cast<UAnimComposite>(SequenceBase))
 						{
@@ -851,7 +852,7 @@ void UPoseSearchDatabase::SynchronizeWithExternalDependencies(TConstArrayView<UA
 							DatabaseAnimComposite.AnimComposite = AnimComposite;
 							DatabaseAnimComposite.SamplingRange = FFloatInterval(NotifyEvent.GetTime(), NotifyEvent.GetTime() + NotifyEvent.GetDuration());
 							DatabaseAnimComposite.bSynchronizeWithExternalDependency = true;
-							AnimationAssetsCopy.Add(FInstancedStruct::Make(DatabaseAnimComposite));
+							NewAnimationAssets.Add(FInstancedStruct::Make(DatabaseAnimComposite));
 						}
 						else if (UAnimMontage* AnimMontage = Cast<UAnimMontage>(SequenceBase))
 						{
@@ -859,7 +860,7 @@ void UPoseSearchDatabase::SynchronizeWithExternalDependencies(TConstArrayView<UA
 							DatabaseAnimMontage.AnimMontage = AnimMontage;
 							DatabaseAnimMontage.SamplingRange = FFloatInterval(NotifyEvent.GetTime(), NotifyEvent.GetTime() + NotifyEvent.GetDuration());
 							DatabaseAnimMontage.bSynchronizeWithExternalDependency = true;
-							AnimationAssetsCopy.Add(FInstancedStruct::Make(DatabaseAnimMontage));
+							NewAnimationAssets.Add(FInstancedStruct::Make(DatabaseAnimMontage));
 						}
 					}
 				}
@@ -867,14 +868,14 @@ void UPoseSearchDatabase::SynchronizeWithExternalDependencies(TConstArrayView<UA
 		}
 	}
 
-	// updating AnimationAssets from AnimationAssetsCopy
+	// updating AnimationAssets from NewAnimationAssets preserving the original sorting
 	bool bModified = false;
 	for (int32 AnimationAssetIndex = AnimationAssets.Num() - 1; AnimationAssetIndex >= 0; --AnimationAssetIndex)
 	{
-		const int32 CopyIndex = AnimationAssetsCopy.Find(AnimationAssets[AnimationAssetIndex]);
-		if (CopyIndex >= 0)
+		const int32 FoundIndex = NewAnimationAssets.Find(AnimationAssets[AnimationAssetIndex]);
+		if (FoundIndex >= 0)
 		{
-			AnimationAssetsCopy.RemoveAt(CopyIndex);
+			NewAnimationAssets.RemoveAt(FoundIndex);
 		}
 		else
 		{
@@ -883,6 +884,13 @@ void UPoseSearchDatabase::SynchronizeWithExternalDependencies(TConstArrayView<UA
 		}
 	}
 
+	// adding the remaining AnimationAsset(s) from AnimationAssetsSet
+	for (const FInstancedStruct& AnimationAsset : NewAnimationAssets)
+	{
+		AnimationAssets.Add(AnimationAsset);
+		bModified = true;
+	}
+
 	if (bModified)
 	{
 		Modify();
