motion matching - code cleanup

diff --git a/Engine/Plugins/Experimental/Animation/PoseSearch/Source/Editor/Private/PoseSearchDebuggerDatabaseView.cpp b/Engine/Plugins/Experimental/Animation/PoseSearch/Source/Editor/Private/PoseSearchDebuggerDatabaseView.cpp
--- a/Engine/Plugins/Experimental/Animation/PoseSearch/Source/Editor/Private/PoseSearchDebuggerDatabaseView.cpp
+++ b/Engine/Plugins/Experimental/Animation/PoseSearch/Source/Editor/Private/PoseSearchDebuggerDatabaseView.cpp
@@ -631,8 +631,6 @@ void SDebuggerDatabaseView::PopulateViewRows()
 		TSet<int32> BestAssetPoseIndex;
 		if (bShowOnlyBestAssetPose)
 		{
-			// @todo: perhaps optimize via:
-			//TArrayView<int32> SortIndex((int32*)FMemory_Alloca(UnfilteredDatabaseRowsNum * sizeof(int32)), UnfilteredDatabaseRowsNum);
 			TArray<int32> SortIndex;
 			SortIndex.SetNumUninitialized(UnfilteredDatabaseRowsNum);
 			for (int32 i = 0; i < UnfilteredDatabaseRowsNum; ++i)
diff --git a/Engine/Plugins/Experimental/Animation/PoseSearch/Source/Editor/Private/PoseSearchDebuggerView.cpp b/Engine/Plugins/Experimental/Animation/PoseSearch/Source/Editor/Private/PoseSearchDebuggerView.cpp
--- a/Engine/Plugins/Experimental/Animation/PoseSearch/Source/Editor/Private/PoseSearchDebuggerView.cpp
+++ b/Engine/Plugins/Experimental/Animation/PoseSearch/Source/Editor/Private/PoseSearchDebuggerView.cpp
@@ -181,7 +181,7 @@ void SDebuggerView::Tick(const FGeometry& AllottedGeometry, const double InCurre
 	const UWorld* DebuggerWorld = FDebugger::GetWorld();
 	check(DebuggerWorld);
 
-	// @TODO: Handle editor world when those features are enabled for the Rewind Debugger
+	// @todo: Handle editor world when those features are enabled for the Rewind Debugger
 	// Currently prevents debug draw remnants from stopped world
 	if (DebuggerWorld->WorldType != EWorldType::PIE)
 	{
@@ -396,20 +396,20 @@ bool SDebuggerView::UpdateNodeSelection()
 	bool bNodeSelected = SelectedNodeId != INDEX_NONE;
 	if (!bNodeSelected)
 	{
-		const TArray<int32>& NodeIds = *Model->GetNodeIds();
-		// Only one node active, bypass selection view
-		if (NodeIds.Num() == 1)
+		const TArray<FTraceMotionMatchingStateMessage>& MotionMatchingStates = Model->GetMotionMatchingStates();
+		// Only one active state, bypass selection view
+		if (MotionMatchingStates.Num() == 1)
 		{
-			SelectedNodeId = *NodeIds.begin();
+			SelectedNodeId = MotionMatchingStates[0].NodeId;
 			bNodeSelected = true;
 		}
 		// Create selection view with buttons for each node, displaying the database name
 		else
 		{
 			SelectionView->ClearChildren();
-			for (int32 NodeId : NodeIds)
+			for (const FTraceMotionMatchingStateMessage& MotionMatchingState : MotionMatchingStates)
 			{
-				Model->OnUpdateNodeSelection(NodeId);
+				Model->OnUpdateNodeSelection(MotionMatchingState.NodeId);
 				SelectionView->AddSlot()
 				.HAlign(HAlign_Fill)
 				.VAlign(VAlign_Center)
@@ -420,7 +420,7 @@ bool SDebuggerView::UpdateNodeSelection()
 					.HAlign(HAlign_Center)
 					.VAlign(VAlign_Center)
 					.ContentPadding(10.0f)
-					.OnClicked(this, &SDebuggerView::OnUpdateNodeSelection, NodeId)
+					.OnClicked(this, &SDebuggerView::OnUpdateNodeSelection, MotionMatchingState.NodeId)
 				];
 			}
 		}
diff --git a/Engine/Plugins/Experimental/Animation/PoseSearch/Source/Editor/Private/PoseSearchDebuggerViewModel.cpp b/Engine/Plugins/Experimental/Animation/PoseSearch/Source/Editor/Private/PoseSearchDebuggerViewModel.cpp
--- a/Engine/Plugins/Experimental/Animation/PoseSearch/Source/Editor/Private/PoseSearchDebuggerViewModel.cpp
+++ b/Engine/Plugins/Experimental/Animation/PoseSearch/Source/Editor/Private/PoseSearchDebuggerViewModel.cpp
@@ -51,24 +51,13 @@ const UPoseSearchDatabase* FDebuggerViewModel::GetCurrentDatabase() const
 	return nullptr;
 }
 
-const TArray<int32>* FDebuggerViewModel::GetNodeIds() const
-{
-	return &NodeIds;
-}
-
 int32 FDebuggerViewModel::GetNodesNum() const
 {
 	return MotionMatchingStates.Num();
 }
 
-const FTransform& FDebuggerViewModel::GetRootBoneTransform() const
-{
-	return RootBoneWorldTransform;
-}
-
 void FDebuggerViewModel::OnUpdate()
 {
-	NodeIds.Empty();
 	MotionMatchingStates.Empty();
 
 	// Get provider and validate
@@ -97,8 +86,6 @@ void FDebuggerViewModel::OnUpdate()
 		if (Message)
 		{
 			check(Message->Roles.Num() == Message->SkeletalMeshComponentIds.Num());
-
-			NodeIds.Add(Message->NodeId);
 			MotionMatchingStates.Add(*Message);
 		}
 	});
@@ -188,16 +175,6 @@ void FDebuggerViewModel::OnUpdate()
 
 									check(ComponentWorldTransform.Equals(PoseMessage.ComponentToWorld));
 
-									// @todo: remove RootBoneWorldTransform and use Trajectoris instead
-									//if (!ComponentSpaceTransforms.IsEmpty())
-									//{
-									//	RootBoneWorldTransform = ComponentSpaceTransforms[RootBoneIndexType] * ComponentWorldTransform;
-									//}
-									//else
-									//{
-									//	RootBoneWorldTransform = ComponentWorldTransform;
-									//}
-
 									PoseSearchMeshComponent->Initialize(ComponentWorldTransform);
 
 									return TraceServices::EEventEnumerate::Stop;
@@ -212,20 +189,18 @@ void FDebuggerViewModel::OnUpdate()
 
 void FDebuggerViewModel::OnUpdateNodeSelection(int32 InNodeId)
 {
-	if (InNodeId == INDEX_NONE)
+	if (InNodeId != INDEX_NONE)
 	{
-		return;
-	}
-
-	// Find node in all motion matching states this frame
-	ActiveMotionMatchingStateIdx = INDEX_NONE;
-	const int32 NodesNum = NodeIds.Num();
-	for (int32 i = 0; i < NodesNum; ++i)
-	{
-		if (NodeIds[i] == InNodeId)
+		// Find node in all motion matching states this frame
+		ActiveMotionMatchingStateIdx = INDEX_NONE;
+		const int32 NodesNum = MotionMatchingStates.Num();
+		for (int32 i = 0; i < NodesNum; ++i)
 		{
-			ActiveMotionMatchingStateIdx = i;
-			break;
+			if (MotionMatchingStates[i].NodeId == InNodeId)
+			{
+				ActiveMotionMatchingStateIdx = i;
+				break;
+			}
 		}
 	}
 }
diff --git a/Engine/Plugins/Experimental/Animation/PoseSearch/Source/Editor/Private/PoseSearchDebuggerViewModel.h b/Engine/Plugins/Experimental/Animation/PoseSearch/Source/Editor/Private/PoseSearchDebuggerViewModel.h
--- a/Engine/Plugins/Experimental/Animation/PoseSearch/Source/Editor/Private/PoseSearchDebuggerViewModel.h
+++ b/Engine/Plugins/Experimental/Animation/PoseSearch/Source/Editor/Private/PoseSearchDebuggerViewModel.h
@@ -27,11 +27,10 @@ public:
 	// Used for view callbacks
     const FTraceMotionMatchingStateMessage* GetMotionMatchingState() const;
 	const TMap<uint64, TWeakObjectPtr<AActor>>& GetDebugDrawActors() const { return DebugDrawActors; }
+	const TArray<FTraceMotionMatchingStateMessage>& GetMotionMatchingStates() const { return MotionMatchingStates; }
 
 	const UPoseSearchDatabase* GetCurrentDatabase() const;
-	const TArray<int32>* GetNodeIds() const;
 	int32 GetNodesNum() const;
-	const FTransform& GetRootBoneTransform() const;
 
 	/** Update motion matching states for frame */
 	void OnUpdate();
@@ -53,20 +52,12 @@ public:
 
 private:
 
-	// @todo: duplicate member, already contained in MotionMatchingStates(FTraceMessage)::NodeId
-	/** List of all Node IDs associated with motion matching states */
-	TArray<int32> NodeIds;
-	
 	/** List of all updated motion matching states per node */
 	TArray<FTraceMotionMatchingStateMessage> MotionMatchingStates;
 	
 	/** Currently active motion matching state index based on node selection in the view */
 	int32 ActiveMotionMatchingStateIdx = INDEX_NONE;
 
-	// @todo: duplicate information already present in MotionMatchingStates::Trajectories
-	/** Currently active root bone transform */
-	FTransform RootBoneWorldTransform = FTransform::Identity;
-
 	/** Pointer to the active rewind debugger in the scene */
 	TAttribute<const IRewindDebugger*> RewindDebugger;
 
diff --git a/Engine/Plugins/Experimental/Animation/PoseSearch/Source/Runtime/Private/PoseSearchAssetIndexer.cpp b/Engine/Plugins/Experimental/Animation/PoseSearch/Source/Runtime/Private/PoseSearchAssetIndexer.cpp
--- a/Engine/Plugins/Experimental/Animation/PoseSearch/Source/Runtime/Private/PoseSearchAssetIndexer.cpp
+++ b/Engine/Plugins/Experimental/Animation/PoseSearch/Source/Runtime/Private/PoseSearchAssetIndexer.cpp
@@ -571,7 +571,6 @@ bool FAssetIndexer::GetSampleRotation(FQuat& OutSampleRotation, float SampleTime
 							bool bUnused;
 							if (SamplingAttribute->Bone.BoneName != NAME_None)
 							{
-								// @todo perhaps cache the initialized bone to speed up this method
 								FBoneReference TempBoneReference = SamplingAttribute->Bone;
 								const int32 SampleRoleIndex = RoleToIndex[SampleRole];
 								TempBoneReference.Initialize(BoneContainers[SampleRoleIndex].GetSkeletonAsset());
@@ -643,7 +642,6 @@ bool FAssetIndexer::GetSamplePositionInternal(FVector& OutSamplePosition, float
 							bool bUnused;
 							if (SamplingAttribute->Bone.BoneName != NAME_None)
 							{
-								// @todo perhaps cache the initialized bone to speed up this method
 								FBoneReference TempBoneReference = SamplingAttribute->Bone;
 								const int32 SampleRoleIndex = RoleToIndex[SampleRole];
 								TempBoneReference.Initialize(BoneContainers[SampleRoleIndex].GetSkeletonAsset());
@@ -729,8 +727,6 @@ bool FAssetIndexer::GetSampleVelocity(FVector& OutSampleVelocity, float SampleTi
 							FVector BonePositionPast, BonePositionPresent;
 							if (SamplingAttribute->Bone.BoneName != NAME_None)
 							{
-								// @todo: FIX ME TempBoneReference is not used!
-								// @todo perhaps cache the initialized bone to speed up this method
 								FBoneReference TempBoneReference = SamplingAttribute->Bone;
 								const int32 SampleRoleIndex = RoleToIndex[SampleRole];
 								TempBoneReference.Initialize(BoneContainers[SampleRoleIndex].GetSkeletonAsset());
@@ -798,7 +794,6 @@ bool FAssetIndexer::GetSampleVelocity(FVector& OutSampleVelocity, float SampleTi
 		}
 	}
 
-	// @todo: should we log some error message?
 	OutSampleVelocity = FVector::ZeroVector;
 	return false;
 }
diff --git a/Engine/Plugins/Experimental/Animation/PoseSearch/Source/Runtime/Private/PoseSearchContext.cpp b/Engine/Plugins/Experimental/Animation/PoseSearch/Source/Runtime/Private/PoseSearchContext.cpp
--- a/Engine/Plugins/Experimental/Animation/PoseSearch/Source/Runtime/Private/PoseSearchContext.cpp
+++ b/Engine/Plugins/Experimental/Animation/PoseSearch/Source/Runtime/Private/PoseSearchContext.cpp
@@ -117,7 +117,7 @@ FVector FDebugDrawParams::ExtractPosition(TConstArrayView<float> PoseVector, flo
 			}
 		}
 	}
-	return GetRootTransform(Role).GetTranslation();
+	return GetRootTransform(Role, SampleTimeOffset).GetTranslation();
 }
 
 FQuat FDebugDrawParams::ExtractRotation(TConstArrayView<float> PoseVector, float SampleTimeOffset, int8 SchemaBoneIdx, const FRole& Role, EPermutationTimeType PermutationTimeType, int32 SamplingAttributeId) const
@@ -243,15 +243,14 @@ FQuat FDebugDrawParams::ExtractRotation(TConstArrayView<float> PoseVector, float
 		}
 	}
 
-	return GetRootTransform(Role).GetRotation();
+	return GetRootTransform(Role, SampleTimeOffset).GetRotation();
 }
 
-FTransform FDebugDrawParams::GetRootTransform(const FRole& Role) const
+FTransform FDebugDrawParams::GetRootTransform(const FRole& Role, float SampleTimeOffset) const
 {
 	if (const int32* RoleIndex = RoleToIndex.Find(Role))
 	{
-		// @todo: add and use SampleTimeOffset to get the root transform at a specific time now using the trajectory
-		return PoseHistories[*RoleIndex]->GetTrajectory().GetSampleAtTime(0.f).GetTransform();
+		return PoseHistories[*RoleIndex]->GetTrajectory().GetSampleAtTime(SampleTimeOffset).GetTransform();
 	}
 	return FTransform::Identity;
 }
@@ -302,7 +301,6 @@ void FDebugDrawParams::DrawCircle(const FMatrix& TransformMatrix, float Radius,
 		else if (!Meshes.IsEmpty())
 		{
 			// any Mesh is fine to draw
-			// @todo: use the DrawDebugCircle API with the up vector to communize with the AnimInstanceProxy call
 			DrawDebugCircle(Meshes[0]->GetWorld(), TransformMatrix, Radius, Segments, Color, false, 0.f, SDPG_Foreground, Thickness);
 		}
 	}
diff --git a/Engine/Plugins/Experimental/Animation/PoseSearch/Source/Runtime/Private/PoseSearchDatabase.cpp b/Engine/Plugins/Experimental/Animation/PoseSearch/Source/Runtime/Private/PoseSearchDatabase.cpp
--- a/Engine/Plugins/Experimental/Animation/PoseSearch/Source/Runtime/Private/PoseSearchDatabase.cpp
+++ b/Engine/Plugins/Experimental/Animation/PoseSearch/Source/Runtime/Private/PoseSearchDatabase.cpp
@@ -404,8 +404,8 @@ UObject* FPoseSearchDatabaseBlendSpace::GetAnimationAsset() const
 #if WITH_EDITOR
 int32 FPoseSearchDatabaseBlendSpace::GetFrameAtTime(float Time) const
 {
-	// @todo: implement GetFrameAtTime logic for BlendSpace
-	return 0;
+	// returning the percentage of time as value to diplay in the pose search debugger (NoTe: BlendSpace->GetPlayLength() is one)
+	return FMath::RoundToInt(Time * 100.f);
 }
 #endif // WITH_EDITOR
 
diff --git a/Engine/Plugins/Experimental/Animation/PoseSearch/Source/Runtime/Private/PoseSearchDerivedData.cpp b/Engine/Plugins/Experimental/Animation/PoseSearch/Source/Runtime/Private/PoseSearchDerivedData.cpp
--- a/Engine/Plugins/Experimental/Animation/PoseSearch/Source/Runtime/Private/PoseSearchDerivedData.cpp
+++ b/Engine/Plugins/Experimental/Animation/PoseSearch/Source/Runtime/Private/PoseSearchDerivedData.cpp
@@ -811,7 +811,7 @@ static void PreprocessSearchIndexVPTree(FSearchIndex& SearchIndex, const UPoseSe
 	}
 }
 
-// @todo: this struct could be replaced by TTuple<const UAnimationAsset*, FTransform, FVector> if FTransform implements operator==
+// this struct exists because FTransform doesn't implement operator== (or it'd be TTuple<const UAnimationAsset*, FTransform, FVector>)
 struct FSamplerMapKey
 {
 	FSamplerMapKey(const UAnimationAsset* InAnimationAsset, const FTransform& InRootTransformOrigin, const FVector& InBlendParameters = FVector::ZeroVector)
@@ -960,7 +960,6 @@ static bool IndexDatabase(FSearchIndexBase& SearchIndexBase, const UPoseSearchDa
 			}
 			else
 			{
-				// @todo: HANDLE ERRORS / WARNINGS....
 				UE_LOG(LogPoseSearch, Warning, TEXT("Role '%s' from asset '%s' contained in database '%s' cannot be found in Schema::Skeletons '%s'"), *Role.ToString(), *DatabaseAnimationAssetBase->GetName(), *Database.GetName(), *Database.Schema->GetName());
 				return false;
 			}
@@ -1427,7 +1426,7 @@ void FPoseSearchDatabaseAsyncCacheTask::Wait(FCriticalSection& OuterMutex)
 	const bool bFailedIndexing = SearchIndex.IsEmpty();
 	if (!bFailedIndexing)
 	{
-		Database->SetSearchIndex(SearchIndex); // @todo: implement FSearchIndex move ctor and assignment operator and use a MoveTemp(SearchIndex) here
+		Database->SetSearchIndex(SearchIndex);
 
 		check(Database->Schema && !SearchIndex.IsEmpty() && SearchIndex.GetNumDimensions() == Database->Schema->SchemaCardinality);
 
diff --git a/Engine/Plugins/Experimental/Animation/PoseSearch/Source/Runtime/Private/PoseSearchFeatureChannel_FilterCrashingLegs.h b/Engine/Plugins/Experimental/Animation/PoseSearch/Source/Runtime/Private/PoseSearchFeatureChannel_FilterCrashingLegs.h
--- a/Engine/Plugins/Experimental/Animation/PoseSearch/Source/Runtime/Private/PoseSearchFeatureChannel_FilterCrashingLegs.h
+++ b/Engine/Plugins/Experimental/Animation/PoseSearch/Source/Runtime/Private/PoseSearchFeatureChannel_FilterCrashingLegs.h
@@ -13,7 +13,7 @@
 // -0.5 if the left foot is exactly in front of the right foot
 // close to 1 or -1 if the feet (and so the legs) are completely crossed
 // at runtime we'll match the CrashingLegsValue and also filter by discarding pose candidates that don't respect the 'AllowedTolerance' between query and database values (happening in IsFilterValid)
-UCLASS(BlueprintType, EditInlineNew, meta = (DisplayName = "CrashingLegs Channel"), CollapseCategories)
+UCLASS(Experimental, BlueprintType, EditInlineNew, meta = (DisplayName = "CrashingLegs Channel"), CollapseCategories)
 class POSESEARCH_API UPoseSearchFeatureChannel_FilterCrashingLegs : public UPoseSearchFeatureChannel
 {
 	GENERATED_BODY()
diff --git a/Engine/Plugins/Experimental/Animation/PoseSearch/Source/Runtime/Private/PoseSearchFeatureChannel_Heading.cpp b/Engine/Plugins/Experimental/Animation/PoseSearch/Source/Runtime/Private/PoseSearchFeatureChannel_Heading.cpp
--- a/Engine/Plugins/Experimental/Animation/PoseSearch/Source/Runtime/Private/PoseSearchFeatureChannel_Heading.cpp
+++ b/Engine/Plugins/Experimental/Animation/PoseSearch/Source/Runtime/Private/PoseSearchFeatureChannel_Heading.cpp
@@ -45,7 +45,6 @@ void UPoseSearchFeatureChannel_Heading::FindOrAddToSchema(UPoseSearchSchema* Sch
 		Heading->OriginRole = Role;
 #if WITH_EDITORONLY_DATA
 		Heading->Weight = 0.f;
-		// @todo: perhaps add a tunable color for injected channels
 		Heading->DebugColor = FLinearColor::Gray;
 #endif // WITH_EDITORONLY_DATA
 		Heading->SampleTimeOffset = SampleTimeOffset;
diff --git a/Engine/Plugins/Experimental/Animation/PoseSearch/Source/Runtime/Private/PoseSearchFeatureChannel_PermutationTime.h b/Engine/Plugins/Experimental/Animation/PoseSearch/Source/Runtime/Private/PoseSearchFeatureChannel_PermutationTime.h
--- a/Engine/Plugins/Experimental/Animation/PoseSearch/Source/Runtime/Private/PoseSearchFeatureChannel_PermutationTime.h
+++ b/Engine/Plugins/Experimental/Animation/PoseSearch/Source/Runtime/Private/PoseSearchFeatureChannel_PermutationTime.h
@@ -5,7 +5,7 @@
 #include "PoseSearch/PoseSearchFeatureChannel.h"
 #include "PoseSearchFeatureChannel_PermutationTime.generated.h"
 
-UCLASS(EditInlineNew, meta = (DisplayName = "Permutation Time Channel"), CollapseCategories)
+UCLASS(Experimental, EditInlineNew, meta = (DisplayName = "Permutation Time Channel"), CollapseCategories)
 class POSESEARCH_API UPoseSearchFeatureChannel_PermutationTime : public UPoseSearchFeatureChannel
 {
 	GENERATED_BODY()
diff --git a/Engine/Plugins/Experimental/Animation/PoseSearch/Source/Runtime/Private/PoseSearchFeatureChannel_Phase.h b/Engine/Plugins/Experimental/Animation/PoseSearch/Source/Runtime/Private/PoseSearchFeatureChannel_Phase.h
--- a/Engine/Plugins/Experimental/Animation/PoseSearch/Source/Runtime/Private/PoseSearchFeatureChannel_Phase.h
+++ b/Engine/Plugins/Experimental/Animation/PoseSearch/Source/Runtime/Private/PoseSearchFeatureChannel_Phase.h
@@ -6,7 +6,7 @@
 #include "BoneContainer.h"
 #include "PoseSearchFeatureChannel_Phase.generated.h"
 
-UCLASS(BlueprintType, EditInlineNew, meta = (DisplayName = "Phase Channel"), CollapseCategories)
+UCLASS(Experimental, BlueprintType, EditInlineNew, meta = (DisplayName = "Phase Channel"), CollapseCategories)
 class POSESEARCH_API UPoseSearchFeatureChannel_Phase : public UPoseSearchFeatureChannel
 {
 	GENERATED_BODY()
diff --git a/Engine/Plugins/Experimental/Animation/PoseSearch/Source/Runtime/Private/PoseSearchFeatureChannel_Position.cpp b/Engine/Plugins/Experimental/Animation/PoseSearch/Source/Runtime/Private/PoseSearchFeatureChannel_Position.cpp
--- a/Engine/Plugins/Experimental/Animation/PoseSearch/Source/Runtime/Private/PoseSearchFeatureChannel_Position.cpp
+++ b/Engine/Plugins/Experimental/Animation/PoseSearch/Source/Runtime/Private/PoseSearchFeatureChannel_Position.cpp
@@ -43,7 +43,6 @@ void UPoseSearchFeatureChannel_Position::FindOrAddToSchema(UPoseSearchSchema* Sc
 		Position->OriginRole = Role;
 #if WITH_EDITORONLY_DATA
 		Position->Weight = 0.f;
-		// @todo: perhaps add a tunable color for injected channels
 		Position->DebugColor = FLinearColor::Gray;
 #endif // WITH_EDITORONLY_DATA
 		Position->SampleTimeOffset = SampleTimeOffset;
diff --git a/Engine/Plugins/Experimental/Animation/PoseSearch/Source/Runtime/Private/PoseSearchHistory.cpp b/Engine/Plugins/Experimental/Animation/PoseSearch/Source/Runtime/Private/PoseSearchHistory.cpp
--- a/Engine/Plugins/Experimental/Animation/PoseSearch/Source/Runtime/Private/PoseSearchHistory.cpp
+++ b/Engine/Plugins/Experimental/Animation/PoseSearch/Source/Runtime/Private/PoseSearchHistory.cpp
@@ -144,7 +144,6 @@ static bool LerpEntries(float Time, bool bExtrapolate, const FPoseHistoryEntry&
 	}
 	else
 	{
-		// @todo: implement if required
 		OutBoneTransform = FTransform::Identity;
 		bSuccess = false;
 		unimplemented();
diff --git a/Engine/Plugins/Experimental/Animation/PoseSearch/Source/Runtime/Private/PoseSearchIndex.cpp b/Engine/Plugins/Experimental/Animation/PoseSearch/Source/Runtime/Private/PoseSearchIndex.cpp
--- a/Engine/Plugins/Experimental/Animation/PoseSearch/Source/Runtime/Private/PoseSearchIndex.cpp
+++ b/Engine/Plugins/Experimental/Animation/PoseSearch/Source/Runtime/Private/PoseSearchIndex.cpp
@@ -100,120 +100,37 @@ static bool CalculateSimilarities(TArray<FPosePairSimilarity>& PosePairSimilarit
 	int32 DataCardinality, int32 NumPoses, const TAlignedArray<float>& Values,
 	TFunctionRef<TConstArrayView<float>(int32, int32)> GetValuesVector)
 {
-	enum EEvalMode
-	{
-		EUseKDTreeEvaluation,
-		EParalleEvaluation,
-		ESerialEvaluation,
-	};
-
-	static EEvalMode EvalMode = EEvalMode::EUseKDTreeEvaluation;
-
 	PosePairSimilarities.Reserve(1024 * 64);
 
-	if (EvalMode == EEvalMode::EUseKDTreeEvaluation)
-	{
-		check(Values.Num() == NumPoses * DataCardinality);
-		FKDTree KDTree(NumPoses, DataCardinality, Values.GetData());
+	check(Values.Num() == NumPoses * DataCardinality);
+	FKDTree KDTree(NumPoses, DataCardinality, Values.GetData());
 
-		TArray<int32> ResultIndexes;
-		TArray<float> ResultDistanceSqr;
-		ResultIndexes.SetNum(NumPoses + 1);
-		ResultDistanceSqr.SetNum(NumPoses + 1);
+	TArray<int32> ResultIndexes;
+	TArray<float> ResultDistanceSqr;
+	ResultIndexes.SetNum(NumPoses + 1);
+	ResultDistanceSqr.SetNum(NumPoses + 1);
 
-		for (int32 PoseIdx = 0; PoseIdx < NumPoses; ++PoseIdx)
-		{
-			TConstArrayView<float> ValuesA = GetValuesVector(PoseIdx, DataCardinality);
-
-			// searching for duplicates within a radius of SimilarityThreshold
-			FKDTree::FRadiusResultSet ResultSet(SimilarityThreshold, NumPoses, ResultIndexes, ResultDistanceSqr);
-			KDTree.FindNeighbors(ResultSet, ValuesA);
-
-			for (int32 ResultIndex = 0; ResultIndex < ResultSet.Num(); ++ResultIndex)
-			{
-				const int32 ResultPoseIdx = ResultIndexes[ResultIndex];
-				if (PoseIdx != ResultPoseIdx)
-				{
-					FPosePairSimilarity PosePair;
-					PosePair.PoseIdxA = PoseIdx;
-					PosePair.PoseIdxB = ResultPoseIdx;
-					PosePair.Similarity = ResultDistanceSqr[ResultIndex];
-					PosePairSimilarities.Emplace(PosePair);
-				}
-			}
-		}
-	}
-	// calculating a sparse proximity matrix with Similarity up to SimilarityThreshold (to perform something similar to hierarchical clustering)
-	else if (EvalMode == EEvalMode::EParalleEvaluation)
+	for (int32 PoseIdx = 0; PoseIdx < NumPoses; ++PoseIdx)
 	{
-		// @todo: we should use a smarter approach, knowing we need NumPoses * (NumPoses - 1) / 2 comparisons, and reconstruct the pose indexes from the comparison index...
-
-		// doing it in batches of items to avoid using too much memory
-		constexpr int32 BatchSize = 1024 * 1024;
-
-		TArray<FPosePair> PosePairs;
-		PosePairs.Reserve(BatchSize);
-		FCriticalSection Mutex;
-
-		auto EvaluatePosePairs = [GetValuesVector, DataCardinality, SimilarityThreshold, &PosePairs, &PosePairSimilarities, &Mutex]()
-		{
-			ParallelFor(PosePairs.Num(), [GetValuesVector, DataCardinality, SimilarityThreshold, &PosePairs, &PosePairSimilarities, &Mutex](int32 ComparisonIndex)
-			{
-				const int32 PoseIdxA = PosePairs[ComparisonIndex].PoseIdxA;
-				const int32 PoseIdxB = PosePairs[ComparisonIndex].PoseIdxB;
-				TConstArrayView<float> ValuesA = GetValuesVector(PoseIdxA, DataCardinality);
-				TConstArrayView<float> ValuesB = GetValuesVector(PoseIdxB, DataCardinality);
-				const float Similarity = CompareFeatureVectors(ValuesA, ValuesB);
-				if (Similarity < SimilarityThreshold)
-				{
-					// since this condition doesn't happen often, we're not gonna spend too much time in mutex contention
-					FScopeLock Lock(&Mutex);
-					FPosePairSimilarity PosePairSimilarity;
-					PosePairSimilarity.PoseIdxA = PoseIdxA;
-					PosePairSimilarity.PoseIdxB = PoseIdxB;
-					PosePairSimilarity.Similarity = Similarity;
-					PosePairSimilarities.Emplace(PosePairSimilarity);
-				}
-			});
-		};
+		TConstArrayView<float> ValuesA = GetValuesVector(PoseIdx, DataCardinality);
 
-		FPosePair PosePair;
-		for (PosePair.PoseIdxA = 0; PosePair.PoseIdxA < NumPoses; ++PosePair.PoseIdxA)
-		{
-			for (PosePair.PoseIdxB = PosePair.PoseIdxA + 1; PosePair.PoseIdxB < NumPoses; ++PosePair.PoseIdxB)
-			{
-				PosePairs.Add(PosePair);
-				if (PosePairs.Num() == BatchSize)
-				{
-					EvaluatePosePairs();
-					PosePairs.Reset();
-				}
-			}
-		}
+		// searching for duplicates within a radius of SimilarityThreshold
+		FKDTree::FRadiusResultSet ResultSet(SimilarityThreshold, NumPoses, ResultIndexes, ResultDistanceSqr);
+		KDTree.FindNeighbors(ResultSet, ValuesA);
 
-		EvaluatePosePairs();
-	}
-	else if (EvalMode == EEvalMode::ESerialEvaluation)
-	{
-		FPosePairSimilarity PosePair;
-		for (PosePair.PoseIdxA = 0; PosePair.PoseIdxA < NumPoses; ++PosePair.PoseIdxA)
+		for (int32 ResultIndex = 0; ResultIndex < ResultSet.Num(); ++ResultIndex)
 		{
-			TConstArrayView<float> ValuesA = GetValuesVector(PosePair.PoseIdxA, DataCardinality);
-			for (PosePair.PoseIdxB = PosePair.PoseIdxA + 1; PosePair.PoseIdxB < NumPoses; ++PosePair.PoseIdxB)
+			const int32 ResultPoseIdx = ResultIndexes[ResultIndex];
+			if (PoseIdx != ResultPoseIdx)
 			{
-				TConstArrayView<float> ValuesB = GetValuesVector(PosePair.PoseIdxB, DataCardinality);
-				PosePair.Similarity = CompareFeatureVectors(ValuesA, ValuesB);
-				if (PosePair.Similarity < SimilarityThreshold)
-				{
-					PosePairSimilarities.Emplace(PosePair);
-				}
+				FPosePairSimilarity PosePair;
+				PosePair.PoseIdxA = PoseIdx;
+				PosePair.PoseIdxB = ResultPoseIdx;
+				PosePair.Similarity = ResultDistanceSqr[ResultIndex];
+				PosePairSimilarities.Emplace(PosePair);
 			}
 		}
 	}
-	else
-	{
-		checkNoEntry();
-	}
 
 	if (!PosePairSimilarities.IsEmpty())
 	{
@@ -563,7 +480,7 @@ TConstArrayView<float> FSearchIndex::GetReconstructedPoseValues(int32 PoseIdx, T
 
 	const int32 NumberOfPrincipalComponents = PCAValues.Num() / NumPoses;
 	
-	// @todo: if one of these checks trigger, most likely PCAValuesPruningSimilarityThreshold > 0.f and we pruned some PCAValues.
+	// NoTe: if one of these checks trigger, most likely PCAValuesPruningSimilarityThreshold > 0.f and we pruned some PCAValues.
 	// currently GetReconstructedPoseValues is not supported with PCAValues pruning
 	check(NumPoses * NumberOfPrincipalComponents == PCAValues.Num());
 	check(PCAProjectionMatrix.Num() == NumDimensions * NumberOfPrincipalComponents);
diff --git a/Engine/Plugins/Experimental/Animation/PoseSearch/Source/Runtime/Private/PoseSearchLibrary.cpp b/Engine/Plugins/Experimental/Animation/PoseSearch/Source/Runtime/Private/PoseSearchLibrary.cpp
--- a/Engine/Plugins/Experimental/Animation/PoseSearch/Source/Runtime/Private/PoseSearchLibrary.cpp
+++ b/Engine/Plugins/Experimental/Animation/PoseSearch/Source/Runtime/Private/PoseSearchLibrary.cpp
@@ -339,15 +339,9 @@ void UPoseSearchLibrary::TraceMotionMatchingState(
 	TraceState.SearchBruteForceCost = CurrentResult.BruteForcePoseCost.GetTotalCost();
 	TraceState.SearchBestPosePos = CurrentResult.BestPosePos;
 
-	// @todo: clean up this code
-	const UAnimInstance* AnimInstance = SearchContext.GetAnimInstances()[0];
-	UObject* SkeletalMeshComponent = AnimInstance->GetOuter();
-
 	TraceState.Cycle = FPlatformTime::Cycles64();
-	TraceState.AnimInstanceId = FObjectTrace::GetObjectId(AnimInstance);
-	TraceState.SkeletalMeshComponentId = FObjectTrace::GetObjectId(SkeletalMeshComponent);
+	TraceState.AnimInstanceId = FObjectTrace::GetObjectId(SearchContext.GetAnimInstances()[0]);
 	TraceState.NodeId = NodeId;
-	TraceState.FrameCounter = FObjectTrace::GetObjectWorldTickCounter(AnimInstance);
 
 	TraceState.Output();
 }
@@ -749,7 +743,7 @@ void UPoseSearchLibrary::MotionMatch(
 		FCompactPose Pose;
 		Pose.SetBoneContainer(&AnimInstances[0]->GetRequiredBonesOnAnyThread());
 		
-		// @todo... add input BlendParameters to support sampling FutureAnimation blendspaces and support for multi character
+		// @todo: add input BlendParameters to support sampling FutureAnimation blendspaces and support for multi character
 		const FAnimationAssetSampler Sampler(FutureAnimation);
 		for (int32 i = 0; i < 2; ++i)
 		{
@@ -982,7 +976,6 @@ UE::PoseSearch::FSearchResult UPoseSearchLibrary::MotionMatch(const FAnimationBa
 
 		if (bDrawMatch || bDrawquery)
 		{
-			// @todo: populate the FDebugDrawParams properly
 			FAnimInstanceProxy* AnimInstanceProxy = Context.AnimInstanceProxy;
 			const TArrayView<FAnimInstanceProxy*> AnimInstanceProxies = MakeArrayView(&AnimInstanceProxy, 1);
 
diff --git a/Engine/Plugins/Experimental/Animation/PoseSearch/Source/Runtime/Private/Trace/PoseSearchTraceLogger.cpp b/Engine/Plugins/Experimental/Animation/PoseSearch/Source/Runtime/Private/Trace/PoseSearchTraceLogger.cpp
--- a/Engine/Plugins/Experimental/Animation/PoseSearch/Source/Runtime/Private/Trace/PoseSearchTraceLogger.cpp
+++ b/Engine/Plugins/Experimental/Animation/PoseSearch/Source/Runtime/Private/Trace/PoseSearchTraceLogger.cpp
@@ -46,9 +46,7 @@ FArchive& operator<<(FArchive& Ar, FTraceMessage& State)
 {
 	Ar << State.Cycle;
 	Ar << State.AnimInstanceId;
-	Ar << State.SkeletalMeshComponentId;
 	Ar << State.NodeId;
-	Ar << State.FrameCounter;
 	return Ar;
 }
 
diff --git a/Engine/Plugins/Experimental/Animation/PoseSearch/Source/Runtime/Public/PoseSearch/PoseSearchContext.h b/Engine/Plugins/Experimental/Animation/PoseSearch/Source/Runtime/Public/PoseSearch/PoseSearchContext.h
--- a/Engine/Plugins/Experimental/Animation/PoseSearch/Source/Runtime/Public/PoseSearch/PoseSearchContext.h
+++ b/Engine/Plugins/Experimental/Animation/PoseSearch/Source/Runtime/Public/PoseSearch/PoseSearchContext.h
@@ -63,7 +63,7 @@ struct POSESEARCH_API FDebugDrawParams
 
 	FQuat ExtractRotation(TConstArrayView<float> PoseVector, float SampleTimeOffset, int8 SchemaBoneIdx, const FRole& Role, EPermutationTimeType PermutationTimeType = EPermutationTimeType::UseSampleTime, int32 SamplingAttributeId = INDEX_NONE) const;
 
-	FTransform GetRootTransform(const FRole& Role) const;
+	FTransform GetRootTransform(const FRole& Role, float SampleTimeOffset = 0.f) const;
 
 	void DrawLine(const FVector& LineStart, const FVector& LineEnd, const FColor& Color, float Thickness = 0.f) const;
 	void DrawPoint(const FVector& Position, const FColor& Color, float Thickness = 6.f) const;
@@ -212,7 +212,6 @@ private:
 
 	TConstArrayView<float> CurrentResultPoseVector;
 
-	// @todo: use a 16 bytes aligned TInlineAllocator with overflow TMemStackAllocator
 	TStackAlignedArray<float> CurrentResultPoseVectorData;
 
 	// transforms cached in world space
diff --git a/Engine/Plugins/Experimental/Animation/PoseSearch/Source/Runtime/Public/PoseSearch/PoseSearchFeatureChannel.h b/Engine/Plugins/Experimental/Animation/PoseSearch/Source/Runtime/Public/PoseSearch/PoseSearchFeatureChannel.h
--- a/Engine/Plugins/Experimental/Animation/PoseSearch/Source/Runtime/Public/PoseSearch/PoseSearchFeatureChannel.h
+++ b/Engine/Plugins/Experimental/Animation/PoseSearch/Source/Runtime/Public/PoseSearch/PoseSearchFeatureChannel.h
@@ -130,7 +130,6 @@ public:
 	virtual TArrayView<TObjectPtr<UPoseSearchFeatureChannel>> GetSubChannels() { return TArrayView<TObjectPtr<UPoseSearchFeatureChannel>>(); }
 	virtual TConstArrayView<TObjectPtr<UPoseSearchFeatureChannel>> GetSubChannels() const { return TConstArrayView<TObjectPtr<UPoseSearchFeatureChannel>>(); }
 
-	// @todo: should this API be under ENABLE_DRAW_DEBUG?
 	virtual void AddDependentChannels(UPoseSearchSchema* Schema) const {}
 
 	virtual EPermutationTimeType GetPermutationTimeType() const { return EPermutationTimeType::UseSampleTime; }
diff --git a/Engine/Plugins/Experimental/Animation/PoseSearch/Source/Runtime/Public/PoseSearch/PoseSearchIndex.h b/Engine/Plugins/Experimental/Animation/PoseSearch/Source/Runtime/Public/PoseSearch/PoseSearchIndex.h
--- a/Engine/Plugins/Experimental/Animation/PoseSearch/Source/Runtime/Public/PoseSearch/PoseSearchIndex.h
+++ b/Engine/Plugins/Experimental/Animation/PoseSearch/Source/Runtime/Public/PoseSearch/PoseSearchIndex.h
@@ -238,7 +238,6 @@ struct FSparsePoseMultiMap
 	, MaxValue(InMaxValue)
 	, DeltaKeyValue(InMaxValue >= InMaxKey ? InMaxValue - InMaxKey + 1 : 0)
 	{
-		// @todo: maybe expose this initial allocation budget
 		DataValues.Reserve(InMaxKey * 2);
 		for (Type Index = 0; Index < InMaxKey; ++Index)
 		{
diff --git a/Engine/Plugins/Experimental/Animation/PoseSearch/Source/Runtime/Public/PoseSearch/Trace/PoseSearchTraceLogger.h b/Engine/Plugins/Experimental/Animation/PoseSearch/Source/Runtime/Public/PoseSearch/Trace/PoseSearchTraceLogger.h
--- a/Engine/Plugins/Experimental/Animation/PoseSearch/Source/Runtime/Public/PoseSearch/Trace/PoseSearchTraceLogger.h
+++ b/Engine/Plugins/Experimental/Animation/PoseSearch/Source/Runtime/Public/PoseSearch/Trace/PoseSearchTraceLogger.h
@@ -27,15 +27,8 @@ struct POSESEARCH_API FTraceMessage
 
 	uint64 AnimInstanceId = 0;
 
-	/** Skeletal Mesh Component Id, outer class of the AnimInstance.
-	 *  Used for retrieval of traced root transforms from the animation provider.
-	 */
-	uint64 SkeletalMeshComponentId = 0;
-
 	/** Node Id of the motion matching node associated with this message */
 	int32 NodeId = 0;
-
-	uint16 FrameCounter = 0;
 };
 POSESEARCH_API FArchive& operator<<(FArchive& Ar, FTraceMessage& State);
 
@@ -51,8 +44,6 @@ POSESEARCH_API FArchive& operator<<(FArchive& Ar, FTraceMotionMatchingStatePoseE
 
 struct POSESEARCH_API FTraceMotionMatchingStateDatabaseEntry
 {
-	// @todo: can we use UPoseSearchDatabase* instead of DatabaseId?
-	//UPoseSearchDatabase const* Database = nullptr;
 	uint64 DatabaseId = 0;
 	TArray<float> QueryVector;
 	TArray<FTraceMotionMatchingStatePoseEntry> PoseEntries;
diff --git a/Engine/Plugins/Experimental/Animation/PoseSearch/Source/Runtime/Public/PoseSearch/VPTree.h b/Engine/Plugins/Experimental/Animation/PoseSearch/Source/Runtime/Public/PoseSearch/VPTree.h
--- a/Engine/Plugins/Experimental/Animation/PoseSearch/Source/Runtime/Public/PoseSearch/VPTree.h
+++ b/Engine/Plugins/Experimental/Animation/PoseSearch/Source/Runtime/Public/PoseSearch/VPTree.h
@@ -201,7 +201,6 @@ struct FVPTree
     void FindNeighbors(const T& Query, TVPTreeResultSet& ResultSet, const TDataSource& DataSource) const
     {
 #if VALIDATE_FINDNEIGHBORS
-		// @todo: WIP validation code
 		int32 NumEvaluatedQueryDistances = 0;
 
 		// brute forcing to find the best Node
