motion matching - fix for UPoseSearchLibrary::MotionMatch WantedPlayRate calculation

diff --git a/Engine/Plugins/Experimental/Animation/PoseSearch/Source/Runtime/Private/PoseSearchFeatureChannel_PermutationTime.cpp b/Engine/Plugins/Experimental/Animation/PoseSearch/Source/Runtime/Private/PoseSearchFeatureChannel_PermutationTime.cpp
--- a/Engine/Plugins/Experimental/Animation/PoseSearch/Source/Runtime/Private/PoseSearchFeatureChannel_PermutationTime.cpp
+++ b/Engine/Plugins/Experimental/Animation/PoseSearch/Source/Runtime/Private/PoseSearchFeatureChannel_PermutationTime.cpp
@@ -33,6 +33,12 @@ void UPoseSearchFeatureChannel_PermutationTime::BuildQuery(UE::PoseSearch::FSear
 	FFeatureVectorHelper::EncodeFloat(SearchContext.EditFeatureVector(), ChannelDataOffset, SearchContext.GetDesiredPermutationTimeOffset());
 }
 
+float UPoseSearchFeatureChannel_PermutationTime::GetPermutationTime(TConstArrayView<float> PoseVector) const
+{
+	using namespace UE::PoseSearch;
+	return FFeatureVectorHelper::DecodeFloat(PoseVector, ChannelDataOffset);
+}
+
 #if WITH_EDITOR
 void UPoseSearchFeatureChannel_PermutationTime::FillWeights(TArrayView<float> Weights) const
 {
diff --git a/Engine/Plugins/Experimental/Animation/PoseSearch/Source/Runtime/Private/PoseSearchFeatureChannel_PermutationTime.h b/Engine/Plugins/Experimental/Animation/PoseSearch/Source/Runtime/Private/PoseSearchFeatureChannel_PermutationTime.h
--- a/Engine/Plugins/Experimental/Animation/PoseSearch/Source/Runtime/Private/PoseSearchFeatureChannel_PermutationTime.h
+++ b/Engine/Plugins/Experimental/Animation/PoseSearch/Source/Runtime/Private/PoseSearchFeatureChannel_PermutationTime.h
@@ -26,5 +26,7 @@ public:
 	virtual UE::PoseSearch::TLabelBuilder& GetLabel(UE::PoseSearch::TLabelBuilder& LabelBuilder, UE::PoseSearch::ELabelFormat LabelFormat = UE::PoseSearch::ELabelFormat::Full_Horizontal) const override;
 #endif
 
+	float GetPermutationTime(TConstArrayView<float> PoseVector) const;
+
 	static void FindOrAddToSchema(UPoseSearchSchema* Schema);
 };
diff --git a/Engine/Plugins/Experimental/Animation/PoseSearch/Source/Runtime/Private/PoseSearchLibrary.cpp b/Engine/Plugins/Experimental/Animation/PoseSearch/Source/Runtime/Private/PoseSearchLibrary.cpp
--- a/Engine/Plugins/Experimental/Animation/PoseSearch/Source/Runtime/Private/PoseSearchLibrary.cpp
+++ b/Engine/Plugins/Experimental/Animation/PoseSearch/Source/Runtime/Private/PoseSearchLibrary.cpp
@@ -22,6 +22,7 @@
 #include "PoseSearch/PoseSearchSchema.h"
 #include "PoseSearchFeatureChannel_Trajectory.h"
 #include "PoseSearch/Trace/PoseSearchTraceLogger.h"
+#include "PoseSearchFeatureChannel_PermutationTime.h"
 
 #include UE_INLINE_GENERATED_CPP_BY_NAME(PoseSearchLibrary)
 
@@ -148,8 +149,12 @@ FVector FMotionMatchingState::GetEstimatedFutureRootMotionVelocity() const
 	{
 		if (const UPoseSearchFeatureChannel_Trajectory* TrajectoryChannel = CurrentSearchResult.Database->Schema->FindFirstChannelOfType<UPoseSearchFeatureChannel_Trajectory>())
 		{
-			TConstArrayView<float> ResultData = CurrentSearchResult.Database->GetSearchIndex().GetPoseValues(CurrentSearchResult.PoseIdx);
-			return TrajectoryChannel->GetEstimatedFutureRootMotionVelocity(ResultData);
+			const FSearchIndex& SearchIndex = CurrentSearchResult.Database->GetSearchIndex();
+			if (!SearchIndex.IsValuesEmpty())
+			{
+				TConstArrayView<float> ResultData = SearchIndex.GetPoseValues(CurrentSearchResult.PoseIdx);
+				return TrajectoryChannel->GetEstimatedFutureRootMotionVelocity(ResultData);
+			}
 		}
 	}
 
@@ -572,10 +577,8 @@ void UPoseSearchLibrary::MotionMatch(
 	UAnimInstance* AnimInstance,
 	const UPoseSearchDatabase* Database,
 	const FName PoseHistoryName,
+	FPoseSearchFutureProperties Future,
 	FPoseSearchBlueprintResult& Result,
-	const UAnimationAsset* FutureAnimation,
-	float FutureAnimationStartTime,
-	float TimeToFutureAnimationStart,
 	const int32 DebugSessionUniqueIdentifier)
 {
 	using namespace UE::PoseSearch;
@@ -588,7 +591,7 @@ void UPoseSearchLibrary::MotionMatch(
 	TArray<FName, TInlineAllocator<PreallocatedRolesNum, TMemStackAllocator<>>> Roles;
 	Roles.Add(UE::PoseSearch::DefaultRole);
 
-	MotionMatch(AnimInstances, Roles, Database, PoseHistoryName, Result, FutureAnimation, FutureAnimationStartTime, TimeToFutureAnimationStart, DebugSessionUniqueIdentifier);
+	MotionMatch(AnimInstances, Roles, Database, PoseHistoryName, Result, Future.FutureAnimation, Future.FutureAnimationStartTime, Future.TimeToFutureAnimationStart, DebugSessionUniqueIdentifier);
 }
 
 void UPoseSearchLibrary::MotionMatchMulti(
@@ -772,8 +775,10 @@ void UPoseSearchLibrary::MotionMatch(
 	FSearchResult SearchResult = Database->Search(SearchContext);
 	if (SearchResult.IsValid())
 	{
+		check(SearchResult.Database == Database);
+
 		const FSearchIndexAsset* SearchIndexAsset = SearchResult.GetSearchIndexAsset();
-		if (const FPoseSearchDatabaseAnimationAssetBase* DatabaseAsset = SearchResult.Database->GetAnimationAssetBase(*SearchIndexAsset))
+		if (const FPoseSearchDatabaseAnimationAssetBase* DatabaseAsset = Database->GetAnimationAssetBase(*SearchIndexAsset))
 		{
 			Result.SelectedAnimation = DatabaseAsset->GetAnimationAsset();
 			Result.SelectedTime = SearchResult.AssetTime;
@@ -782,6 +787,22 @@ void UPoseSearchLibrary::MotionMatch(
 			Result.BlendParameters = SearchIndexAsset->GetBlendParameters();
 			Result.SelectedDatabase = Database;
 			Result.SearchCost = SearchResult.PoseCost.GetTotalCost();
+			
+			// figuring out the WantedPlayRate
+			Result.WantedPlayRate = 1.f;
+			if (FutureAnimation)
+			{
+				if (const UPoseSearchFeatureChannel_PermutationTime* PermutationTimeChannel = Database->Schema->FindFirstChannelOfType<UPoseSearchFeatureChannel_PermutationTime>())
+				{
+					const FSearchIndex& SearchIndex = Database->GetSearchIndex();
+					if (!SearchIndex.IsValuesEmpty())
+					{
+						TConstArrayView<float> ResultData = Database->GetSearchIndex().GetPoseValues(SearchResult.PoseIdx);
+						const float ActualTimeToFutureAnimationStart = PermutationTimeChannel->GetPermutationTime(ResultData);
+						Result.WantedPlayRate = ActualTimeToFutureAnimationStart / TimeToFutureAnimationStart;
+					}
+				}
+			}
 		}
 	}
 
diff --git a/Engine/Plugins/Experimental/Animation/PoseSearch/Source/Runtime/Public/PoseSearch/PoseSearchLibrary.h b/Engine/Plugins/Experimental/Animation/PoseSearch/Source/Runtime/Public/PoseSearch/PoseSearchLibrary.h
--- a/Engine/Plugins/Experimental/Animation/PoseSearch/Source/Runtime/Public/PoseSearch/PoseSearchLibrary.h
+++ b/Engine/Plugins/Experimental/Animation/PoseSearch/Source/Runtime/Public/PoseSearch/PoseSearchLibrary.h
@@ -88,6 +88,25 @@ struct FMotionMatchingState
 #endif //UE_POSE_SEARCH_TRACE_ENABLED
 };
 
+USTRUCT(Experimental, BlueprintType, Category="Animation|Pose Search")
+struct POSESEARCH_API FPoseSearchFutureProperties
+{
+	GENERATED_BODY()
+
+public:
+	// Input animation we want to match after TimeToFutureAnimationStart seconds
+	UPROPERTY(EditAnywhere, BlueprintReadWrite, Category=State)
+	TObjectPtr<UAnimationAsset> FutureAnimation;
+
+	// Input start time for the first pose of FutureAnimation
+	UPROPERTY(EditAnywhere, BlueprintReadWrite, Category=State)
+	float FutureAnimationStartTime;
+
+	// Input time in seconds before start playing FutureAnimation (from FutureAnimationStartTime seconds)
+	UPROPERTY(EditAnywhere, BlueprintReadWrite, Category=State)
+	float TimeToFutureAnimationStart;
+};
+
 UCLASS()
 class POSESEARCH_API UPoseSearchLibrary : public UBlueprintFunctionLibrary
 {
@@ -153,23 +172,19 @@ public:
 	* @param bIsMirrored					Output selected animation mirror state
 	* @param BlendParameters				Output selected animation blend space parameters (if SelectedAnimation is a blend space)
 	* @param SearchCost						Output search associated cost
-	* @param FutureAnimation				Input animation we want to match after TimeToFutureAnimationStart seconds
-	* @param FutureAnimationStartTime		Input start time for the first pose of FutureAnimation
-	* @param TimeToFutureAnimationStart		Input time in seconds before start playing FutureAnimation (from FutureAnimationStartTime seconds)
+	* @param Future							Input future properties to match (animation / start time / time offset)
 	* @param DebugSessionUniqueIdentifier	Input unique identifier used to identify TraceMotionMatchingState (rewind debugger / pose search debugger) session. Similarly the MM node uses Context.GetCurrentNodeId()
 	*/
-	UFUNCTION(BlueprintPure, Category = "Animation|Pose Search", meta = (BlueprintThreadSafe, Keywords = "PoseMatch"))
+	UFUNCTION(BlueprintPure, Category = "Animation|Pose Search|Experimental", meta = (BlueprintThreadSafe, Keywords = "PoseMatch"))
 	static void MotionMatch(
 		UAnimInstance* AnimInstance,
 		const UPoseSearchDatabase* Database,
 		const FName PoseHistoryName,
+		FPoseSearchFutureProperties Future,
 		FPoseSearchBlueprintResult& Result,
-		const UAnimationAsset* FutureAnimation = nullptr,
-		float FutureAnimationStartTime = 0.f,
-		float TimeToFutureAnimationStart = 0.f,
 		const int32 DebugSessionUniqueIdentifier = 6174);
 
-	UFUNCTION(BlueprintPure, Category = "Animation|Pose Search", meta = (BlueprintThreadSafe, Keywords = "PoseMatch"))
+	UFUNCTION(BlueprintPure, Category = "Animation|Pose Search|Experimental", meta = (BlueprintThreadSafe, Keywords = "PoseMatch"))
 	static void MotionMatchMulti(
 		TArray<ACharacter*> AnimInstances,
 		TArray<FName> Roles,
