motion matching - fix for UPoseSearchDatabase::IsCachedCookedPlatformDataLoaded hanging while cooking if a database fails to index

diff --git a/Engine/Plugins/Experimental/Animation/PoseSearch/Source/Editor/Private/PoseSearchDatabaseDataDetails.cpp b/Engine/Plugins/Experimental/Animation/PoseSearch/Source/Editor/Private/PoseSearchDatabaseDataDetails.cpp
--- a/Engine/Plugins/Experimental/Animation/PoseSearch/Source/Editor/Private/PoseSearchDatabaseDataDetails.cpp
+++ b/Engine/Plugins/Experimental/Animation/PoseSearch/Source/Editor/Private/PoseSearchDatabaseDataDetails.cpp
@@ -244,7 +244,7 @@ public:
 								if (const TSharedPtr<FDatabaseViewModel> ViewModel = EditorViewModel.Pin())
 								{
 									const UPoseSearchDatabase* PoseSearchDatabase = ViewModel->GetPoseSearchDatabase();
-									if (FAsyncPoseSearchDatabasesManagement::RequestAsyncBuildIndex(PoseSearchDatabase, ERequestAsyncBuildFlag::ContinueRequest))
+									if (EAsyncBuildIndexResult::Success == FAsyncPoseSearchDatabasesManagement::RequestAsyncBuildIndex(PoseSearchDatabase, ERequestAsyncBuildFlag::ContinueRequest))
 									{
 										for (const TArray<FDatabasePreviewActor>& PreviewActorGroup : ViewModel->GetPreviewActors())
 										{
@@ -277,7 +277,7 @@ public:
 								if (const TSharedPtr<FDatabaseViewModel> ViewModel = EditorViewModel.Pin())
 								{
 									const UPoseSearchDatabase* PoseSearchDatabase = ViewModel->GetPoseSearchDatabase();
-									if (FAsyncPoseSearchDatabasesManagement::RequestAsyncBuildIndex(PoseSearchDatabase, ERequestAsyncBuildFlag::ContinueRequest))
+									if (EAsyncBuildIndexResult::Success == FAsyncPoseSearchDatabasesManagement::RequestAsyncBuildIndex(PoseSearchDatabase, ERequestAsyncBuildFlag::ContinueRequest))
 									{
 										for (const TArray<FDatabasePreviewActor>& PreviewActorGroup : ViewModel->GetPreviewActors())
 										{
@@ -310,7 +310,7 @@ public:
 								if (const TSharedPtr<FDatabaseViewModel> ViewModel = EditorViewModel.Pin())
 								{
 									const UPoseSearchDatabase* PoseSearchDatabase = ViewModel->GetPoseSearchDatabase();
-									if (FAsyncPoseSearchDatabasesManagement::RequestAsyncBuildIndex(PoseSearchDatabase, ERequestAsyncBuildFlag::ContinueRequest))
+									if (EAsyncBuildIndexResult::Success == FAsyncPoseSearchDatabasesManagement::RequestAsyncBuildIndex(PoseSearchDatabase, ERequestAsyncBuildFlag::ContinueRequest))
 									{
 										for (const TArray<FDatabasePreviewActor>& PreviewActorGroup : ViewModel->GetPreviewActors())
 										{
@@ -348,7 +348,7 @@ public:
 						if (const TSharedPtr<FDatabaseViewModel> ViewModel = EditorViewModel.Pin())
 						{
 							const UPoseSearchDatabase* PoseSearchDatabase = ViewModel->GetPoseSearchDatabase();
-							if (FAsyncPoseSearchDatabasesManagement::RequestAsyncBuildIndex(PoseSearchDatabase, ERequestAsyncBuildFlag::ContinueRequest))
+							if (EAsyncBuildIndexResult::Success == FAsyncPoseSearchDatabasesManagement::RequestAsyncBuildIndex(PoseSearchDatabase, ERequestAsyncBuildFlag::ContinueRequest))
 							{
 								for (const TArray<FDatabasePreviewActor>& PreviewActorGroup : ViewModel->GetPreviewActors())
 								{
@@ -447,7 +447,7 @@ void SDatabaseDataDetails::Reconstruct(int32 MaxPreviewActors)
 	}
 	
 	const UPoseSearchDatabase* PoseSearchDatabase = ViewModel->GetPoseSearchDatabase();
-	if (!FAsyncPoseSearchDatabasesManagement::RequestAsyncBuildIndex(PoseSearchDatabase, ERequestAsyncBuildFlag::ContinueRequest))
+	if (EAsyncBuildIndexResult::Success != FAsyncPoseSearchDatabasesManagement::RequestAsyncBuildIndex(PoseSearchDatabase, ERequestAsyncBuildFlag::ContinueRequest))
 	{
 		ChannelItems.Reset();
 		return;
diff --git a/Engine/Plugins/Experimental/Animation/PoseSearch/Source/Editor/Private/PoseSearchDatabaseEditor.cpp b/Engine/Plugins/Experimental/Animation/PoseSearch/Source/Editor/Private/PoseSearchDatabaseEditor.cpp
--- a/Engine/Plugins/Experimental/Animation/PoseSearch/Source/Editor/Private/PoseSearchDatabaseEditor.cpp
+++ b/Engine/Plugins/Experimental/Animation/PoseSearch/Source/Editor/Private/PoseSearchDatabaseEditor.cpp
@@ -213,7 +213,7 @@ namespace UE::PoseSearch
 			if (DatabaseAsset)
 			{
 				// Init statistics
-				if (FAsyncPoseSearchDatabasesManagement::RequestAsyncBuildIndex(GetPoseSearchDatabase(), ERequestAsyncBuildFlag::ContinueRequest))
+				if (EAsyncBuildIndexResult::Success == FAsyncPoseSearchDatabasesManagement::RequestAsyncBuildIndex(GetPoseSearchDatabase(), ERequestAsyncBuildFlag::ContinueRequest))
 				{
 					RefreshStatisticsWidgetInformation();
 				}
diff --git a/Engine/Plugins/Experimental/Animation/PoseSearch/Source/Editor/Private/PoseSearchDatabasePreviewScene.cpp b/Engine/Plugins/Experimental/Animation/PoseSearch/Source/Editor/Private/PoseSearchDatabasePreviewScene.cpp
--- a/Engine/Plugins/Experimental/Animation/PoseSearch/Source/Editor/Private/PoseSearchDatabasePreviewScene.cpp
+++ b/Engine/Plugins/Experimental/Animation/PoseSearch/Source/Editor/Private/PoseSearchDatabasePreviewScene.cpp
@@ -60,7 +60,7 @@ namespace UE::PoseSearch
 		FDatabaseViewModel* ViewModel = GetEditor()->GetViewModel();
 		const UPoseSearchDatabase* Database = ViewModel->GetPoseSearchDatabase();
 
-		if (!ViewModel->GetPreviewActors().IsEmpty() && FAsyncPoseSearchDatabasesManagement::RequestAsyncBuildIndex(Database, ERequestAsyncBuildFlag::ContinueRequest))
+		if (!ViewModel->GetPreviewActors().IsEmpty() && EAsyncBuildIndexResult::Success == FAsyncPoseSearchDatabasesManagement::RequestAsyncBuildIndex(Database, ERequestAsyncBuildFlag::ContinueRequest))
 		{
 			const bool bDisplayRootMotionSpeed = ViewModel->IsDisplayRootMotionSpeedChecked();
 			const bool bDisplayBlockTransition = ViewModel->IsDisplayBlockTransitionChecked();
diff --git a/Engine/Plugins/Experimental/Animation/PoseSearch/Source/Editor/Private/PoseSearchDatabaseViewModel.cpp b/Engine/Plugins/Experimental/Animation/PoseSearch/Source/Editor/Private/PoseSearchDatabaseViewModel.cpp
--- a/Engine/Plugins/Experimental/Animation/PoseSearch/Source/Editor/Private/PoseSearchDatabaseViewModel.cpp
+++ b/Engine/Plugins/Experimental/Animation/PoseSearch/Source/Editor/Private/PoseSearchDatabaseViewModel.cpp
@@ -512,7 +512,7 @@ void FDatabaseViewModel::Tick(float DeltaSeconds)
 
 		if (const UPoseSearchDatabase* Database = GetPoseSearchDatabase())
 		{
-			if (FAsyncPoseSearchDatabasesManagement::RequestAsyncBuildIndex(Database, ERequestAsyncBuildFlag::ContinueRequest))
+			if (EAsyncBuildIndexResult::Success == FAsyncPoseSearchDatabasesManagement::RequestAsyncBuildIndex(Database, ERequestAsyncBuildFlag::ContinueRequest))
 			{
 				for (TArray<FDatabasePreviewActor>& PreviewActorGroup : PreviewActors)
 				{
@@ -699,7 +699,7 @@ int32 FDatabaseViewModel::SetSelectedNode(int32 PoseIdx, bool bClearSelection, b
 
 	if (const UPoseSearchDatabase* Database = GetPoseSearchDatabase())
 	{
-		if (FAsyncPoseSearchDatabasesManagement::RequestAsyncBuildIndex(Database, ERequestAsyncBuildFlag::ContinueRequest))
+		if (EAsyncBuildIndexResult::Success == FAsyncPoseSearchDatabasesManagement::RequestAsyncBuildIndex(Database, ERequestAsyncBuildFlag::ContinueRequest))
 		{
 			const FSearchIndex& SearchIndex = Database->GetSearchIndex();
 			if (SearchIndex.PoseMetadata.IsValidIndex(PoseIdx))
@@ -765,7 +765,7 @@ void FDatabaseViewModel::SetSelectedNodes(const TArrayView<TSharedPtr<FDatabaseA
 
 	if (const UPoseSearchDatabase* Database = GetPoseSearchDatabase())
 	{
-		if (FAsyncPoseSearchDatabasesManagement::RequestAsyncBuildIndex(Database, ERequestAsyncBuildFlag::ContinueRequest))
+		if (EAsyncBuildIndexResult::Success == FAsyncPoseSearchDatabasesManagement::RequestAsyncBuildIndex(Database, ERequestAsyncBuildFlag::ContinueRequest))
 		{
 			TMap<int32, int32> AssociatedAssetIndices;
 			for (int32 i = 0; i < InSelectedNodes.Num(); ++i)
@@ -853,7 +853,7 @@ const FSearchIndexAsset* FDatabaseViewModel::GetSelectedActorIndexAsset() const
 {
 	if (const UPoseSearchDatabase* Database = GetPoseSearchDatabase())
 	{
-		if (FAsyncPoseSearchDatabasesManagement::RequestAsyncBuildIndex(Database, ERequestAsyncBuildFlag::ContinueRequest))
+		if (EAsyncBuildIndexResult::Success == FAsyncPoseSearchDatabasesManagement::RequestAsyncBuildIndex(Database, ERequestAsyncBuildFlag::ContinueRequest))
 		{
 			const FSearchIndex& SearchIndex = Database->GetSearchIndex();
 			if (SearchIndex.Assets.IsValidIndex(SelectedActorIndexAssetIndex))
@@ -887,7 +887,7 @@ void FDatabaseViewModel::SetPlayTime(float NewPlayTime, bool bInTickPlayTime)
 		
 		if (const UPoseSearchDatabase* Database = GetPoseSearchDatabase())
 		{
-			if (FAsyncPoseSearchDatabasesManagement::RequestAsyncBuildIndex(Database, ERequestAsyncBuildFlag::ContinueRequest))
+			if (EAsyncBuildIndexResult::Success == FAsyncPoseSearchDatabasesManagement::RequestAsyncBuildIndex(Database, ERequestAsyncBuildFlag::ContinueRequest))
 			{
 				for (TArray<FDatabasePreviewActor>& PreviewActorGroup : PreviewActors)
 				{
@@ -905,7 +905,7 @@ bool FDatabaseViewModel::GetAnimationTime(int32 SourceAssetIdx, float& CurrentPl
 {
 	if (const UPoseSearchDatabase* Database = GetPoseSearchDatabase())
 	{
-		if (FAsyncPoseSearchDatabasesManagement::RequestAsyncBuildIndex(Database, ERequestAsyncBuildFlag::ContinueRequest))
+		if (EAsyncBuildIndexResult::Success == FAsyncPoseSearchDatabasesManagement::RequestAsyncBuildIndex(Database, ERequestAsyncBuildFlag::ContinueRequest))
 		{
 			const FSearchIndex& SearchIndex = Database->GetSearchIndex();
 			for (const TArray<FDatabasePreviewActor>& PreviewActorGroup : PreviewActors)
diff --git a/Engine/Plugins/Experimental/Animation/PoseSearch/Source/Editor/Private/PoseSearchDebuggerDatabaseView.cpp b/Engine/Plugins/Experimental/Animation/PoseSearch/Source/Editor/Private/PoseSearchDebuggerDatabaseView.cpp
--- a/Engine/Plugins/Experimental/Animation/PoseSearch/Source/Editor/Private/PoseSearchDebuggerDatabaseView.cpp
+++ b/Engine/Plugins/Experimental/Animation/PoseSearch/Source/Editor/Private/PoseSearchDebuggerDatabaseView.cpp
@@ -36,7 +36,7 @@ public:
 		for (const FTraceMotionMatchingStateDatabaseEntry& DbEntry : DatabaseEntries)
 		{
 			const UPoseSearchDatabase* Database = FTraceMotionMatchingStateMessage::GetObjectFromId<UPoseSearchDatabase>(DbEntry.DatabaseId);
-			if (FAsyncPoseSearchDatabasesManagement::RequestAsyncBuildIndex(Database, ERequestAsyncBuildFlag::ContinueRequest))
+			if (EAsyncBuildIndexResult::Success == FAsyncPoseSearchDatabasesManagement::RequestAsyncBuildIndex(Database, ERequestAsyncBuildFlag::ContinueRequest))
 			{
 				for (const TObjectPtr<UPoseSearchFeatureChannel>& ChannelPtr : Database->Schema->GetChannels())
 				{
@@ -257,7 +257,7 @@ void SDebuggerDatabaseView::Update(const FTraceMotionMatchingStateMessage& State
 	for (const FTraceMotionMatchingStateDatabaseEntry& DbEntry : State.DatabaseEntries)
 	{
 		const UPoseSearchDatabase* Database = FTraceMotionMatchingStateMessage::GetObjectFromId<UPoseSearchDatabase>(DbEntry.DatabaseId);
-		if (FAsyncPoseSearchDatabasesManagement::RequestAsyncBuildIndex(Database, ERequestAsyncBuildFlag::ContinueRequest))
+		if (EAsyncBuildIndexResult::Success == FAsyncPoseSearchDatabasesManagement::RequestAsyncBuildIndex(Database, ERequestAsyncBuildFlag::ContinueRequest))
 		{
 			const FSearchIndex& SearchIndex = Database->GetSearchIndex();
 			SearchIndex.GetPoseToPCAValuesVectorIndexes(PoseToPCAValuesVectorIndexes);
diff --git a/Engine/Plugins/Experimental/Animation/PoseSearch/Source/Editor/Private/PoseSearchDebuggerView.cpp b/Engine/Plugins/Experimental/Animation/PoseSearch/Source/Editor/Private/PoseSearchDebuggerView.cpp
--- a/Engine/Plugins/Experimental/Animation/PoseSearch/Source/Editor/Private/PoseSearchDebuggerView.cpp
+++ b/Engine/Plugins/Experimental/Animation/PoseSearch/Source/Editor/Private/PoseSearchDebuggerView.cpp
@@ -300,7 +300,7 @@ void SDebuggerView::Tick(const FGeometry& AllottedGeometry, const double InCurre
 				{
 					const UPoseSearchDatabase* Database = FTraceMotionMatchingStateMessage::GetObjectFromId<UPoseSearchDatabase>(DbEntry.DatabaseId);
 					if (Database && Database == CurrentDatabase &&
-						FAsyncPoseSearchDatabasesManagement::RequestAsyncBuildIndex(CurrentDatabase, ERequestAsyncBuildFlag::ContinueRequest) &&
+						EAsyncBuildIndexResult::Success == FAsyncPoseSearchDatabasesManagement::RequestAsyncBuildIndex(CurrentDatabase, ERequestAsyncBuildFlag::ContinueRequest) &&
 						DbEntry.QueryVector.Num() == Database->Schema->SchemaCardinality)
 					{
 						FDebugDrawParams DrawParams(Meshes, PoseHistories, RoleToIndex, CurrentDatabase, EDebugDrawFlags::DrawQuery);
@@ -323,7 +323,7 @@ void SDebuggerView::Tick(const FGeometry& AllottedGeometry, const double InCurre
 		{
 			const TSharedRef<FDebuggerDatabaseRowData>& Row = SelectedRows[RowIdx];
 			const UPoseSearchDatabase* RowDatabase = Row->SharedData->SourceDatabase.Get();
-			if (FAsyncPoseSearchDatabasesManagement::RequestAsyncBuildIndex(RowDatabase, ERequestAsyncBuildFlag::ContinueRequest))
+			if (EAsyncBuildIndexResult::Success == FAsyncPoseSearchDatabasesManagement::RequestAsyncBuildIndex(RowDatabase, ERequestAsyncBuildFlag::ContinueRequest))
 			{
 				FDebugDrawParams DrawParams(Meshes, PoseHistories, RoleToIndex, RowDatabase);
 				DrawParams.DrawFeatureVector(Row->PoseIdx);
@@ -339,7 +339,7 @@ void SDebuggerView::Tick(const FGeometry& AllottedGeometry, const double InCurre
 		if (!ActiveRows.IsEmpty())
 		{
 			const UPoseSearchDatabase* Database = ActiveRows[0]->SharedData->SourceDatabase.Get();
-			if (Database && FAsyncPoseSearchDatabasesManagement::RequestAsyncBuildIndex(Database, ERequestAsyncBuildFlag::ContinueRequest))
+			if (EAsyncBuildIndexResult::Success == FAsyncPoseSearchDatabasesManagement::RequestAsyncBuildIndex(Database, ERequestAsyncBuildFlag::ContinueRequest))
 			{
 				// Use the motion-matching state's pose idx, as the active row may be update-throttled at this point
 				FDebugDrawParams DrawParams(Meshes, PoseHistories, RoleToIndex, Database);
@@ -357,7 +357,7 @@ void SDebuggerView::Tick(const FGeometry& AllottedGeometry, const double InCurre
 		if (!ContinuingRows.IsEmpty())
 		{
 			const UPoseSearchDatabase* Database = ContinuingRows[0]->SharedData->SourceDatabase.Get();
-			if (Database && FAsyncPoseSearchDatabasesManagement::RequestAsyncBuildIndex(Database, ERequestAsyncBuildFlag::ContinueRequest))
+			if (EAsyncBuildIndexResult::Success == FAsyncPoseSearchDatabasesManagement::RequestAsyncBuildIndex(Database, ERequestAsyncBuildFlag::ContinueRequest))
 			{
 				FDebugDrawParams DrawParams(Meshes, PoseHistories, RoleToIndex, Database);
 				DrawParams.DrawFeatureVector(ContinuingRows[0]->PoseIdx);
diff --git a/Engine/Plugins/Experimental/Animation/PoseSearch/Source/Runtime/Private/AnimNode_MotionMatching.cpp b/Engine/Plugins/Experimental/Animation/PoseSearch/Source/Runtime/Private/AnimNode_MotionMatching.cpp
--- a/Engine/Plugins/Experimental/Animation/PoseSearch/Source/Runtime/Private/AnimNode_MotionMatching.cpp
+++ b/Engine/Plugins/Experimental/Animation/PoseSearch/Source/Runtime/Private/AnimNode_MotionMatching.cpp
@@ -105,7 +105,7 @@ void FAnimNode_MotionMatching::UpdateAssetPlayer(const FAnimationUpdateContext&
 		!UpdateCounter.WasSynchronizedCounter(Context.AnimInstanceProxy->GetUpdateCounter());
 
 #if WITH_EDITOR
-	if (!FAsyncPoseSearchDatabasesManagement::RequestAsyncBuildIndex(MotionMatchingState.CurrentSearchResult.Database.Get(), ERequestAsyncBuildFlag::ContinueRequest))
+	if (EAsyncBuildIndexResult::Success != FAsyncPoseSearchDatabasesManagement::RequestAsyncBuildIndex(MotionMatchingState.CurrentSearchResult.Database.Get(), ERequestAsyncBuildFlag::ContinueRequest))
 	{
 		bNeedsReset = true;
 	}
diff --git a/Engine/Plugins/Experimental/Animation/PoseSearch/Source/Runtime/Private/PoseSearchDatabase.cpp b/Engine/Plugins/Experimental/Animation/PoseSearch/Source/Runtime/Private/PoseSearchDatabase.cpp
--- a/Engine/Plugins/Experimental/Animation/PoseSearch/Source/Runtime/Private/PoseSearchDatabase.cpp
+++ b/Engine/Plugins/Experimental/Animation/PoseSearch/Source/Runtime/Private/PoseSearchDatabase.cpp
@@ -924,7 +924,7 @@ bool UPoseSearchDatabase::IsCachedCookedPlatformDataLoaded(const ITargetPlatform
 {
 	using namespace UE::PoseSearch;
 	check(IsInGameThread());
-	return FAsyncPoseSearchDatabasesManagement::RequestAsyncBuildIndex(this, ERequestAsyncBuildFlag::ContinueRequest);
+	return EAsyncBuildIndexResult::InProgress != FAsyncPoseSearchDatabasesManagement::RequestAsyncBuildIndex(this, ERequestAsyncBuildFlag::ContinueRequest);
 }
 #endif // WITH_EDITOR
 
@@ -1013,7 +1013,7 @@ UE::PoseSearch::FSearchResult UPoseSearchDatabase::Search(UE::PoseSearch::FSearc
 	FSearchResult Result;
 
 #if WITH_EDITOR
-	if (!FAsyncPoseSearchDatabasesManagement::RequestAsyncBuildIndex(this, ERequestAsyncBuildFlag::ContinueRequest))
+	if (EAsyncBuildIndexResult::Success != FAsyncPoseSearchDatabasesManagement::RequestAsyncBuildIndex(this, ERequestAsyncBuildFlag::ContinueRequest))
 	{
 		SearchContext.SetAsyncBuildIndexInProgress();
 		return Result;
@@ -1122,7 +1122,7 @@ UE::PoseSearch::FSearchResult UPoseSearchDatabase::SearchContinuingPose(UE::Pose
 	FSearchResult Result;
 
 #if WITH_EDITOR
-	if (!FAsyncPoseSearchDatabasesManagement::RequestAsyncBuildIndex(this, ERequestAsyncBuildFlag::ContinueRequest))
+	if (EAsyncBuildIndexResult::Success != FAsyncPoseSearchDatabasesManagement::RequestAsyncBuildIndex(this, ERequestAsyncBuildFlag::ContinueRequest))
 	{
 		SearchContext.SetAsyncBuildIndexInProgress();
 		return Result;
diff --git a/Engine/Plugins/Experimental/Animation/PoseSearch/Source/Runtime/Private/PoseSearchDerivedData.cpp b/Engine/Plugins/Experimental/Animation/PoseSearch/Source/Runtime/Private/PoseSearchDerivedData.cpp
--- a/Engine/Plugins/Experimental/Animation/PoseSearch/Source/Runtime/Private/PoseSearchDerivedData.cpp
+++ b/Engine/Plugins/Experimental/Animation/PoseSearch/Source/Runtime/Private/PoseSearchDerivedData.cpp
@@ -2165,12 +2165,11 @@ void FAsyncPoseSearchDatabasesManagement::AddReferencedObjects(FReferenceCollect
 	}
 }
 
-// returns true if the index has been built and the Database updated correctly  
-bool FAsyncPoseSearchDatabasesManagement::RequestAsyncBuildIndex(const UPoseSearchDatabase* Database, ERequestAsyncBuildFlag Flag)
+EAsyncBuildIndexResult FAsyncPoseSearchDatabasesManagement::RequestAsyncBuildIndex(const UPoseSearchDatabase* Database, ERequestAsyncBuildFlag Flag)
 {
 	if (!IsValid(Database))
 	{
-		return false;
+		return EAsyncBuildIndexResult::Failed;
 	}
 
 	FScopeLock Lock(&Mutex);
@@ -2216,7 +2215,17 @@ bool FAsyncPoseSearchDatabasesManagement::RequestAsyncBuildIndex(const UPoseSear
 		}
 	}
 
-	return Task->GetState() == FPoseSearchDatabaseAsyncCacheTask::EState::Ended;
+	if (Task->GetState() == FPoseSearchDatabaseAsyncCacheTask::EState::Ended)
+	{
+		return EAsyncBuildIndexResult::Success;
+	}
+
+	if (Task->GetState() == FPoseSearchDatabaseAsyncCacheTask::EState::Failed)
+	{
+		return EAsyncBuildIndexResult::Failed;
+	}
+
+	return EAsyncBuildIndexResult::InProgress;
 }
 
 } // namespace UE::PoseSearch
diff --git a/Engine/Plugins/Experimental/Animation/PoseSearch/Source/Runtime/Private/PoseSearchLibrary.cpp b/Engine/Plugins/Experimental/Animation/PoseSearch/Source/Runtime/Private/PoseSearchLibrary.cpp
--- a/Engine/Plugins/Experimental/Animation/PoseSearch/Source/Runtime/Private/PoseSearchLibrary.cpp
+++ b/Engine/Plugins/Experimental/Animation/PoseSearch/Source/Runtime/Private/PoseSearchLibrary.cpp
@@ -551,7 +551,7 @@ void UPoseSearchLibrary::UpdateMotionMatchingState(
 		const UPoseSearchDatabase* CurResultDatabase = CurResult.Database.Get();
 
 #if WITH_EDITOR
-		if (FAsyncPoseSearchDatabasesManagement::RequestAsyncBuildIndex(CurResultDatabase, ERequestAsyncBuildFlag::ContinueRequest))
+		if (EAsyncBuildIndexResult::Success == FAsyncPoseSearchDatabasesManagement::RequestAsyncBuildIndex(CurResultDatabase, ERequestAsyncBuildFlag::ContinueRequest))
 #endif // WITH_EDITOR
 		{
 			FAnimInstanceProxy* AnimInstanceProxy = Context.AnimInstanceProxy;
@@ -927,7 +927,7 @@ UE::PoseSearch::FSearchResult UPoseSearchLibrary::MotionMatch(const FAnimationBa
 			const UPoseSearchDatabase* Database = AssetsToSearchPerDatabasePair.Key;
 
 #if WITH_EDITOR
-			if (!FAsyncPoseSearchDatabasesManagement::RequestAsyncBuildIndex(Database, ERequestAsyncBuildFlag::ContinueRequest))
+			if (EAsyncBuildIndexResult::Success != FAsyncPoseSearchDatabasesManagement::RequestAsyncBuildIndex(Database, ERequestAsyncBuildFlag::ContinueRequest))
 			{
 				SearchContext.SetAsyncBuildIndexInProgress();
 			}
diff --git a/Engine/Plugins/Experimental/Animation/PoseSearch/Source/Runtime/Public/PoseSearch/PoseSearchDerivedData.h b/Engine/Plugins/Experimental/Animation/PoseSearch/Source/Runtime/Public/PoseSearch/PoseSearchDerivedData.h
--- a/Engine/Plugins/Experimental/Animation/PoseSearch/Source/Runtime/Public/PoseSearch/PoseSearchDerivedData.h
+++ b/Engine/Plugins/Experimental/Animation/PoseSearch/Source/Runtime/Public/PoseSearch/PoseSearchDerivedData.h
@@ -23,10 +23,17 @@ namespace UE::PoseSearch
 	};
 	ENUM_CLASS_FLAGS(ERequestAsyncBuildFlag);
 
+	enum class EAsyncBuildIndexResult
+	{
+		InProgress,						// indexing in progress
+		Success,						// the index has been built and the Database updated correctly
+		Failed							// indexing failed
+	};
+
 	class POSESEARCH_API FAsyncPoseSearchDatabasesManagement : public FTickableGameObject, public FTickableCookObject, public FGCObject
 	{
 	public:
-		static bool RequestAsyncBuildIndex(const UPoseSearchDatabase* Database, ERequestAsyncBuildFlag Flag);
+		static EAsyncBuildIndexResult RequestAsyncBuildIndex(const UPoseSearchDatabase* Database, ERequestAsyncBuildFlag Flag);
 
 	private:
 		FAsyncPoseSearchDatabasesManagement();
