motion matching - improvements for MM search in chooser player

diff --git a/Engine/Plugins/Experimental/Animation/PoseSearch/Source/Runtime/Private/PoseSearchDatabase.cpp b/Engine/Plugins/Experimental/Animation/PoseSearch/Source/Runtime/Private/PoseSearchDatabase.cpp
--- a/Engine/Plugins/Experimental/Animation/PoseSearch/Source/Runtime/Private/PoseSearchDatabase.cpp
+++ b/Engine/Plugins/Experimental/Animation/PoseSearch/Source/Runtime/Private/PoseSearchDatabase.cpp
@@ -1120,6 +1120,7 @@ UE::PoseSearch::FSearchResult UPoseSearchDatabase::SearchContinuingPose(UE::Pose
 	check(SearchContext.GetCurrentResult().Database.Get() == this);
 
 	FSearchResult Result;
+	Result.bIsContinuingPoseSearch = true;
 
 #if WITH_EDITOR
 	if (EAsyncBuildIndexResult::Success != FAsyncPoseSearchDatabasesManagement::RequestAsyncBuildIndex(this, ERequestAsyncBuildFlag::ContinueRequest))
@@ -1484,34 +1485,89 @@ UE::PoseSearch::FSearchResult UPoseSearchDatabase::SearchBruteForce(UE::PoseSear
 		);
 
 		const int32 NumDimensions = Schema->SchemaCardinality;
-		const FSearchFilters SearchFilters(Schema, NonSelectableIdx, SelectableAssetIdx, SearchIndex.bAnyBlockTransition);
 		const bool bUpdateBestCandidates = PoseSearchMode == EPoseSearchMode::BruteForce;
 
-		// do we need to reconstruct pose values?
-		if (SearchIndex.IsValuesEmpty())
+		const FSearchFilters SearchFilters(Schema, NonSelectableIdx, FSelectableAssetIdx(), SearchIndex.bAnyBlockTransition);
+
+		if (SelectableAssetIdx.IsEmpty())
 		{
-			// FMemory_Alloca is forced 16 bytes aligned
-			TArrayView<float> ReconstructedPoseValuesBuffer((float*)FMemory_Alloca(NumDimensions * sizeof(float)), NumDimensions);
-			check(IsAligned(ReconstructedPoseValuesBuffer.GetData(), alignof(VectorRegister4Float)));
-			for (int32 PoseIdx = 0; PoseIdx < SearchIndex.GetNumPoses(); ++PoseIdx)
+			// do we need to reconstruct pose values?
+			if (SearchIndex.IsValuesEmpty())
 			{
-				EvaluatePoseKernel<true, false>(Result, SearchIndex, QueryValues, ReconstructedPoseValuesBuffer, PoseIdx, SearchFilters, SearchContext, this, bUpdateBestCandidates, PoseIdx);
+				// FMemory_Alloca is forced 16 bytes aligned
+				TArrayView<float> ReconstructedPoseValuesBuffer((float*)FMemory_Alloca(NumDimensions * sizeof(float)), NumDimensions);
+				check(IsAligned(ReconstructedPoseValuesBuffer.GetData(), alignof(VectorRegister4Float)));
+				for (int32 PoseIdx = 0; PoseIdx < SearchIndex.GetNumPoses(); ++PoseIdx)
+				{
+					EvaluatePoseKernel<true, false>(Result, SearchIndex, QueryValues, ReconstructedPoseValuesBuffer, PoseIdx, SearchFilters, SearchContext, this, bUpdateBestCandidates, PoseIdx);
+				}
 			}
-		}
-		// is the data padded at 16 bytes (and 16 bytes aligned by construction)?
-		else if (NumDimensions % 4 == 0)
-		{
-			for (int32 PoseIdx = 0; PoseIdx < SearchIndex.GetNumPoses(); ++PoseIdx)
+			// is the data padded at 16 bytes (and 16 bytes aligned by construction)?
+			else if (NumDimensions % 4 == 0)
+			{
+				for (int32 PoseIdx = 0; PoseIdx < SearchIndex.GetNumPoses(); ++PoseIdx)
+				{
+					EvaluatePoseKernel<false, true>(Result, SearchIndex, QueryValues, TArrayView<float>(), PoseIdx, SearchFilters, SearchContext, this, bUpdateBestCandidates, PoseIdx);
+				}
+			}
+			// no reconstruction, but data is not 16 bytes padded
+			else
 			{
-				EvaluatePoseKernel<false, true>(Result, SearchIndex, QueryValues, TArrayView<float>(), PoseIdx, SearchFilters, SearchContext, this, bUpdateBestCandidates, PoseIdx);
+				for (int32 PoseIdx = 0; PoseIdx < SearchIndex.GetNumPoses(); ++PoseIdx)
+				{
+					EvaluatePoseKernel<false, false>(Result, SearchIndex, QueryValues, TArrayView<float>(), PoseIdx, SearchFilters, SearchContext, this, bUpdateBestCandidates, PoseIdx);
+				}
 			}
 		}
-		// no reconstruction, but data is not 16 bytes padded
 		else
 		{
-			for (int32 PoseIdx = 0; PoseIdx < SearchIndex.GetNumPoses(); ++PoseIdx)
+			int32 ResultIndex = -1;
+
+			// do we need to reconstruct pose values?
+			if (SearchIndex.IsValuesEmpty())
+			{
+				// FMemory_Alloca is forced 16 bytes aligned
+				TArrayView<float> ReconstructedPoseValuesBuffer((float*)FMemory_Alloca(NumDimensions * sizeof(float)), NumDimensions);
+				check(IsAligned(ReconstructedPoseValuesBuffer.GetData(), alignof(VectorRegister4Float)));
+
+				for (int32 AssetIdx : SelectableAssetIdx)
+				{
+					const FSearchIndexAsset& SearchIndexAsset = SearchIndex.Assets[AssetIdx];
+					const int32 FirstPoseIdx = SearchIndexAsset.GetFirstPoseIdx();
+					const int32 LastPoseIdx = FirstPoseIdx + SearchIndexAsset.GetNumPoses();
+					for (int32 PoseIdx = FirstPoseIdx; PoseIdx < LastPoseIdx; ++PoseIdx)
+					{
+						EvaluatePoseKernel<true, false>(Result, SearchIndex, QueryValues, ReconstructedPoseValuesBuffer, PoseIdx, SearchFilters, SearchContext, this, bUpdateBestCandidates, ++ResultIndex);
+					}
+				}
+			}
+			// is the data padded at 16 bytes (and 16 bytes aligned by construction)?
+			else if (NumDimensions % 4 == 0)
 			{
-				EvaluatePoseKernel<false, false>(Result, SearchIndex, QueryValues, TArrayView<float>(), PoseIdx, SearchFilters, SearchContext, this, bUpdateBestCandidates, PoseIdx);
+				for (int32 AssetIdx : SelectableAssetIdx)
+				{
+					const FSearchIndexAsset& SearchIndexAsset = SearchIndex.Assets[AssetIdx];
+					const int32 FirstPoseIdx = SearchIndexAsset.GetFirstPoseIdx();
+					const int32 LastPoseIdx = FirstPoseIdx + SearchIndexAsset.GetNumPoses();
+					for (int32 PoseIdx = FirstPoseIdx; PoseIdx < LastPoseIdx; ++PoseIdx)
+					{
+						EvaluatePoseKernel<false, true>(Result, SearchIndex, QueryValues, TArrayView<float>(), PoseIdx, SearchFilters, SearchContext, this, bUpdateBestCandidates, ++ResultIndex);
+					}
+				}
+			}
+			// no reconstruction, but data is not 16 bytes padded
+			else
+			{
+				for (int32 AssetIdx : SelectableAssetIdx)
+				{
+					const FSearchIndexAsset& SearchIndexAsset = SearchIndex.Assets[AssetIdx];
+					const int32 FirstPoseIdx = SearchIndexAsset.GetFirstPoseIdx();
+					const int32 LastPoseIdx = FirstPoseIdx + SearchIndexAsset.GetNumPoses();
+					for (int32 PoseIdx = FirstPoseIdx; PoseIdx < LastPoseIdx; ++PoseIdx)
+					{
+						EvaluatePoseKernel<false, false>(Result, SearchIndex, QueryValues, TArrayView<float>(), PoseIdx, SearchFilters, SearchContext, this, bUpdateBestCandidates, ++ResultIndex);
+					}
+				}
 			}
 		}
 	}
diff --git a/Engine/Plugins/Experimental/Animation/PoseSearch/Source/Runtime/Private/PoseSearchLibrary.cpp b/Engine/Plugins/Experimental/Animation/PoseSearch/Source/Runtime/Private/PoseSearchLibrary.cpp
--- a/Engine/Plugins/Experimental/Animation/PoseSearch/Source/Runtime/Private/PoseSearchLibrary.cpp
+++ b/Engine/Plugins/Experimental/Animation/PoseSearch/Source/Runtime/Private/PoseSearchLibrary.cpp
@@ -919,6 +919,7 @@ UE::PoseSearch::FSearchResult UPoseSearchLibrary::MotionMatch(const FAnimationBa
 	};
 
 	// collecting all the possible continuing pose search (it could be multiple searches, but most likely only one)
+	const float DeltaSeconds = AnimInstance->GetDeltaSeconds();
 	if (const UAnimationAsset* PlayingAnimationAsset = Cast<UAnimationAsset>(PlayingAsset))
 	{
 		AddToSearch(AssetsToSearchPerDatabaseMap, PlayingAnimationAsset);
@@ -945,16 +946,21 @@ UE::PoseSearch::FSearchResult UPoseSearchLibrary::MotionMatch(const FAnimationBa
 						if (PlayingAnimationAsset == DatabaseAnimationAssetBase->GetAnimationAsset())
 						{
 							const float FirstSampleTime = SearchIndexAsset.GetFirstSampleTime(Database->Schema->SampleRate);
-							const float LastSampleTime = SearchIndexAsset.GetLastSampleTime(Database->Schema->SampleRate);
-							if (PlayingAssetAccumulatedTime >= FirstSampleTime && PlayingAssetAccumulatedTime <= LastSampleTime)
+							const float LastSampleTime = SearchIndexAsset.GetLastSampleTime(Database->Schema->SampleRate) - DeltaSeconds;
+							const float MaxTimeToBeAbleToContinuingPlayingAnimation = LastSampleTime - DeltaSeconds;
+							if (PlayingAssetAccumulatedTime >= FirstSampleTime && PlayingAssetAccumulatedTime < MaxTimeToBeAbleToContinuingPlayingAnimation)
 							{
 								ReconstructedPreviousSearchResult.Database = Database;
 								ReconstructedPreviousSearchResult.AssetTime = PlayingAssetAccumulatedTime;
 								ReconstructedPreviousSearchResult.PoseIdx = Database->GetPoseIndexFromTime(PlayingAssetAccumulatedTime, SearchIndexAsset);
 								SearchContext.UpdateCurrentResultPoseVector();
 
-								SearchResult = Database->SearchContinuingPose(SearchContext);
-								SearchContext.UpdateCurrentBestCost(SearchResult.PoseCost);
+								const FSearchResult NewSearchResult = Database->SearchContinuingPose(SearchContext);
+								if (NewSearchResult.PoseCost.GetTotalCost() < SearchResult.PoseCost.GetTotalCost())
+								{
+									SearchResult = NewSearchResult;
+									SearchContext.UpdateCurrentBestCost(SearchResult.PoseCost);
+								}
 							}
 						}
 					}
@@ -1019,7 +1025,7 @@ UE::PoseSearch::FSearchResult UPoseSearchLibrary::MotionMatch(const FAnimationBa
 	const float SearchBestCost = SearchResult.PoseCost.GetTotalCost();
 	const float SearchBruteForceCost = SearchResult.BruteForcePoseCost.GetTotalCost();
 	TraceMotionMatchingState(SearchContext, SearchResult, 0.f, FTransform::Identity, Context.GetCurrentNodeId(),
-		AnimInstance->GetDeltaSeconds(), true, FObjectTrace::GetWorldElapsedTime(AnimInstance->GetWorld()));
+		DeltaSeconds, true, FObjectTrace::GetWorldElapsedTime(AnimInstance->GetWorld()));
 #endif // UE_POSE_SEARCH_TRACE_ENABLED
 
 	return SearchResult;
diff --git a/Engine/Plugins/Experimental/Animation/PoseSearch/Source/Runtime/Private/PoseSearchModule.cpp b/Engine/Plugins/Experimental/Animation/PoseSearch/Source/Runtime/Private/PoseSearchModule.cpp
--- a/Engine/Plugins/Experimental/Animation/PoseSearch/Source/Runtime/Private/PoseSearchModule.cpp
+++ b/Engine/Plugins/Experimental/Animation/PoseSearch/Source/Runtime/Private/PoseSearchModule.cpp
@@ -33,6 +33,7 @@ public:
 				ProviderResult.SelectedAsset = DatabaseAnimationAssetBase->GetAnimationAsset();
 				ProviderResult.Dissimilarity = SearchResult.PoseCost.GetTotalCost();
 				ProviderResult.TimeOffsetSeconds = SearchResult.AssetTime;
+				ProviderResult.bIsFromContinuingPlaying = SearchResult.bIsContinuingPoseSearch;
 			}
 		}
 
diff --git a/Engine/Plugins/Experimental/Animation/PoseSearch/Source/Runtime/Private/PoseSearchResult.cpp b/Engine/Plugins/Experimental/Animation/PoseSearch/Source/Runtime/Private/PoseSearchResult.cpp
--- a/Engine/Plugins/Experimental/Animation/PoseSearch/Source/Runtime/Private/PoseSearchResult.cpp
+++ b/Engine/Plugins/Experimental/Animation/PoseSearch/Source/Runtime/Private/PoseSearchResult.cpp
@@ -76,22 +76,24 @@ const FSearchIndexAsset* FSearchResult::GetSearchIndexAsset(bool bMandatory) con
 bool FSearchResult::CanAdvance(float DeltaTime) const
 {
 	bool bCanAdvance = false;
-	if (IsValid())
+	if (const FSearchIndexAsset* SearchIndexAsset = GetSearchIndexAsset())
 	{
 		float SteppedTime = AssetTime;
-		const FSearchIndexAsset* SearchIndexAsset = GetSearchIndexAsset(true);
 		const FInstancedStruct& DatabaseAsset = Database->GetAnimationAssetStruct(*SearchIndexAsset);
 		if (const FPoseSearchDatabaseBlendSpace* DatabaseBlendSpace = DatabaseAsset.GetPtr<FPoseSearchDatabaseBlendSpace>())
 		{
-			TArray<FBlendSampleData> BlendSamples;
-			int32 TriangulationIndex = 0;
-			DatabaseBlendSpace->BlendSpace->GetSamplesFromBlendInput(SearchIndexAsset->GetBlendParameters(), BlendSamples, TriangulationIndex, true);
+			if (const UBlendSpace* BlendSpace = DatabaseBlendSpace->BlendSpace.Get())
+			{
+				TArray<FBlendSampleData> BlendSamples;
+				int32 TriangulationIndex = 0;
+				BlendSpace->GetSamplesFromBlendInput(SearchIndexAsset->GetBlendParameters(), BlendSamples, TriangulationIndex, true);
 
-			const float PlayLength = DatabaseBlendSpace->BlendSpace->GetAnimationLengthFromSampleData(BlendSamples);
+				const float PlayLength = BlendSpace->GetAnimationLengthFromSampleData(BlendSamples);
 
-			// Asset player time for blend spaces is normalized [0, 1] so we need to convert it back to real time before we advance it
-			SteppedTime = AssetTime * PlayLength;
-			bCanAdvance = ETAA_Finished != FAnimationRuntime::AdvanceTime(SearchIndexAsset->IsLooping(), DeltaTime, SteppedTime, PlayLength);
+				// Asset player time for blend spaces is normalized [0, 1] so we need to convert it back to real time before we advance it
+				SteppedTime = AssetTime * PlayLength;
+				bCanAdvance = ETAA_Finished != FAnimationRuntime::AdvanceTime(SearchIndexAsset->IsLooping(), DeltaTime, SteppedTime, PlayLength);
+			}
 		}
 		else if (const FPoseSearchDatabaseAnimationAssetBase* DatabaseAnimationAssetBase = DatabaseAsset.GetPtr<FPoseSearchDatabaseAnimationAssetBase>())
 		{
diff --git a/Engine/Plugins/Experimental/Animation/PoseSearch/Source/Runtime/Public/PoseSearch/PoseSearchLibrary.h b/Engine/Plugins/Experimental/Animation/PoseSearch/Source/Runtime/Public/PoseSearch/PoseSearchLibrary.h
--- a/Engine/Plugins/Experimental/Animation/PoseSearch/Source/Runtime/Public/PoseSearch/PoseSearchLibrary.h
+++ b/Engine/Plugins/Experimental/Animation/PoseSearch/Source/Runtime/Public/PoseSearch/PoseSearchLibrary.h
@@ -100,11 +100,11 @@ public:
 
 	// Input start time for the first pose of FutureAnimation
 	UPROPERTY(EditAnywhere, BlueprintReadWrite, Category=State)
-	float FutureAnimationStartTime;
+	float FutureAnimationStartTime = 0.f;
 
 	// Input time in seconds before start playing FutureAnimation (from FutureAnimationStartTime seconds)
 	UPROPERTY(EditAnywhere, BlueprintReadWrite, Category=State)
-	float TimeToFutureAnimationStart;
+	float TimeToFutureAnimationStart = 0.f;
 };
 
 UCLASS()
diff --git a/Engine/Plugins/Experimental/Animation/PoseSearch/Source/Runtime/Public/PoseSearch/PoseSearchResult.h b/Engine/Plugins/Experimental/Animation/PoseSearch/Source/Runtime/Public/PoseSearch/PoseSearchResult.h
--- a/Engine/Plugins/Experimental/Animation/PoseSearch/Source/Runtime/Public/PoseSearch/PoseSearchResult.h
+++ b/Engine/Plugins/Experimental/Animation/PoseSearch/Source/Runtime/Public/PoseSearch/PoseSearchResult.h
@@ -18,10 +18,9 @@ struct FSearchResult
 	// best cost of the currently selected PoseIdx (it could be equal to ContinuingPoseCost)
 	FPoseSearchCost PoseCost;
 	int32 PoseIdx = INDEX_NONE;
-
 	TWeakObjectPtr<const UPoseSearchDatabase> Database;
-
-	float AssetTime = 0.0f;
+	float AssetTime = 0.f;
+	bool bIsContinuingPoseSearch = false;
 
 #if UE_POSE_SEARCH_TRACE_ENABLED
 	FPoseSearchCost BruteForcePoseCost;
diff --git a/Engine/Plugins/Experimental/Chooser/Source/Chooser/Internal/AnimNode_ChooserPlayer.h b/Engine/Plugins/Experimental/Chooser/Source/Chooser/Internal/AnimNode_ChooserPlayer.h
--- a/Engine/Plugins/Experimental/Chooser/Source/Chooser/Internal/AnimNode_ChooserPlayer.h
+++ b/Engine/Plugins/Experimental/Chooser/Source/Chooser/Internal/AnimNode_ChooserPlayer.h
@@ -164,6 +164,7 @@ private:
 	float CurrentStartTime = 0;
 	bool CurrentMirror = false;
 	bool bInitialized = false;
+	bool bForceBlendTo = false;
 	uint32 CurrentCurveOverridesHash = 0;
 	
 
diff --git a/Engine/Plugins/Experimental/Chooser/Source/Chooser/Private/AnimNode_ChooserPlayer.cpp b/Engine/Plugins/Experimental/Chooser/Source/Chooser/Private/AnimNode_ChooserPlayer.cpp
--- a/Engine/Plugins/Experimental/Chooser/Source/Chooser/Private/AnimNode_ChooserPlayer.cpp
+++ b/Engine/Plugins/Experimental/Chooser/Source/Chooser/Private/AnimNode_ChooserPlayer.cpp
@@ -45,7 +45,11 @@ UAnimationAsset* FAnimNode_ChooserPlayer::ChooseAsset(const FAnimationUpdateCont
 				const UE::Anim::IPoseSearchProvider::FSearchResult SearchResult = PoseSearchProvider->Search(Context, ChosenAssets, GetAnimAsset(), GetAccumulatedTime());
 				if (UAnimationAsset* SelectedAnimationAsset = Cast<UAnimationAsset>(SearchResult.SelectedAsset))
 				{
-					Settings.StartTime = SearchResult.TimeOffsetSeconds;
+					if (!SearchResult.bIsFromContinuingPlaying)
+					{
+						bForceBlendTo = true;
+						Settings.StartTime = SearchResult.TimeOffsetSeconds;
+					}
 					return SelectedAnimationAsset;
 				}
 			}
@@ -124,16 +128,18 @@ void FAnimNode_ChooserPlayer::UpdateAssetPlayer(const FAnimationUpdateContext& C
 	}
 
 	// Restart the animation:
+	// - if we've been told to do so via bForceBlendTo
 	// - if this node just became relevant
 	// - if we chose a new animation
 	// - if the mirror setting has changed
 	// - for playback rate of 0, when the start time changes - for choosing poses as frames of an animation sequence
 	// - if the curve values are different
-	if (bJustBecameRelevant || NewAsset != CurrentAsset || AnimPlayers.IsEmpty() ||
+	if (bForceBlendTo || bJustBecameRelevant || NewAsset != CurrentAsset || AnimPlayers.IsEmpty() ||
 		CurrentMirror != Settings.bMirror ||
 		(CurrentStartTime != Settings.StartTime && Settings.PlaybackRate == 0.0f) ||
 		CurrentCurveOverridesHash != Settings.CurveOverrides.Hash)
 	{
+		bForceBlendTo = false;
 		CurrentCurveOverridesHash = Settings.CurveOverrides.Hash;
 		CurrentMirror = Settings.bMirror;
 		
diff --git a/Engine/Source/Runtime/Engine/Public/Animation/AnimPoseSearchProvider.h b/Engine/Source/Runtime/Engine/Public/Animation/AnimPoseSearchProvider.h
--- a/Engine/Source/Runtime/Engine/Public/Animation/AnimPoseSearchProvider.h
+++ b/Engine/Source/Runtime/Engine/Public/Animation/AnimPoseSearchProvider.h
@@ -27,6 +27,7 @@ public:
 		UObject* SelectedAsset = nullptr;
 		float TimeOffsetSeconds = 0.f;
 		float Dissimilarity = MAX_flt;
+		bool bIsFromContinuingPlaying = false;
 	};
 
 	/**
