motion matching
- UPoseSearchDatabase::SynchronizeWithExternalDependencies fix
- UPoseSearchLibrary::MotionMatch continuing pose search for looping assets
- FPoseSearchFutureProperties improvements

diff --git a/Engine/Plugins/Experimental/Animation/PoseSearch/Source/Runtime/Private/PoseSearchDatabase.cpp b/Engine/Plugins/Experimental/Animation/PoseSearch/Source/Runtime/Private/PoseSearchDatabase.cpp
--- a/Engine/Plugins/Experimental/Animation/PoseSearch/Source/Runtime/Private/PoseSearchDatabase.cpp
+++ b/Engine/Plugins/Experimental/Animation/PoseSearch/Source/Runtime/Private/PoseSearchDatabase.cpp
@@ -775,12 +775,12 @@ void UPoseSearchDatabase::SynchronizeWithExternalDependencies()
 	TArray<FAssetIdentifier> Referencers;
 	AssetRegistry.GetReferencers(GetPackage()->GetFName(), Referencers);
 
+	TArray<UAnimSequenceBase*> SequencesBase;
 	for (const FAssetIdentifier& Referencer : Referencers)
 	{
 		TArray<FAssetData> Assets;
 		AssetRegistry.GetAssetsByPackageName(Referencer.PackageName, Assets);
 
-		TArray<UAnimSequenceBase*> SequencesBase;
 		for (const FAssetData& Asset : Assets)
 		{
 			if (Asset.IsInstanceOf(UAnimSequenceBase::StaticClass()))
@@ -793,7 +793,7 @@ void UPoseSearchDatabase::SynchronizeWithExternalDependencies()
 						{
 							if (BranchIn->Database == this)
 							{
-								SequencesBase.Add(SequenceBase);
+								SequencesBase.AddUnique(SequenceBase);
 								break;
 							}
 						}
@@ -801,11 +801,11 @@ void UPoseSearchDatabase::SynchronizeWithExternalDependencies()
 				}
 			}
 		}
-		
-		if (!SequencesBase.IsEmpty())
-		{
-			SynchronizeWithExternalDependencies(SequencesBase);
-		}
+	}
+
+	if (!SequencesBase.IsEmpty())
+	{
+		SynchronizeWithExternalDependencies(SequencesBase);
 	}
 }
 
@@ -815,15 +815,18 @@ void UPoseSearchDatabase::SynchronizeWithExternalDependencies(TConstArrayView<UA
 	TArray<FInstancedStruct> NewAnimationAssets;
 
 	// collecting all the database AnimationAsset(s) that don't require synchronization
+	TArray<bool> DisableReselection;
+	DisableReselection.Reserve(AnimationAssets.Num());
 	for (FInstancedStruct& AnimationAsset : AnimationAssets)
 	{
-		if (const FPoseSearchDatabaseAnimationAssetBase* AnimationAssetBase = AnimationAsset.GetPtr<FPoseSearchDatabaseAnimationAssetBase>())
+		FPoseSearchDatabaseAnimationAssetBase& AnimationAssetBase = AnimationAsset.GetMutable<FPoseSearchDatabaseAnimationAssetBase>();
+		DisableReselection.Add(AnimationAssetBase.bDisableReselection);
+		AnimationAssetBase.bDisableReselection = false;
+
+		const bool bRequiresSynchronization = AnimationAssetBase.bSynchronizeWithExternalDependency && SequencesBase.Contains(AnimationAssetBase.GetAnimationAsset());
+		if (!bRequiresSynchronization)
 		{
-			const bool bRequiresSynchronization = AnimationAssetBase->bSynchronizeWithExternalDependency && SequencesBase.Contains(AnimationAssetBase->GetAnimationAsset());
-			if (!bRequiresSynchronization)
-			{
-				NewAnimationAssets.Add(AnimationAsset);
-			}
+			NewAnimationAssets.Add(AnimationAsset);
 		}
 	}
 
@@ -838,11 +841,21 @@ void UPoseSearchDatabase::SynchronizeWithExternalDependencies(TConstArrayView<UA
 				{
 					if (PoseSearchBranchIn->Database == this)
 					{
+						auto GetSamplingRange = [](const FAnimNotifyEvent& NotifyEvent, const UAnimSequenceBase* SequenceBase) -> FFloatInterval
+						{
+							FFloatInterval SamplingRange(NotifyEvent.GetTime(), NotifyEvent.GetTime() + NotifyEvent.GetDuration());
+							if (SamplingRange.Min <= NotifyEvent.TriggerTimeOffset && SamplingRange.Max >= SequenceBase->GetPlayLength() - NotifyEvent.TriggerTimeOffset)
+							{
+								SamplingRange = FFloatInterval(0.f, 0.f);
+							}
+							return SamplingRange;
+						};
+
 						if (UAnimSequence* Sequence = Cast<UAnimSequence>(SequenceBase))
 						{
 							FPoseSearchDatabaseSequence DatabaseSequence;
 							DatabaseSequence.Sequence = Sequence;
-							DatabaseSequence.SamplingRange = FFloatInterval(NotifyEvent.GetTime(), NotifyEvent.GetTime() + NotifyEvent.GetDuration());
+							DatabaseSequence.SamplingRange = GetSamplingRange(NotifyEvent, SequenceBase);
 							DatabaseSequence.bSynchronizeWithExternalDependency = true;
 							NewAnimationAssets.Add(FInstancedStruct::Make(DatabaseSequence));
 						}
@@ -850,7 +863,7 @@ void UPoseSearchDatabase::SynchronizeWithExternalDependencies(TConstArrayView<UA
 						{
 							FPoseSearchDatabaseAnimComposite DatabaseAnimComposite;
 							DatabaseAnimComposite.AnimComposite = AnimComposite;
-							DatabaseAnimComposite.SamplingRange = FFloatInterval(NotifyEvent.GetTime(), NotifyEvent.GetTime() + NotifyEvent.GetDuration());
+							DatabaseAnimComposite.SamplingRange = GetSamplingRange(NotifyEvent, SequenceBase);
 							DatabaseAnimComposite.bSynchronizeWithExternalDependency = true;
 							NewAnimationAssets.Add(FInstancedStruct::Make(DatabaseAnimComposite));
 						}
@@ -858,7 +871,7 @@ void UPoseSearchDatabase::SynchronizeWithExternalDependencies(TConstArrayView<UA
 						{
 							FPoseSearchDatabaseAnimMontage DatabaseAnimMontage;
 							DatabaseAnimMontage.AnimMontage = AnimMontage;
-							DatabaseAnimMontage.SamplingRange = FFloatInterval(NotifyEvent.GetTime(), NotifyEvent.GetTime() + NotifyEvent.GetDuration());
+							DatabaseAnimMontage.SamplingRange = GetSamplingRange(NotifyEvent, SequenceBase);
 							DatabaseAnimMontage.bSynchronizeWithExternalDependency = true;
 							NewAnimationAssets.Add(FInstancedStruct::Make(DatabaseAnimMontage));
 						}
@@ -875,7 +888,9 @@ void UPoseSearchDatabase::SynchronizeWithExternalDependencies(TConstArrayView<UA
 		const int32 FoundIndex = NewAnimationAssets.Find(AnimationAssets[AnimationAssetIndex]);
 		if (FoundIndex >= 0)
 		{
-			NewAnimationAssets.RemoveAt(FoundIndex);
+			FPoseSearchDatabaseAnimationAssetBase& AnimationAssetBase = AnimationAssets[AnimationAssetIndex].GetMutable<FPoseSearchDatabaseAnimationAssetBase>();
+			AnimationAssetBase.bDisableReselection = DisableReselection[AnimationAssetIndex];
+			NewAnimationAssets.RemoveAt(FoundIndex); 
 		}
 		else
 		{
diff --git a/Engine/Plugins/Experimental/Animation/PoseSearch/Source/Runtime/Private/PoseSearchLibrary.cpp b/Engine/Plugins/Experimental/Animation/PoseSearch/Source/Runtime/Private/PoseSearchLibrary.cpp
--- a/Engine/Plugins/Experimental/Animation/PoseSearch/Source/Runtime/Private/PoseSearchLibrary.cpp
+++ b/Engine/Plugins/Experimental/Animation/PoseSearch/Source/Runtime/Private/PoseSearchLibrary.cpp
@@ -591,7 +591,7 @@ void UPoseSearchLibrary::MotionMatch(
 	TArray<FName, TInlineAllocator<PreallocatedRolesNum, TMemStackAllocator<>>> Roles;
 	Roles.Add(UE::PoseSearch::DefaultRole);
 
-	MotionMatch(AnimInstances, Roles, Database, PoseHistoryName, Result, Future.FutureAnimation, Future.FutureAnimationStartTime, Future.TimeToFutureAnimationStart, DebugSessionUniqueIdentifier);
+	MotionMatch(AnimInstances, Roles, Database, PoseHistoryName, Future, Result, DebugSessionUniqueIdentifier);
 }
 
 void UPoseSearchLibrary::MotionMatchMulti(
@@ -617,7 +617,7 @@ void UPoseSearchLibrary::MotionMatchMulti(
 		AnimInstances.Add(AnimInstance);
 	}
 
-	MotionMatch(AnimInstances, Roles, Database, PoseHistoryName, Result, nullptr, 0.f, 0.f, DebugSessionUniqueIdentifier);
+	MotionMatch(AnimInstances, Roles, Database, PoseHistoryName, FPoseSearchFutureProperties(), Result, DebugSessionUniqueIdentifier);
 }
 
 void UPoseSearchLibrary::MotionMatch(
@@ -625,10 +625,8 @@ void UPoseSearchLibrary::MotionMatch(
 	TConstArrayView<FName> Roles,
 	const UPoseSearchDatabase* Database,
 	const FName PoseHistoryName,
+	const FPoseSearchFutureProperties& Future,
 	FPoseSearchBlueprintResult& Result,
-	const UAnimationAsset* FutureAnimation,
-	float FutureAnimationStartTime,
-	float TimeToFutureAnimationStart,
 	const int32 DebugSessionUniqueIdentifier)
 {
 #if ENABLE_DRAW_DEBUG && ENABLE_ANIM_DEBUG
@@ -720,7 +718,8 @@ void UPoseSearchLibrary::MotionMatch(
 		return;
 	}
 
-	if (FutureAnimation)
+	float FutureIntervalTime = Future.IntervalTime;
+	if (Future.Animation)
 	{
 		if (AnimInstances.Num() != 1)
 		{
@@ -728,17 +727,18 @@ void UPoseSearchLibrary::MotionMatch(
 			return;
 		}
 		
-		if (FutureAnimationStartTime < FiniteDelta)
+		float FutureAnimationTime = Future.AnimationTime;
+		if (FutureAnimationTime < FiniteDelta)
 		{
-			UE_LOG(LogPoseSearch, Warning, TEXT("UPoseSearchLibrary::MotionMatch - provided FutureAnimationStartTime (%f) is too small to be able to calculate velocities. Clamping it to minimum value of %f"), FutureAnimationStartTime, FiniteDelta);
-			FutureAnimationStartTime = FiniteDelta;
+			UE_LOG(LogPoseSearch, Warning, TEXT("UPoseSearchLibrary::MotionMatch - provided Future.AnimationTime (%f) is too small to be able to calculate velocities. Clamping it to minimum value of %f"), FutureAnimationTime, FiniteDelta);
+			FutureAnimationTime = FiniteDelta;
 		}
 
-		const float MinTimeToFutureAnimationStart = FiniteDelta + UE_KINDA_SMALL_NUMBER;
-		if (TimeToFutureAnimationStart < MinTimeToFutureAnimationStart)
+		const float MinFutureIntervalTime = FiniteDelta + UE_KINDA_SMALL_NUMBER;
+		if (FutureIntervalTime < MinFutureIntervalTime)
 		{
-			UE_LOG(LogPoseSearch, Warning, TEXT("UPoseSearchLibrary::MotionMatch - provided TimeToFutureAnimationStart (%f) is too small. Clamping it to minimum value of %f"), TimeToFutureAnimationStart, MinTimeToFutureAnimationStart);
-			TimeToFutureAnimationStart = MinTimeToFutureAnimationStart;
+			UE_LOG(LogPoseSearch, Warning, TEXT("UPoseSearchLibrary::MotionMatch - provided TimeToFutureAnimationStart (%f) is too small. Clamping it to minimum value of %f"), FutureIntervalTime, MinFutureIntervalTime);
+			FutureIntervalTime = MinFutureIntervalTime;
 		}
 
 		// extracting 2 poses to be able to calculate velocities
@@ -747,15 +747,15 @@ void UPoseSearchLibrary::MotionMatch(
 		Pose.SetBoneContainer(&AnimInstances[0]->GetRequiredBonesOnAnyThread());
 		
 		// @todo: add input BlendParameters to support sampling FutureAnimation blendspaces and support for multi character
-		const FAnimationAssetSampler Sampler(FutureAnimation);
+		const FAnimationAssetSampler Sampler(Future.Animation);
 		for (int32 i = 0; i < 2; ++i)
 		{
-			const float ExtractionTime = FutureAnimationStartTime + (i - 1) * FiniteDelta;
-			const float FutureAnimationTime = TimeToFutureAnimationStart + (i - 1) * FiniteDelta;
+			const float FuturePoseExtractionTime = FutureAnimationTime + (i - 1) * FiniteDelta;
+			const float FuturePoseAnimationTime = FutureIntervalTime + (i - 1) * FiniteDelta;
 				
-			Sampler.ExtractPose(ExtractionTime, Pose);
+			Sampler.ExtractPose(FuturePoseExtractionTime, Pose);
 			ComponentSpacePose.InitPose(Pose);
-			MemStackPoseHistories[0].AddFuturePose(FutureAnimationTime, ComponentSpacePose);
+			MemStackPoseHistories[0].AddFuturePose(FuturePoseAnimationTime, ComponentSpacePose);
 		}
 
 #if ENABLE_DRAW_DEBUG && ENABLE_ANIM_DEBUG
@@ -766,7 +766,7 @@ void UPoseSearchLibrary::MotionMatch(
 #endif // ENABLE_DRAW_DEBUG && ENABLE_ANIM_DEBUG
 	}
 
-	FSearchContext SearchContext(TimeToFutureAnimationStart);
+	FSearchContext SearchContext(FutureIntervalTime);
 	for (int32 RoleIndex = 0; RoleIndex < Roles.Num(); ++RoleIndex)
 	{
 		SearchContext.AddRole(Roles[RoleIndex], AnimInstances[RoleIndex], MemStackPoseHistories[RoleIndex].GetThisOrPoseHistory());
@@ -790,7 +790,7 @@ void UPoseSearchLibrary::MotionMatch(
 			
 			// figuring out the WantedPlayRate
 			Result.WantedPlayRate = 1.f;
-			if (FutureAnimation)
+			if (Future.Animation)
 			{
 				if (const UPoseSearchFeatureChannel_PermutationTime* PermutationTimeChannel = Database->Schema->FindFirstChannelOfType<UPoseSearchFeatureChannel_PermutationTime>())
 				{
@@ -798,8 +798,8 @@ void UPoseSearchLibrary::MotionMatch(
 					if (!SearchIndex.IsValuesEmpty())
 					{
 						TConstArrayView<float> ResultData = Database->GetSearchIndex().GetPoseValues(SearchResult.PoseIdx);
-						const float ActualTimeToFutureAnimationStart = PermutationTimeChannel->GetPermutationTime(ResultData);
-						Result.WantedPlayRate = ActualTimeToFutureAnimationStart / TimeToFutureAnimationStart;
+						const float ActualIntervalTime = PermutationTimeChannel->GetPermutationTime(ResultData);
+						Result.WantedPlayRate = ActualIntervalTime / FutureIntervalTime;
 					}
 				}
 			}
@@ -945,10 +945,17 @@ UE::PoseSearch::FSearchResult UPoseSearchLibrary::MotionMatch(const FAnimationBa
 					{
 						if (PlayingAnimationAsset == DatabaseAnimationAssetBase->GetAnimationAsset())
 						{
-							const float FirstSampleTime = SearchIndexAsset.GetFirstSampleTime(Database->Schema->SampleRate);
-							const float LastSampleTime = SearchIndexAsset.GetLastSampleTime(Database->Schema->SampleRate) - DeltaSeconds;
-							const float MaxTimeToBeAbleToContinuingPlayingAnimation = LastSampleTime - DeltaSeconds;
-							if (PlayingAssetAccumulatedTime >= FirstSampleTime && PlayingAssetAccumulatedTime < MaxTimeToBeAbleToContinuingPlayingAnimation)
+							bool bCanAdvance = true;
+							if (!SearchIndexAsset.IsLooping())
+							{
+								const float FirstSampleTime = SearchIndexAsset.GetFirstSampleTime(Database->Schema->SampleRate);
+								const float LastSampleTime = SearchIndexAsset.GetLastSampleTime(Database->Schema->SampleRate) - DeltaSeconds;
+								const float MaxTimeToBeAbleToContinuingPlayingAnimation = LastSampleTime - DeltaSeconds;
+
+								bCanAdvance = PlayingAssetAccumulatedTime >= FirstSampleTime && PlayingAssetAccumulatedTime < MaxTimeToBeAbleToContinuingPlayingAnimation;
+							}
+
+							if (bCanAdvance)
 							{
 								ReconstructedPreviousSearchResult.Database = Database;
 								ReconstructedPreviousSearchResult.AssetTime = PlayingAssetAccumulatedTime;
diff --git a/Engine/Plugins/Experimental/Animation/PoseSearch/Source/Runtime/Public/PoseSearch/PoseSearchLibrary.h b/Engine/Plugins/Experimental/Animation/PoseSearch/Source/Runtime/Public/PoseSearch/PoseSearchLibrary.h
--- a/Engine/Plugins/Experimental/Animation/PoseSearch/Source/Runtime/Public/PoseSearch/PoseSearchLibrary.h
+++ b/Engine/Plugins/Experimental/Animation/PoseSearch/Source/Runtime/Public/PoseSearch/PoseSearchLibrary.h
@@ -94,17 +94,17 @@ struct POSESEARCH_API FPoseSearchFutureProperties
 	GENERATED_BODY()
 
 public:
-	// Input animation we want to match after TimeToFutureAnimationStart seconds
+	// Animation to play (it'll start at AnimationTime seconds)
 	UPROPERTY(EditAnywhere, BlueprintReadWrite, Category=State)
-	TObjectPtr<UAnimationAsset> FutureAnimation;
+	TObjectPtr<UAnimationAsset> Animation;
 
-	// Input start time for the first pose of FutureAnimation
+	// Start time for Animation
 	UPROPERTY(EditAnywhere, BlueprintReadWrite, Category=State)
-	float FutureAnimationStartTime = 0.f;
+	float AnimationTime = 0.f;
 
-	// Input time in seconds before start playing FutureAnimation (from FutureAnimationStartTime seconds)
+	// Interval time before playing Animation
 	UPROPERTY(EditAnywhere, BlueprintReadWrite, Category=State)
-	float TimeToFutureAnimationStart = 0.f;
+	float IntervalTime = 0.f;
 };
 
 UCLASS()
@@ -198,10 +198,8 @@ public:
 		TConstArrayView<FName> Roles,
 		const UPoseSearchDatabase* Database,
 		const FName PoseHistoryName,
+		const FPoseSearchFutureProperties& Future,
 		FPoseSearchBlueprintResult& Result,
-		const UAnimationAsset* FutureAnimation,
-		float FutureAnimationStartTime,
-		float TimeToFutureAnimationStart,
 		const int32 DebugSessionUniqueIdentifier);
 
 	static UE::PoseSearch::FSearchResult MotionMatch(const FAnimationBaseContext& Context, TConstArrayView<UObject*> AssetsToSearch,
