WIP: Adding support for jump trajectory prediction.

[RN] minor

diff --git a/Engine/Plugins/Experimental/Animation/PoseSearch/Source/Runtime/Private/PoseSearchTrajectoryLibrary.cpp b/Engine/Plugins/Experimental/Animation/PoseSearch/Source/Runtime/Private/PoseSearchTrajectoryLibrary.cpp
--- a/Engine/Plugins/Experimental/Animation/PoseSearch/Source/Runtime/Private/PoseSearchTrajectoryLibrary.cpp
+++ b/Engine/Plugins/Experimental/Animation/PoseSearch/Source/Runtime/Private/PoseSearchTrajectoryLibrary.cpp
@@ -33,6 +33,9 @@ void FPoseSearchTrajectoryData::UpdateData(float DeltaTime, const UAnimInstance*
 
 		TrajectoryDataDerived.Velocity = MoveComp->Velocity;
 		TrajectoryDataDerived.Acceleration = MoveComp->GetCurrentAcceleration();
+		
+		TrajectoryDataDerived.bIsFalling = MoveComp->IsFalling();
+		TrajectoryDataDerived.Gravity = -MoveComp->GetGravityDirection() * MoveComp->GetGravityZ();
 
 		if (TrajectoryDataDerived.Acceleration.IsZero())
 		{
@@ -136,6 +139,22 @@ FVector FPoseSearchTrajectoryData::StepCharacterMovementGroundPrediction(float D
 	return OutVelocity;
 }
 
+FVector FPoseSearchTrajectoryData::StepCharacterFallingPrediction(float DeltaTime, const FVector& InitialVelocity, const FDerived& TrajectoryDataDerived) const
+{
+	FVector Result = InitialVelocity;
+
+	if (DeltaTime > 0.f)
+	{
+		// Apply gravity.
+		Result += TrajectoryDataDerived.Gravity * DeltaTime;
+	}
+
+	// @todo: Support for air control to give more accurate results.
+	// @todo: Support for collision prediction to properly predict future on ground path after landing.
+	
+	return Result;
+}
+
 void UPoseSearchTrajectoryLibrary::InitTrajectorySamples(FPoseSearchQueryTrajectory& Trajectory, const FPoseSearchTrajectoryData& TrajectoryData, const FPoseSearchTrajectoryData::FDerived& TrajectoryDataDerived, const FPoseSearchTrajectoryData::FSampling& TrajectoryDataSampling)
 {
 	const int32 NumHistorySamples = TrajectoryDataSampling.NumHistorySamples;
@@ -303,6 +322,89 @@ void UPoseSearchTrajectoryLibrary::UpdatePrediction_SimulateCharacterMovement(FP
 	}
 }
 
+void UPoseSearchTrajectoryLibrary::UpdatePrediction_SimulateCharacterMovementAdvanced(FPoseSearchQueryTrajectory& Trajectory, const FPoseSearchTrajectoryData& TrajectoryData, const FPoseSearchTrajectoryData::FDerived& TrajectoryDataDerived, const FPoseSearchTrajectoryData::FSampling& TrajectoryDataSampling)
+{
+	FVector CurrentPositionWS = TrajectoryDataDerived.Position;
+	FVector CurrentVelocityWS = RemapVectorMagnitudeWithCurve(TrajectoryDataDerived.Velocity, TrajectoryData.bUseSpeedRemappingCurve, TrajectoryData.SpeedRemappingCurve);
+	FVector CurrentAccelerationWS = RemapVectorMagnitudeWithCurve(TrajectoryDataDerived.Acceleration, TrajectoryData.bUseAccelerationRemappingCurve, TrajectoryData.AccelerationRemappingCurve);
+
+	// Bending CurrentVelocityWS towards CurrentAccelerationWS
+	if (TrajectoryData.BendVelocityTowardsAcceleration > UE_KINDA_SMALL_NUMBER && !CurrentAccelerationWS.IsNearlyZero())
+	{
+		const float CurrentSpeed = CurrentVelocityWS.Length();
+		const FVector VelocityWSAlongAcceleration = CurrentAccelerationWS.GetUnsafeNormal() * CurrentSpeed;
+		if (TrajectoryData.BendVelocityTowardsAcceleration < 1.f - UE_KINDA_SMALL_NUMBER)
+		{
+			CurrentVelocityWS = FMath::Lerp(CurrentVelocityWS, VelocityWSAlongAcceleration, TrajectoryData.BendVelocityTowardsAcceleration);
+
+			const float NewLength = CurrentVelocityWS.Length();
+			if (NewLength > UE_KINDA_SMALL_NUMBER)
+			{
+				CurrentVelocityWS *= CurrentSpeed / NewLength;
+			}
+			else
+			{
+				// @todo: consider setting the CurrentVelocityWS = VelocityWSAlongAcceleration if vel and acc are in opposite directions
+			}
+		}
+		else
+		{
+			CurrentVelocityWS = VelocityWSAlongAcceleration;
+		}
+	}
+
+	FQuat CurrentFacingWS = TrajectoryDataDerived.Facing;
+	
+	const int32 NumHistorySamples = TrajectoryDataSampling.NumHistorySamples;
+	const float SecondsPerPredictionSample = TrajectoryDataSampling.SecondsPerPredictionSample;
+	const FQuat ControllerRotationPerStep = FQuat::MakeFromEuler(FVector(0.f, 0.f, TrajectoryDataDerived.ControllerYawRate * SecondsPerPredictionSample));
+
+	float AccumulatedSeconds = 0.f;
+
+	const int32 LastIndex = Trajectory.Samples.Num() - 1;
+	if (NumHistorySamples <= LastIndex)
+	{
+		for (int32 Index = NumHistorySamples; ; ++Index)
+		{
+			Trajectory.Samples[Index].Position = CurrentPositionWS;
+			Trajectory.Samples[Index].Facing = CurrentFacingWS;
+			Trajectory.Samples[Index].AccumulatedSeconds = AccumulatedSeconds;
+
+			if (Index == LastIndex)
+			{
+				break;
+			}
+
+			CurrentPositionWS += CurrentVelocityWS * SecondsPerPredictionSample;
+			AccumulatedSeconds += SecondsPerPredictionSample;
+
+			// Account for the controller (e.g. the camera) rotating.
+			CurrentFacingWS = ControllerRotationPerStep * CurrentFacingWS;
+			CurrentAccelerationWS = RemapVectorMagnitudeWithCurve(ControllerRotationPerStep * CurrentAccelerationWS,
+			TrajectoryData.bUseAccelerationRemappingCurve, TrajectoryData.AccelerationRemappingCurve);
+
+			// Perform trajectory predictions (jumping, moving, etc).
+			{
+				FVector NewVelocityWS = TrajectoryData.StepCharacterMovementGroundPrediction(SecondsPerPredictionSample, CurrentVelocityWS, CurrentAccelerationWS, TrajectoryDataDerived);
+
+				if (TrajectoryDataDerived.bIsFalling)
+				{
+					NewVelocityWS = TrajectoryData.StepCharacterFallingPrediction(SecondsPerPredictionSample, CurrentVelocityWS, TrajectoryDataDerived);
+				}
+				
+				CurrentVelocityWS = RemapVectorMagnitudeWithCurve(NewVelocityWS, TrajectoryData.bUseSpeedRemappingCurve, TrajectoryData.SpeedRemappingCurve);
+			}
+		
+			if (TrajectoryDataDerived.bOrientRotationToMovement && !CurrentAccelerationWS.IsNearlyZero())
+			{
+				// Rotate towards acceleration.
+				const FVector CurrentAccelerationCS = TrajectoryDataDerived.MeshCompRelativeRotation.RotateVector(CurrentAccelerationWS);
+				CurrentFacingWS = FMath::QInterpConstantTo(CurrentFacingWS, CurrentAccelerationCS.ToOrientationQuat(), SecondsPerPredictionSample, TrajectoryData.RotateTowardsMovementSpeed);
+			}
+		}
+	}
+}
+
 void UPoseSearchTrajectoryLibrary::PoseSearchGenerateTrajectory(
 	const UAnimInstance* InAnimInstance, UPARAM(ref) const FPoseSearchTrajectoryData& InTrajectoryData, float InDeltaTime,
 	UPARAM(ref) FPoseSearchQueryTrajectory& InOutTrajectory, UPARAM(ref) float& InOutDesiredControllerYawLastUpdate, FPoseSearchQueryTrajectory& OutTrajectory,
@@ -326,4 +428,30 @@ void UPoseSearchTrajectoryLibrary::PoseSearchGenerateTrajectory(
 	InOutDesiredControllerYawLastUpdate = TrajectoryDataState.DesiredControllerYawLastUpdate;
 
 	OutTrajectory = InOutTrajectory;
-}
\ No newline at end of file
+}
+
+void UPoseSearchTrajectoryLibrary::PoseSearchGenerateTrajectoryAdvanced(const UAnimInstance* InAnimInstance, const FPoseSearchTrajectoryData& InTrajectoryData, float InDeltaTime,
+	FPoseSearchQueryTrajectory& InOutTrajectory, float& InOutDesiredControllerYawLastUpdate, FPoseSearchQueryTrajectory& OutTrajectory, float InHistorySamplingInterval,
+	int32 InTrajectoryHistoryCount, float InPredictionSamplingInterval, int32 InTrajectoryPredictionCount)
+{
+	FPoseSearchTrajectoryData::FSampling TrajectoryDataSampling;
+	TrajectoryDataSampling.NumHistorySamples = InTrajectoryHistoryCount;
+	TrajectoryDataSampling.SecondsPerHistorySample = InHistorySamplingInterval;
+	TrajectoryDataSampling.NumPredictionSamples = InTrajectoryPredictionCount;
+	TrajectoryDataSampling.SecondsPerPredictionSample = InPredictionSamplingInterval;
+
+	FPoseSearchTrajectoryData::FState TrajectoryDataState;
+	TrajectoryDataState.DesiredControllerYawLastUpdate = InOutDesiredControllerYawLastUpdate;
+
+	FPoseSearchTrajectoryData::FDerived TrajectoryDataDerived;
+	InTrajectoryData.UpdateData(InDeltaTime, InAnimInstance, TrajectoryDataDerived, TrajectoryDataState);
+	
+	InitTrajectorySamples(InOutTrajectory, InTrajectoryData, TrajectoryDataDerived, TrajectoryDataSampling);
+	UpdateHistory_TransformHistory(InOutTrajectory, InTrajectoryData, TrajectoryDataDerived, TrajectoryDataSampling, InDeltaTime);
+	UpdatePrediction_SimulateCharacterMovementAdvanced(InOutTrajectory, InTrajectoryData, TrajectoryDataDerived, TrajectoryDataSampling);
+
+	InOutDesiredControllerYawLastUpdate = TrajectoryDataState.DesiredControllerYawLastUpdate;
+
+	OutTrajectory = InOutTrajectory;
+}
+
diff --git a/Engine/Plugins/Experimental/Animation/PoseSearch/Source/Runtime/Public/PoseSearch/PoseSearchTrajectoryLibrary.h b/Engine/Plugins/Experimental/Animation/PoseSearch/Source/Runtime/Public/PoseSearch/PoseSearchTrajectoryLibrary.h
--- a/Engine/Plugins/Experimental/Animation/PoseSearch/Source/Runtime/Public/PoseSearch/PoseSearchTrajectoryLibrary.h
+++ b/Engine/Plugins/Experimental/Animation/PoseSearch/Source/Runtime/Public/PoseSearch/PoseSearchTrajectoryLibrary.h
@@ -29,11 +29,13 @@ public:
 
 		FVector Velocity = FVector::ZeroVector;
 		FVector Acceleration = FVector::ZeroVector;
-
+		FVector Gravity = FVector::ZeroVector;
+		
 		FVector Position = FVector::ZeroVector;
 		FQuat Facing = FQuat::Identity;
 		FQuat MeshCompRelativeRotation = FQuat::Identity;
 		bool bOrientRotationToMovement = false;
+		bool bIsFalling = false;
 	};
 
 	struct FSampling
@@ -48,6 +50,7 @@ public:
 	void UpdateData(float DeltaTime, const FAnimInstanceProxy& AnimInstanceProxy, FDerived& TrajectoryDataDerived, FState& TrajectoryDataState) const;
 	void UpdateData(float DeltaTime, const UAnimInstance* AnimInstance, FDerived& TrajectoryDataDerived, FState& TrajectoryDataState) const;
 	FVector StepCharacterMovementGroundPrediction(float DeltaTime, const FVector& InVelocity, const FVector& InAcceleration, const FDerived& TrajectoryDataDerived) const;
+	FVector StepCharacterFallingPrediction(float DeltaTime, const FVector& InitialVelocity, const FDerived& TrajectoryDataDerived) const;
 	
 	// If the character is forward facing (i.e. bOrientRotationToMovement is true), this controls how quickly the trajectory will rotate
 	// to face acceleration. It's common for this to differ from the rotation rate of the character, because animations are often authored 
@@ -97,6 +100,9 @@ public:
 	// Update prediction by simulating the movement math for ground locomotion from UCharacterMovementComponent.
 	static void UpdatePrediction_SimulateCharacterMovement(FPoseSearchQueryTrajectory& Trajectory, const FPoseSearchTrajectoryData& TrajectoryData, const FPoseSearchTrajectoryData::FDerived& TrajectoryDataDerived, const FPoseSearchTrajectoryData::FSampling& TrajectoryDataSampling);
 
+	// Update prediction by simulating the movement math for locomotion from UCharacterMovementComponent.
+	static void UpdatePrediction_SimulateCharacterMovementAdvanced(FPoseSearchQueryTrajectory& Trajectory, const FPoseSearchTrajectoryData& TrajectoryData, const FPoseSearchTrajectoryData::FDerived& TrajectoryDataDerived, const FPoseSearchTrajectoryData::FSampling& TrajectoryDataSampling);
+	
 	/** Get a Pose History node context from an anim node context (pure) */
 	UFUNCTION(BlueprintPure, Category = "Animation|PoseSearch", meta = (BlueprintThreadSafe, DisplayName = "Pose Search Generate Trajectory"))
 	static void PoseSearchGenerateTrajectory(
@@ -104,6 +110,13 @@ public:
 		UPARAM(ref) FPoseSearchQueryTrajectory& InOutTrajectory, UPARAM(ref) float& InOutDesiredControllerYawLastUpdate, FPoseSearchQueryTrajectory& OutTrajectory,
 		float InHistorySamplingInterval = 0.04f, int32 InTrajectoryHistoryCount = 10, float InPredictionSamplingInterval = 0.4f, int32 InTrajectoryPredictionCount = 8);
 
+	/** Generate trajectory to be used by a Pose History node. Supports more than on ground movement prediction (i.e. jumping, collisions etc). */
+	UFUNCTION(BlueprintPure, Category = "Animation|PoseSearch", meta = (BlueprintThreadSage, DisplayName = "Pose Search Generate Trajectory Advanced"))
+	static void PoseSearchGenerateTrajectoryAdvanced(
+		const UAnimInstance* InAnimInstance, UPARAM(ref) const FPoseSearchTrajectoryData& InTrajectoryData, float InDeltaTime,
+		UPARAM(ref) FPoseSearchQueryTrajectory& InOutTrajectory, UPARAM(ref) float& InOutDesiredControllerYawLastUpdate, FPoseSearchQueryTrajectory& OutTrajectory,
+		float InHistorySamplingInterval = 0.04f, int32 InTrajectoryHistoryCount = 10, float InPredictionSamplingInterval = 0.4f, int32 InTrajectoryPredictionCount = 8);
+	
 private:
 	static FVector RemapVectorMagnitudeWithCurve(const FVector& Vector, bool bUseCurve, const FRuntimeFloatCurve& Curve);
 };
\ No newline at end of file
