motion matching
- commonized UPoseSearchLibrary::MotionMatch functions
- support for multiple assets to search at once

diff --git a/Engine/Plugins/Experimental/Animation/PoseSearch/Source/Runtime/Private/PoseSearchLibrary.cpp b/Engine/Plugins/Experimental/Animation/PoseSearch/Source/Runtime/Private/PoseSearchLibrary.cpp
--- a/Engine/Plugins/Experimental/Animation/PoseSearch/Source/Runtime/Private/PoseSearchLibrary.cpp
+++ b/Engine/Plugins/Experimental/Animation/PoseSearch/Source/Runtime/Private/PoseSearchLibrary.cpp
@@ -19,6 +19,7 @@
 #include "PoseSearch/PoseSearchAnimNotifies.h"
 #include "PoseSearch/PoseSearchDatabase.h"
 #include "PoseSearch/PoseSearchDerivedData.h"
+#include "PoseSearch/PoseSearchMultiSequence.h"
 #include "PoseSearch/PoseSearchSchema.h"
 #include "PoseSearchFeatureChannel_Trajectory.h"
 #include "PoseSearch/Trace/PoseSearchTraceLogger.h"
@@ -106,6 +107,21 @@ namespace UE::PoseSearch
 
 		return false;
 	}
+
+#if ENABLE_DRAW_DEBUG && ENABLE_ANIM_DEBUG
+	class UAnimInstanceProxyProvider : public UAnimInstance
+	{
+	public:
+		static FAnimInstanceProxy* GetAnimInstanceProxy(UAnimInstance* AnimInstance)
+		{
+			if (AnimInstance)
+			{
+				return &static_cast<UAnimInstanceProxyProvider*>(AnimInstance)->GetProxyOnAnyThread<FAnimInstanceProxy>();
+			}
+			return nullptr;
+		}
+	};
+#endif //ENABLE_DRAW_DEBUG && ENABLE_ANIM_DEBUG
 }
 
 //////////////////////////////////////////////////////////////////////////
@@ -575,7 +591,7 @@ void UPoseSearchLibrary::UpdateMotionMatchingState(
 
 void UPoseSearchLibrary::MotionMatch(
 	UAnimInstance* AnimInstance,
-	const UPoseSearchDatabase* Database,
+	TArray<UObject*> AssetsToSearch,
 	const FName PoseHistoryName,
 	FPoseSearchFutureProperties Future,
 	FPoseSearchBlueprintResult& Result,
@@ -591,13 +607,14 @@ void UPoseSearchLibrary::MotionMatch(
 	TArray<FName, TInlineAllocator<PreallocatedRolesNum, TMemStackAllocator<>>> Roles;
 	Roles.Add(UE::PoseSearch::DefaultRole);
 
-	MotionMatch(AnimInstances, Roles, Database, PoseHistoryName, Future, Result, DebugSessionUniqueIdentifier);
+	TArray<const UObject*>& AssetsToSearchConst = reinterpret_cast<TArray<const UObject*>&>(AssetsToSearch);
+	MotionMatch(AnimInstances, Roles, AssetsToSearchConst, PoseHistoryName, Future, Result, DebugSessionUniqueIdentifier);
 }
 
 void UPoseSearchLibrary::MotionMatchMulti(
 	TArray<ACharacter*> Characters,
 	TArray<FName> Roles,
-	const UPoseSearchDatabase* Database,
+	TArray<UObject*> AssetsToSearch,
 	const FName PoseHistoryName,
 	FPoseSearchBlueprintResult& Result,
 	const int32 DebugSessionUniqueIdentifier)
@@ -617,33 +634,19 @@ void UPoseSearchLibrary::MotionMatchMulti(
 		AnimInstances.Add(AnimInstance);
 	}
 
-	MotionMatch(AnimInstances, Roles, Database, PoseHistoryName, FPoseSearchFutureProperties(), Result, DebugSessionUniqueIdentifier);
+	TArray<const UObject*>& AssetsToSearchConst = reinterpret_cast<TArray<const UObject*>&>(AssetsToSearch);
+	MotionMatch(AnimInstances, Roles, AssetsToSearchConst, PoseHistoryName, FPoseSearchFutureProperties(), Result, DebugSessionUniqueIdentifier);
 }
 
 void UPoseSearchLibrary::MotionMatch(
 	TArrayView<UAnimInstance*> AnimInstances,
-	TConstArrayView<FName> Roles,
-	const UPoseSearchDatabase* Database,
+	TArrayView<const UE::PoseSearch::FRole> Roles,
+	TArrayView<const UObject*> AssetsToSearch,
 	const FName PoseHistoryName,
 	const FPoseSearchFutureProperties& Future,
 	FPoseSearchBlueprintResult& Result,
 	const int32 DebugSessionUniqueIdentifier)
 {
-#if ENABLE_DRAW_DEBUG && ENABLE_ANIM_DEBUG
-	class UAnimInstanceProxyProvider : public UAnimInstance
-	{
-	public:
-		static FAnimInstanceProxy* GetAnimInstanceProxy(UAnimInstance* AnimInstance)
-		{
-			if (AnimInstance)
-			{
-				return &static_cast<UAnimInstanceProxyProvider*>(AnimInstance)->GetProxyOnAnyThread<FAnimInstanceProxy>();
-			}
-			return nullptr;
-		}
-	};
-#endif //ENABLE_DRAW_DEBUG && ENABLE_ANIM_DEBUG
-
 	using namespace UE::Anim;
 	using namespace UE::PoseSearch;
 
@@ -655,12 +658,6 @@ void UPoseSearchLibrary::MotionMatch(
 	Result.SelectedDatabase = nullptr;
 	Result.SearchCost = MAX_flt;
 
-	if (!Database)
-	{
-		UE_LOG(LogPoseSearch, Error, TEXT("UPoseSearchLibrary::MotionMatch - null Database"));
-		return;
-	}
-
 	if (AnimInstances.IsEmpty() || AnimInstances.Num() != Roles.Num())
 	{
 		UE_LOG(LogPoseSearch, Error, TEXT("UPoseSearchLibrary::MotionMatch - invalid input AnimInstances or Roles"));
@@ -721,12 +718,6 @@ void UPoseSearchLibrary::MotionMatch(
 	float FutureIntervalTime = Future.IntervalTime;
 	if (Future.Animation)
 	{
-		if (AnimInstances.Num() != 1)
-		{
-			UE_LOG(LogPoseSearch, Error, TEXT("UPoseSearchLibrary::MotionMatch - FutureAnimation not supported for more than one AnimInstances"));
-			return;
-		}
-		
 		float FutureAnimationTime = Future.AnimationTime;
 		if (FutureAnimationTime < FiniteDelta)
 		{
@@ -741,42 +732,57 @@ void UPoseSearchLibrary::MotionMatch(
 			FutureIntervalTime = MinFutureIntervalTime;
 		}
 
-		// extracting 2 poses to be able to calculate velocities
-		FCSPose<FCompactPose> ComponentSpacePose;
-		FCompactPose Pose;
-		Pose.SetBoneContainer(&AnimInstances[0]->GetRequiredBonesOnAnyThread());
-		
-		// @todo: add input BlendParameters to support sampling FutureAnimation blendspaces and support for multi character
-		const FAnimationAssetSampler Sampler(Future.Animation);
-		for (int32 i = 0; i < 2; ++i)
+		for (int32 RoleIndex = 0; RoleIndex < Roles.Num(); ++RoleIndex)
 		{
-			const float FuturePoseExtractionTime = FutureAnimationTime + (i - 1) * FiniteDelta;
-			const float FuturePoseAnimationTime = FutureIntervalTime + (i - 1) * FiniteDelta;
-				
-			Sampler.ExtractPose(FuturePoseExtractionTime, Pose);
-			ComponentSpacePose.InitPose(Pose);
-			MemStackPoseHistories[0].AddFuturePose(FuturePoseAnimationTime, ComponentSpacePose);
-		}
+			// extracting 2 poses to be able to calculate velocities
+			FCSPose<FCompactPose> ComponentSpacePose;
+			FCompactPose Pose;
+			Pose.SetBoneContainer(&AnimInstances[0]->GetRequiredBonesOnAnyThread());
+
+			// @todo: add input BlendParameters to support sampling FutureAnimation blendspaces and support for multi character
+			const UAnimationAsset* AnimationAsset = Cast<UAnimationAsset>(Future.Animation);
+			if (!AnimationAsset)
+			{
+				if (const UPoseSearchMultiSequence* MultiSequence = Cast<UPoseSearchMultiSequence>(Future.Animation))
+				{
+					AnimationAsset = MultiSequence->GetSequence(Roles[RoleIndex]);
+				}
+				else
+				{
+					checkNoEntry();
+				}
+			}
+
+			const FAnimationAssetSampler Sampler(AnimationAsset);
+			for (int32 i = 0; i < 2; ++i)
+			{
+				const float FuturePoseExtractionTime = FutureAnimationTime + (i - 1) * FiniteDelta;
+				const float FuturePoseAnimationTime = FutureIntervalTime + (i - 1) * FiniteDelta;
+
+				Sampler.ExtractPose(FuturePoseExtractionTime, Pose);
+				ComponentSpacePose.InitPose(Pose);
+				MemStackPoseHistories[RoleIndex].AddFuturePose(FuturePoseAnimationTime, ComponentSpacePose);
+			}
 
 #if ENABLE_DRAW_DEBUG && ENABLE_ANIM_DEBUG
-		if (FAnimInstanceProxy* AnimInstanceProxy = UAnimInstanceProxyProvider::GetAnimInstanceProxy(AnimInstances[0]))
-		{
-			MemStackPoseHistories[0].DebugDraw(*AnimInstanceProxy, HistoryCollectorColors[0]);
-		}
+			if (FAnimInstanceProxy* AnimInstanceProxy = UAnimInstanceProxyProvider::GetAnimInstanceProxy(AnimInstances[RoleIndex]))
+			{
+				MemStackPoseHistories[RoleIndex].DebugDraw(*AnimInstanceProxy, HistoryCollectorColors[RoleIndex]);
+			}
 #endif // ENABLE_DRAW_DEBUG && ENABLE_ANIM_DEBUG
+		}
 	}
 
-	FSearchContext SearchContext(FutureIntervalTime);
-	for (int32 RoleIndex = 0; RoleIndex < Roles.Num(); ++RoleIndex)
+	TArray<const UE::PoseSearch::IPoseHistory*, TInlineAllocator<PreallocatedRolesNum, TMemStackAllocator<>>> PoseHistories;
+	for (const FMemStackPoseHistory& MemStackPoseHistory : MemStackPoseHistories)
 	{
-		SearchContext.AddRole(Roles[RoleIndex], AnimInstances[RoleIndex], MemStackPoseHistories[RoleIndex].GetThisOrPoseHistory());
+		PoseHistories.Add(MemStackPoseHistory.GetThisOrPoseHistory());
 	}
 
-	FSearchResult SearchResult = Database->Search(SearchContext);
+	const FSearchResult SearchResult = MotionMatch(AnimInstances, Roles, PoseHistories, AssetsToSearch, nullptr, 0.f, DebugSessionUniqueIdentifier);
 	if (SearchResult.IsValid())
 	{
-		check(SearchResult.Database == Database);
-
+		const UPoseSearchDatabase* Database = SearchResult.Database.Get();
 		const FSearchIndexAsset* SearchIndexAsset = SearchResult.GetSearchIndexAsset();
 		if (const FPoseSearchDatabaseAnimationAssetBase* DatabaseAsset = Database->GetAnimationAssetBase(*SearchIndexAsset))
 		{
@@ -805,64 +811,52 @@ void UPoseSearchLibrary::MotionMatch(
 			}
 		}
 	}
-
-#if ENABLE_DRAW_DEBUG && ENABLE_ANIM_DEBUG
-	if (SearchResult.IsValid())
-	{
-		const bool bDrawMatch = CVarAnimMotionMatchDrawMatchEnable.GetValueOnAnyThread();
-		const bool bDrawquery = CVarAnimMotionMatchDrawQueryEnable.GetValueOnAnyThread();
-
-		if (bDrawMatch || bDrawquery)
-		{
-			TArray<FAnimInstanceProxy*, TInlineAllocator<PreallocatedRolesNum, TMemStackAllocator<>>> AnimInstanceProxies;
-			AnimInstanceProxies.SetNum(Roles.Num());
-			
-			for (int32 RoleIndex = 0; RoleIndex < Roles.Num(); ++RoleIndex)
-			{
-				AnimInstanceProxies[RoleIndex] = UAnimInstanceProxyProvider::GetAnimInstanceProxy(AnimInstances[RoleIndex]);
-			}
-
-			if (bDrawMatch)
-			{
-				FDebugDrawParams DrawParams(AnimInstanceProxies, SearchContext.GetPoseHistories(), SearchContext.GetRoleToIndex(), SearchResult.Database.Get());
-				DrawParams.DrawFeatureVector(SearchResult.PoseIdx);
-			}
-
-			if (bDrawquery)
-			{
-				FDebugDrawParams DrawParams(AnimInstanceProxies, SearchContext.GetPoseHistories(), SearchContext.GetRoleToIndex(), SearchResult.Database.Get(), EDebugDrawFlags::DrawQuery);
-				DrawParams.DrawFeatureVector(SearchContext.GetOrBuildQuery(SearchResult.Database->Schema));
-			}
-		}
-	}
-#endif // ENABLE_DRAW_DEBUG && ENABLE_ANIM_DEBUG
-
-#if UE_POSE_SEARCH_TRACE_ENABLED
-	TraceMotionMatchingState(SearchContext, SearchResult, 0.f, FTransform::Identity, DebugSessionUniqueIdentifier,
-		AnimInstances[0]->GetDeltaSeconds(), true, FObjectTrace::GetWorldElapsedTime(AnimInstances[0]->GetWorld()));
-#endif // UE_POSE_SEARCH_TRACE_ENABLED
 }
 
-UE::PoseSearch::FSearchResult UPoseSearchLibrary::MotionMatch(const FAnimationBaseContext& Context, 
-	TConstArrayView<UObject*> AssetsToSearch, const UObject* PlayingAsset, float PlayingAssetAccumulatedTime)
+UE::PoseSearch::FSearchResult UPoseSearchLibrary::MotionMatch(
+	const FAnimationBaseContext& Context,
+	TArrayView<const UObject*> AssetsToSearch,
+	const UObject* PlayingAsset,
+	float PlayingAssetAccumulatedTime)
 {
 	using namespace UE::PoseSearch;
 
-	FSearchResult SearchResult;
-
 	const IPoseHistory* PoseHistory = nullptr;
 	if (IPoseHistoryProvider* PoseHistoryProvider = Context.GetMessage<IPoseHistoryProvider>())
 	{
 		PoseHistory = &PoseHistoryProvider->GetPoseHistory();
 	}
 
-	const UAnimInstance* AnimInstance = Cast<const UAnimInstance>(Context.AnimInstanceProxy->GetAnimInstanceObject());
+	UAnimInstance* AnimInstance = Cast<UAnimInstance>(Context.AnimInstanceProxy->GetAnimInstanceObject());
 	check(AnimInstance);
 
+	return MotionMatch(MakeArrayView(&AnimInstance, 1), MakeArrayView(&DefaultRole, 1), MakeArrayView(&PoseHistory, 1), AssetsToSearch, PlayingAsset, PlayingAssetAccumulatedTime, Context.GetCurrentNodeId());
+}
+	
+UE::PoseSearch::FSearchResult UPoseSearchLibrary::MotionMatch(
+	TArrayView<UAnimInstance*> AnimInstances,
+	TArrayView<const UE::PoseSearch::FRole> Roles,
+	TArrayView<const UE::PoseSearch::IPoseHistory*> PoseHistories, 
+	TArrayView<const UObject*> AssetsToSearch,
+	const UObject* PlayingAsset,
+	float PlayingAssetAccumulatedTime,
+	const int32 DebugSessionUniqueIdentifier)
+{
+	check(!AnimInstances.IsEmpty() && AnimInstances.Num() == Roles.Num() && AnimInstances.Num() == PoseHistories.Num());
+
+	using namespace UE::PoseSearch;
+
+	FSearchResult SearchResult;
+
 	FMemMark Mark(FMemStack::Get());
 	FSearchResult ReconstructedPreviousSearchResult;
 	FSearchContext SearchContext(0.f, nullptr, ReconstructedPreviousSearchResult);
-	SearchContext.AddRole(DefaultRole, AnimInstance, PoseHistory);
+
+	// @todo: all assets in AssetsToSearch should have a consistent Roles requirements, or else the search will throw an error!
+	for (int32 RoleIndex = 0; RoleIndex < Roles.Num(); ++RoleIndex)
+	{
+		SearchContext.AddRole(Roles[RoleIndex], AnimInstances[RoleIndex], PoseHistories[RoleIndex]);
+	}
 
 	// budgeting some stack allocations for simple use cases. bigger requests of AnimationAssets contining 
 	// UAnimNotifyState_PoseSearchBranchIn referencing multiple datbases will default to slower heap allocations
@@ -919,7 +913,7 @@ UE::PoseSearch::FSearchResult UPoseSearchLibrary::MotionMatch(const FAnimationBa
 	};
 
 	// collecting all the possible continuing pose search (it could be multiple searches, but most likely only one)
-	const float DeltaSeconds = AnimInstance->GetDeltaSeconds();
+	const float DeltaSeconds = AnimInstances[0]->GetDeltaSeconds();
 	if (const UAnimationAsset* PlayingAnimationAsset = Cast<UAnimationAsset>(PlayingAsset))
 	{
 		AddToSearch(AssetsToSearchPerDatabaseMap, PlayingAnimationAsset);
@@ -1010,8 +1004,13 @@ UE::PoseSearch::FSearchResult UPoseSearchLibrary::MotionMatch(const FAnimationBa
 
 		if (bDrawMatch || bDrawquery)
 		{
-			FAnimInstanceProxy* AnimInstanceProxy = Context.AnimInstanceProxy;
-			const TArrayView<FAnimInstanceProxy*> AnimInstanceProxies = MakeArrayView(&AnimInstanceProxy, 1);
+			TArray<FAnimInstanceProxy*, TInlineAllocator<PreallocatedRolesNum, TMemStackAllocator<>>> AnimInstanceProxies;
+			AnimInstanceProxies.SetNum(Roles.Num());
+			
+			for (int32 RoleIndex = 0; RoleIndex < Roles.Num(); ++RoleIndex)
+			{
+				AnimInstanceProxies[RoleIndex] = UAnimInstanceProxyProvider::GetAnimInstanceProxy(AnimInstances[RoleIndex]);
+			}
 
 			if (bDrawMatch)
 			{
@@ -1031,8 +1030,8 @@ UE::PoseSearch::FSearchResult UPoseSearchLibrary::MotionMatch(const FAnimationBa
 #if UE_POSE_SEARCH_TRACE_ENABLED
 	const float SearchBestCost = SearchResult.PoseCost.GetTotalCost();
 	const float SearchBruteForceCost = SearchResult.BruteForcePoseCost.GetTotalCost();
-	TraceMotionMatchingState(SearchContext, SearchResult, 0.f, FTransform::Identity, Context.GetCurrentNodeId(),
-		DeltaSeconds, true, FObjectTrace::GetWorldElapsedTime(AnimInstance->GetWorld()));
+	TraceMotionMatchingState(SearchContext, SearchResult, 0.f, FTransform::Identity, DebugSessionUniqueIdentifier,
+		DeltaSeconds, true, FObjectTrace::GetWorldElapsedTime(AnimInstances[0]->GetWorld()));
 #endif // UE_POSE_SEARCH_TRACE_ENABLED
 
 	return SearchResult;
diff --git a/Engine/Plugins/Experimental/Animation/PoseSearch/Source/Runtime/Private/PoseSearchModule.cpp b/Engine/Plugins/Experimental/Animation/PoseSearch/Source/Runtime/Private/PoseSearchModule.cpp
--- a/Engine/Plugins/Experimental/Animation/PoseSearch/Source/Runtime/Private/PoseSearchModule.cpp
+++ b/Engine/Plugins/Experimental/Animation/PoseSearch/Source/Runtime/Private/PoseSearchModule.cpp
@@ -22,7 +22,7 @@ public:
 	}
 
 	// IPoseSearchProvider
-	virtual UE::Anim::IPoseSearchProvider::FSearchResult Search(const FAnimationBaseContext& GraphContext, TConstArrayView<UObject*> AssetsToSearch, const UObject* PlayingAsset, float PlayingAssetAccumulatedTime) const override
+	virtual UE::Anim::IPoseSearchProvider::FSearchResult Search(const FAnimationBaseContext& GraphContext, TArrayView<const UObject*> AssetsToSearch, const UObject* PlayingAsset, float PlayingAssetAccumulatedTime) const override
 	{
 		const UE::PoseSearch::FSearchResult SearchResult = UPoseSearchLibrary::MotionMatch(GraphContext, AssetsToSearch, PlayingAsset, PlayingAssetAccumulatedTime);
 		UE::Anim::IPoseSearchProvider::FSearchResult ProviderResult;
diff --git a/Engine/Plugins/Experimental/Animation/PoseSearch/Source/Runtime/Public/PoseSearch/PoseSearchLibrary.h b/Engine/Plugins/Experimental/Animation/PoseSearch/Source/Runtime/Public/PoseSearch/PoseSearchLibrary.h
--- a/Engine/Plugins/Experimental/Animation/PoseSearch/Source/Runtime/Public/PoseSearch/PoseSearchLibrary.h
+++ b/Engine/Plugins/Experimental/Animation/PoseSearch/Source/Runtime/Public/PoseSearch/PoseSearchLibrary.h
@@ -96,7 +96,7 @@ struct POSESEARCH_API FPoseSearchFutureProperties
 public:
 	// Animation to play (it'll start at AnimationTime seconds)
 	UPROPERTY(EditAnywhere, BlueprintReadWrite, Category=State)
-	TObjectPtr<UAnimationAsset> Animation;
+	TObjectPtr<UObject> Animation;
 
 	// Start time for Animation
 	UPROPERTY(EditAnywhere, BlueprintReadWrite, Category=State)
@@ -164,45 +164,63 @@ public:
 	* Implementation of the core motion matching algorithm
 	*
 	* @param AnimInstance					Input animation instance
-	* @param Database						Input database to search
+	* @param AssetsToSearch					Input assets to search (UPoseSearchDatabase or any animation asset containing UAnimNotifyState_PoseSearchBranchIn)
 	* @param PoseHistoryName				Input tag of the associated PoseSearchHistoryCollector node in the anim graph
-	* @param SelectedAnimation				Output selected animation from the Database asset
-	* @param SelectedTime					Output selected animation time
-	* @param bLoop							Output selected animation looping state
-	* @param bIsMirrored					Output selected animation mirror state
-	* @param BlendParameters				Output selected animation blend space parameters (if SelectedAnimation is a blend space)
-	* @param SearchCost						Output search associated cost
 	* @param Future							Input future properties to match (animation / start time / time offset)
+	* @param SelectedAnimation				Output selected animation from the Database asset
+	* @param Result							Output FPoseSearchBlueprintResult with the search result
 	* @param DebugSessionUniqueIdentifier	Input unique identifier used to identify TraceMotionMatchingState (rewind debugger / pose search debugger) session. Similarly the MM node uses Context.GetCurrentNodeId()
 	*/
 	UFUNCTION(BlueprintPure, Category = "Animation|Pose Search|Experimental", meta = (BlueprintThreadSafe, Keywords = "PoseMatch"))
 	static void MotionMatch(
 		UAnimInstance* AnimInstance,
-		const UPoseSearchDatabase* Database,
+		TArray<UObject*> AssetsToSearch,
 		const FName PoseHistoryName,
 		FPoseSearchFutureProperties Future,
 		FPoseSearchBlueprintResult& Result,
 		const int32 DebugSessionUniqueIdentifier = 6174);
 
+	/**
+	* Implementation of the core motion matching algorithm for multiple characters
+	*
+	* @param AnimInstances					Input animation instances
+	* @param Roles							Input Roles associated to the animation instances
+	* @param AssetsToSearch					Input assets to search (UPoseSearchDatabase or any animation asset containing UAnimNotifyState_PoseSearchBranchIn)
+	* @param PoseHistoryName				Input tag of the associated PoseSearchHistoryCollector node in the anim graphs of the AnimInstances
+	* @param Result							Output FPoseSearchBlueprintResult with the search result
+	* @param DebugSessionUniqueIdentifier	Input unique identifier used to identify TraceMotionMatchingState (rewind debugger / pose search debugger) session. Similarly the MM node uses Context.GetCurrentNodeId()
+	*/
 	UFUNCTION(BlueprintPure, Category = "Animation|Pose Search|Experimental", meta = (BlueprintThreadSafe, Keywords = "PoseMatch"))
 	static void MotionMatchMulti(
 		TArray<ACharacter*> AnimInstances,
 		TArray<FName> Roles,
-		const UPoseSearchDatabase* Database,
+		TArray<UObject*> AssetsToSearch,
 		const FName PoseHistoryName,
 		FPoseSearchBlueprintResult& Result,
 		const int32 DebugSessionUniqueIdentifier = 6174);
 
 	static void MotionMatch(
 		TArrayView<UAnimInstance*> AnimInstances,
-		TConstArrayView<FName> Roles,
-		const UPoseSearchDatabase* Database,
+		TArrayView<const UE::PoseSearch::FRole> Roles,
+		TArrayView<const UObject*> AssetsToSearch,
 		const FName PoseHistoryName,
 		const FPoseSearchFutureProperties& Future,
 		FPoseSearchBlueprintResult& Result,
 		const int32 DebugSessionUniqueIdentifier);
 
-	static UE::PoseSearch::FSearchResult MotionMatch(const FAnimationBaseContext& Context, TConstArrayView<UObject*> AssetsToSearch,
-		const UObject* PlayingAsset = nullptr, float PlayingAssetAccumulatedTime = 0.f);
+	static UE::PoseSearch::FSearchResult MotionMatch(
+		const FAnimationBaseContext& Context,
+		TArrayView<const UObject*> AssetsToSearch,
+		const UObject* PlayingAsset = nullptr,
+		float PlayingAssetAccumulatedTime = 0.f);
+
+	static UE::PoseSearch::FSearchResult MotionMatch(
+		TArrayView<UAnimInstance*> AnimInstances,
+		TArrayView<const UE::PoseSearch::FRole> Roles,
+		TArrayView<const UE::PoseSearch::IPoseHistory*> PoseHistories, 
+		TArrayView<const UObject*> AssetsToSearch,
+		const UObject* PlayingAsset,
+		float PlayingAssetAccumulatedTime,
+		const int32 DebugSessionUniqueIdentifier);
 };
 
diff --git a/Engine/Plugins/Experimental/Chooser/Source/Chooser/Private/AnimNode_ChooserPlayer.cpp b/Engine/Plugins/Experimental/Chooser/Source/Chooser/Private/AnimNode_ChooserPlayer.cpp
--- a/Engine/Plugins/Experimental/Chooser/Source/Chooser/Private/AnimNode_ChooserPlayer.cpp
+++ b/Engine/Plugins/Experimental/Chooser/Source/Chooser/Private/AnimNode_ChooserPlayer.cpp
@@ -35,7 +35,7 @@ UAnimationAsset* FAnimNode_ChooserPlayer::ChooseAsset(const FAnimationUpdateCont
 		{
 			if (UE::Anim::IPoseSearchProvider* PoseSearchProvider = UE::Anim::IPoseSearchProvider::Get())
 			{
-				TArray<UObject*, TInlineAllocator<128>> ChosenAssets;
+				TArray<const UObject*, TInlineAllocator<128>> ChosenAssets;
 				ChooserBase.ChooseMulti(ChooserContext, FObjectChooserBase::FObjectChooserIteratorCallback::CreateLambda([&ChosenAssets](UObject* InResult)
 					{
 						ChosenAssets.Add(InResult);
diff --git a/Engine/Source/Runtime/Engine/Private/Animation/AnimNode_SequencePlayer.cpp b/Engine/Source/Runtime/Engine/Private/Animation/AnimNode_SequencePlayer.cpp
--- a/Engine/Source/Runtime/Engine/Private/Animation/AnimNode_SequencePlayer.cpp
+++ b/Engine/Source/Runtime/Engine/Private/Animation/AnimNode_SequencePlayer.cpp
@@ -159,7 +159,7 @@ float FAnimNode_SequencePlayerBase::GetTimeFromEnd(float CurrentNodeTime) const
 float FAnimNode_SequencePlayerBase::GetEffectiveStartPosition(const FAnimationBaseContext& Context) const
 {
 	// Override the start position if pose matching is enabled
-	UObject* CurrentSequence = GetSequence();
+	const UObject* CurrentSequence = GetSequence();
 	if (CurrentSequence != nullptr && GetStartFromMatchingPose())
 	{
 		if (const UE::Anim::IPoseSearchProvider* PoseSearchProvider = UE::Anim::IPoseSearchProvider::Get())
diff --git a/Engine/Source/Runtime/Engine/Public/Animation/AnimPoseSearchProvider.h b/Engine/Source/Runtime/Engine/Public/Animation/AnimPoseSearchProvider.h
--- a/Engine/Source/Runtime/Engine/Public/Animation/AnimPoseSearchProvider.h
+++ b/Engine/Source/Runtime/Engine/Public/Animation/AnimPoseSearchProvider.h
@@ -40,7 +40,7 @@ public:
 	* 
 	* @return	FSearchResult				The search result identifying the asset from AssetsToSearch or PlayingAsset that most closely matches the query
 	*/
-	virtual FSearchResult Search(const FAnimationBaseContext& GraphContext, TConstArrayView<UObject*> AssetsToSearch,
+	virtual FSearchResult Search(const FAnimationBaseContext& GraphContext, TArrayView<const UObject*> AssetsToSearch,
 		const UObject* PlayingAsset = nullptr, float PlayingAssetAccumulatedTime = 0.f) const = 0;
 };
 
