jasper - MM trajectory world collisitons handling WIP

diff --git a/Engine/Plugins/Experimental/Animation/PoseSearch/Source/Runtime/Private/AnimNode_PoseSearchHistoryCollector.cpp b/Engine/Plugins/Experimental/Animation/PoseSearch/Source/Runtime/Private/AnimNode_PoseSearchHistoryCollector.cpp
--- a/Engine/Plugins/Experimental/Animation/PoseSearch/Source/Runtime/Private/AnimNode_PoseSearchHistoryCollector.cpp
+++ b/Engine/Plugins/Experimental/Animation/PoseSearch/Source/Runtime/Private/AnimNode_PoseSearchHistoryCollector.cpp
@@ -91,7 +91,7 @@ void FAnimNode_PoseSearchHistoryCollector_Base::Update_AnyThread(const FAnimatio
 {
 	GetEvaluateGraphExposedInputs().Execute(Context);
 
-	PoseHistory.Update_AnyThread(Context.GetDeltaTime(), bGenerateTrajectory ? FPoseSearchQueryTrajectory() : Trajectory, TrajectorySpeedMultiplier);
+	PoseHistory.SetTrajectory(bGenerateTrajectory ? FPoseSearchQueryTrajectory() : Trajectory, TrajectorySpeedMultiplier);
 
 	UpdateCounter.SynchronizeWith(Context.AnimInstanceProxy->GetUpdateCounter());
 }
diff --git a/Engine/Plugins/Experimental/Animation/PoseSearch/Source/Runtime/Private/PoseSearchHistory.cpp b/Engine/Plugins/Experimental/Animation/PoseSearch/Source/Runtime/Private/PoseSearchHistory.cpp
--- a/Engine/Plugins/Experimental/Animation/PoseSearch/Source/Runtime/Private/PoseSearchHistory.cpp
+++ b/Engine/Plugins/Experimental/Animation/PoseSearch/Source/Runtime/Private/PoseSearchHistory.cpp
@@ -557,7 +557,7 @@ void FPoseHistory::PreUpdate(const UAnimInstance* AnimInstance, float DeltaTime,
 	ReadData = WriteData;
 }
 
-void FPoseHistory::Update_AnyThread(float DeltaTime, const FPoseSearchQueryTrajectory& InTrajectory, float InTrajectorySpeedMultiplier)
+void FPoseHistory::SetTrajectory(const FPoseSearchQueryTrajectory& InTrajectory, float InTrajectorySpeedMultiplier)
 {
 	if (!InTrajectory.Samples.IsEmpty())
 	{
diff --git a/Engine/Plugins/Experimental/Animation/PoseSearch/Source/Runtime/Private/PoseSearchHistoryCollectorAnimNodeLibrary.cpp b/Engine/Plugins/Experimental/Animation/PoseSearch/Source/Runtime/Private/PoseSearchHistoryCollectorAnimNodeLibrary.cpp
--- a/Engine/Plugins/Experimental/Animation/PoseSearch/Source/Runtime/Private/PoseSearchHistoryCollectorAnimNodeLibrary.cpp
+++ b/Engine/Plugins/Experimental/Animation/PoseSearch/Source/Runtime/Private/PoseSearchHistoryCollectorAnimNodeLibrary.cpp
@@ -14,4 +14,12 @@ void UPoseSearchHistoryCollectorAnimNodeLibrary::GetPoseHistoryNodeTrajectory(co
 	{
 		Trajectory = PoseSearchHistoryCollectorNodePtr->GetPoseHistory().GetTrajectory();
 	}
-}
\ No newline at end of file
+}
+
+void UPoseSearchHistoryCollectorAnimNodeLibrary::SetPoseHistoryNodeTrajectory(const FPoseSearchHistoryCollectorAnimNodeReference& PoseSearchHistoryCollectorNode, const FPoseSearchQueryTrajectory& Trajectory)
+{
+	if (FAnimNode_PoseSearchHistoryCollector* PoseSearchHistoryCollectorNodePtr = PoseSearchHistoryCollectorNode.GetAnimNodePtr<FAnimNode_PoseSearchHistoryCollector>())
+	{
+		PoseSearchHistoryCollectorNodePtr->GetPoseHistory().SetTrajectory(Trajectory);
+	}
+}
diff --git a/Engine/Plugins/Experimental/Animation/PoseSearch/Source/Runtime/Private/PoseSearchTrajectoryLibrary.cpp b/Engine/Plugins/Experimental/Animation/PoseSearch/Source/Runtime/Private/PoseSearchTrajectoryLibrary.cpp
--- a/Engine/Plugins/Experimental/Animation/PoseSearch/Source/Runtime/Private/PoseSearchTrajectoryLibrary.cpp
+++ b/Engine/Plugins/Experimental/Animation/PoseSearch/Source/Runtime/Private/PoseSearchTrajectoryLibrary.cpp
@@ -6,12 +6,20 @@
 #include "GameFramework/Character.h"
 #include "GameFramework/CharacterMovementComponent.h"
 
-void FPoseSearchTrajectoryData::UpdateData(float DeltaTime, const FAnimInstanceProxy& AnimInstanceProxy, FDerived& TrajectoryDataDerived, FState& TrajectoryDataState) const
+void FPoseSearchTrajectoryData::UpdateData(
+	float DeltaTime,
+	const FAnimInstanceProxy& AnimInstanceProxy,
+	FDerived& TrajectoryDataDerived,
+	FState& TrajectoryDataState) const
 {
 	UpdateData(DeltaTime, Cast<const UAnimInstance>(AnimInstanceProxy.GetAnimInstanceObject()), TrajectoryDataDerived, TrajectoryDataState);
 }
 
-void FPoseSearchTrajectoryData::UpdateData(float DeltaTime, const UAnimInstance* AnimInstance, FDerived& TrajectoryDataDerived, FState& TrajectoryDataState) const
+void FPoseSearchTrajectoryData::UpdateData(
+	float DeltaTime,
+	const UAnimInstance* AnimInstance,
+	FDerived& TrajectoryDataDerived,
+	FState& TrajectoryDataState) const
 {
 	// An AnimInstance might call this during an AnimBP recompile with 0 delta time.
 	if (DeltaTime <= 0.f || !AnimInstance)
@@ -34,7 +42,7 @@ void FPoseSearchTrajectoryData::UpdateData(float DeltaTime, const UAnimInstance*
 		TrajectoryDataDerived.Velocity = MoveComp->Velocity;
 		TrajectoryDataDerived.Acceleration = MoveComp->GetCurrentAcceleration();
 		
-		TrajectoryDataDerived.bIsFalling = MoveComp->IsFalling();
+		TrajectoryDataDerived.bStepGroundPrediction = !MoveComp->IsFalling() && !MoveComp->IsFlying();
 		TrajectoryDataDerived.Gravity = -MoveComp->GetGravityDirection() * MoveComp->GetGravityZ();
 
 		if (TrajectoryDataDerived.Acceleration.IsZero())
@@ -73,7 +81,11 @@ void FPoseSearchTrajectoryData::UpdateData(float DeltaTime, const UAnimInstance*
 	}
 }
 
-FVector FPoseSearchTrajectoryData::StepCharacterMovementGroundPrediction(float DeltaTime, const FVector& InVelocity, const FVector& InAcceleration, const FDerived& TrajectoryDataDerived) const
+FVector FPoseSearchTrajectoryData::StepCharacterMovementGroundPrediction(
+	float DeltaTime,
+	const FVector& InVelocity,
+	const FVector& InAcceleration,
+	const FDerived& TrajectoryDataDerived) const
 {
 	FVector OutVelocity = InVelocity;
 
@@ -139,23 +151,11 @@ FVector FPoseSearchTrajectoryData::StepCharacterMovementGroundPrediction(float D
 	return OutVelocity;
 }
 
-FVector FPoseSearchTrajectoryData::StepCharacterFallingPrediction(float DeltaTime, const FVector& InitialVelocity, const FDerived& TrajectoryDataDerived) const
-{
-	FVector Result = InitialVelocity;
-
-	if (DeltaTime > 0.f)
-	{
-		// Apply gravity.
-		Result += TrajectoryDataDerived.Gravity * DeltaTime;
-	}
-
-	// @todo: Support for air control to give more accurate results.
-	// @todo: Support for collision prediction to properly predict future on ground path after landing.
-	
-	return Result;
-}
-
-void UPoseSearchTrajectoryLibrary::InitTrajectorySamples(FPoseSearchQueryTrajectory& Trajectory, const FPoseSearchTrajectoryData& TrajectoryData, const FPoseSearchTrajectoryData::FDerived& TrajectoryDataDerived, const FPoseSearchTrajectoryData::FSampling& TrajectoryDataSampling)
+void UPoseSearchTrajectoryLibrary::InitTrajectorySamples(
+	FPoseSearchQueryTrajectory& Trajectory,
+	const FPoseSearchTrajectoryData& TrajectoryData,
+	const FPoseSearchTrajectoryData::FDerived& TrajectoryDataDerived,
+	const FPoseSearchTrajectoryData::FSampling& TrajectoryDataSampling)
 {
 	const int32 NumHistorySamples = TrajectoryDataSampling.NumHistorySamples;
 	const int32 NumPredictionSamples = TrajectoryDataSampling.NumPredictionSamples;
@@ -187,7 +187,12 @@ void UPoseSearchTrajectoryLibrary::InitTrajectorySamples(FPoseSearchQueryTraject
 	}
 }
 
-void UPoseSearchTrajectoryLibrary::UpdateHistory_TransformHistory(FPoseSearchQueryTrajectory& Trajectory, const FPoseSearchTrajectoryData& TrajectoryData, const FPoseSearchTrajectoryData::FDerived& TrajectoryDataDerived, const FPoseSearchTrajectoryData::FSampling& TrajectoryDataSampling, float DeltaTime)
+void UPoseSearchTrajectoryLibrary::UpdateHistory_TransformHistory(
+	FPoseSearchQueryTrajectory& Trajectory,
+	const FPoseSearchTrajectoryData& TrajectoryData,
+	const FPoseSearchTrajectoryData::FDerived& TrajectoryDataDerived,
+	const FPoseSearchTrajectoryData::FSampling& TrajectoryDataSampling,
+	float DeltaTime)
 {
 	const int32 NumHistorySamples = TrajectoryDataSampling.NumHistorySamples;
 	const float SecondsPerHistorySample = TrajectoryDataSampling.SecondsPerHistorySample;
@@ -232,7 +237,10 @@ void UPoseSearchTrajectoryLibrary::UpdateHistory_TransformHistory(FPoseSearchQue
 	}
 }
 
-FVector UPoseSearchTrajectoryLibrary::RemapVectorMagnitudeWithCurve(const FVector& Vector, bool bUseCurve, const FRuntimeFloatCurve& Curve)
+FVector UPoseSearchTrajectoryLibrary::RemapVectorMagnitudeWithCurve(
+	const FVector& Vector,
+	bool bUseCurve,
+	const FRuntimeFloatCurve& Curve)
 {
 	if (bUseCurve)
 	{
@@ -247,82 +255,12 @@ FVector UPoseSearchTrajectoryLibrary::RemapVectorMagnitudeWithCurve(const FVecto
 	return Vector;
 }
 
-void UPoseSearchTrajectoryLibrary::UpdatePrediction_SimulateCharacterMovement(FPoseSearchQueryTrajectory& Trajectory, const FPoseSearchTrajectoryData& TrajectoryData, const FPoseSearchTrajectoryData::FDerived& TrajectoryDataDerived, const FPoseSearchTrajectoryData::FSampling& TrajectoryDataSampling)
-{
-	FVector CurrentPositionWS = TrajectoryDataDerived.Position;
-	FVector CurrentVelocityWS = RemapVectorMagnitudeWithCurve(TrajectoryDataDerived.Velocity, TrajectoryData.bUseSpeedRemappingCurve, TrajectoryData.SpeedRemappingCurve);
-	FVector CurrentAccelerationWS = RemapVectorMagnitudeWithCurve(TrajectoryDataDerived.Acceleration, TrajectoryData.bUseAccelerationRemappingCurve, TrajectoryData.AccelerationRemappingCurve);
-
-	// bending CurrentVelocityWS towards CurrentAccelerationWS
-	if (TrajectoryData.BendVelocityTowardsAcceleration > UE_KINDA_SMALL_NUMBER && !CurrentAccelerationWS.IsNearlyZero())
-	{
-		const float CurrentSpeed = CurrentVelocityWS.Length();
-		const FVector VelocityWSAlongAcceleration = CurrentAccelerationWS.GetUnsafeNormal() * CurrentSpeed;
-		if (TrajectoryData.BendVelocityTowardsAcceleration < 1.f - UE_KINDA_SMALL_NUMBER)
-		{
-			CurrentVelocityWS = FMath::Lerp(CurrentVelocityWS, VelocityWSAlongAcceleration, TrajectoryData.BendVelocityTowardsAcceleration);
-
-			const float NewLength = CurrentVelocityWS.Length();
-			if (NewLength > UE_KINDA_SMALL_NUMBER)
-			{
-				CurrentVelocityWS *= CurrentSpeed / NewLength;
-			}
-			else
-			{
-				// @todo: consider setting the CurrentVelocityWS = VelocityWSAlongAcceleration if vel and acc are in opposite directions
-			}
-		}
-		else
-		{
-			CurrentVelocityWS = VelocityWSAlongAcceleration;
-		}
-	}
-
-	FQuat CurrentFacingWS = TrajectoryDataDerived.Facing;
-	
-	const int32 NumHistorySamples = TrajectoryDataSampling.NumHistorySamples;
-	const float SecondsPerPredictionSample = TrajectoryDataSampling.SecondsPerPredictionSample;
-	const FQuat ControllerRotationPerStep = FQuat::MakeFromEuler(FVector(0.f, 0.f, TrajectoryDataDerived.ControllerYawRate * SecondsPerPredictionSample));
-
-	float AccumulatedSeconds = 0.f;
-
-	const int32 LastIndex = Trajectory.Samples.Num() - 1;
-	if (NumHistorySamples <= LastIndex)
-	{
-		for (int32 Index = NumHistorySamples; ; ++Index)
-		{
-			Trajectory.Samples[Index].Position = CurrentPositionWS;
-			Trajectory.Samples[Index].Facing = CurrentFacingWS;
-			Trajectory.Samples[Index].AccumulatedSeconds = AccumulatedSeconds;
-
-			if (Index == LastIndex)
-			{
-				break;
-			}
-
-			CurrentPositionWS += CurrentVelocityWS * SecondsPerPredictionSample;
-			AccumulatedSeconds += SecondsPerPredictionSample;
-
-			// Account for the controller (e.g. the camera) rotating.
-			CurrentFacingWS = ControllerRotationPerStep * CurrentFacingWS;
-			CurrentAccelerationWS = RemapVectorMagnitudeWithCurve(ControllerRotationPerStep * CurrentAccelerationWS,
-				TrajectoryData.bUseAccelerationRemappingCurve, TrajectoryData.AccelerationRemappingCurve);
-
-			const FVector NewVelocityWS = TrajectoryData.StepCharacterMovementGroundPrediction(SecondsPerPredictionSample, CurrentVelocityWS, CurrentAccelerationWS, TrajectoryDataDerived);
-
-			CurrentVelocityWS = RemapVectorMagnitudeWithCurve(NewVelocityWS, TrajectoryData.bUseSpeedRemappingCurve, TrajectoryData.SpeedRemappingCurve);
-
-			if (TrajectoryDataDerived.bOrientRotationToMovement && !CurrentAccelerationWS.IsNearlyZero())
-			{
-				// Rotate towards acceleration.
-				const FVector CurrentAccelerationCS = TrajectoryDataDerived.MeshCompRelativeRotation.RotateVector(CurrentAccelerationWS);
-				CurrentFacingWS = FMath::QInterpConstantTo(CurrentFacingWS, CurrentAccelerationCS.ToOrientationQuat(), SecondsPerPredictionSample, TrajectoryData.RotateTowardsMovementSpeed);
-			}
-		}
-	}
-}
-
-void UPoseSearchTrajectoryLibrary::UpdatePrediction_SimulateCharacterMovementAdvanced(FPoseSearchQueryTrajectory& Trajectory, const FPoseSearchTrajectoryData& TrajectoryData, const FPoseSearchTrajectoryData::FDerived& TrajectoryDataDerived, const FPoseSearchTrajectoryData::FSampling& TrajectoryDataSampling)
+void UPoseSearchTrajectoryLibrary::UpdatePrediction_SimulateCharacterMovement(
+	FPoseSearchQueryTrajectory& Trajectory,
+	const FPoseSearchTrajectoryData& TrajectoryData,
+	const FPoseSearchTrajectoryData::FDerived& TrajectoryDataDerived,
+	const FPoseSearchTrajectoryData::FSampling& TrajectoryDataSampling,
+	bool bAlwaysApplyGravity)
 {
 	FVector CurrentPositionWS = TrajectoryDataDerived.Position;
 	FVector CurrentVelocityWS = RemapVectorMagnitudeWithCurve(TrajectoryDataDerived.Velocity, TrajectoryData.bUseSpeedRemappingCurve, TrajectoryData.SpeedRemappingCurve);
@@ -378,37 +316,47 @@ void UPoseSearchTrajectoryLibrary::UpdatePrediction_SimulateCharacterMovementAdv
 			CurrentPositionWS += CurrentVelocityWS * SecondsPerPredictionSample;
 			AccumulatedSeconds += SecondsPerPredictionSample;
 
-			// Account for the controller (e.g. the camera) rotating.
-			CurrentFacingWS = ControllerRotationPerStep * CurrentFacingWS;
-			CurrentAccelerationWS = RemapVectorMagnitudeWithCurve(ControllerRotationPerStep * CurrentAccelerationWS,
-			TrajectoryData.bUseAccelerationRemappingCurve, TrajectoryData.AccelerationRemappingCurve);
-
-			// Perform trajectory predictions (jumping, moving, etc).
+			if (TrajectoryDataDerived.bStepGroundPrediction)
 			{
-				FVector NewVelocityWS = TrajectoryData.StepCharacterMovementGroundPrediction(SecondsPerPredictionSample, CurrentVelocityWS, CurrentAccelerationWS, TrajectoryDataDerived);
+				CurrentAccelerationWS = RemapVectorMagnitudeWithCurve(ControllerRotationPerStep * CurrentAccelerationWS,
+					TrajectoryData.bUseAccelerationRemappingCurve, TrajectoryData.AccelerationRemappingCurve);
+				const FVector NewVelocityWS = TrajectoryData.StepCharacterMovementGroundPrediction(SecondsPerPredictionSample, CurrentVelocityWS, CurrentAccelerationWS, TrajectoryDataDerived);
+				CurrentVelocityWS = RemapVectorMagnitudeWithCurve(NewVelocityWS, TrajectoryData.bUseSpeedRemappingCurve, TrajectoryData.SpeedRemappingCurve);
 
-				if (TrajectoryDataDerived.bIsFalling)
+				// Account for the controller (e.g. the camera) rotating.
+				CurrentFacingWS = ControllerRotationPerStep * CurrentFacingWS;
+				if (TrajectoryDataDerived.bOrientRotationToMovement && !CurrentAccelerationWS.IsNearlyZero())
 				{
-					NewVelocityWS = TrajectoryData.StepCharacterFallingPrediction(SecondsPerPredictionSample, CurrentVelocityWS, TrajectoryDataDerived);
+					// Rotate towards acceleration.
+					const FVector CurrentAccelerationCS = TrajectoryDataDerived.MeshCompRelativeRotation.RotateVector(CurrentAccelerationWS);
+					CurrentFacingWS = FMath::QInterpConstantTo(CurrentFacingWS, CurrentAccelerationCS.ToOrientationQuat(), SecondsPerPredictionSample, TrajectoryData.RotateTowardsMovementSpeed);
 				}
-				
-				CurrentVelocityWS = RemapVectorMagnitudeWithCurve(NewVelocityWS, TrajectoryData.bUseSpeedRemappingCurve, TrajectoryData.SpeedRemappingCurve);
 			}
-		
-			if (TrajectoryDataDerived.bOrientRotationToMovement && !CurrentAccelerationWS.IsNearlyZero())
+		}
+
+		if (!TrajectoryDataDerived.bStepGroundPrediction || bAlwaysApplyGravity)
+		{
+			for (int32 Index = NumHistorySamples + 1; Index < Trajectory.Samples.Num(); ++Index)
 			{
-				// Rotate towards acceleration.
-				const FVector CurrentAccelerationCS = TrajectoryDataDerived.MeshCompRelativeRotation.RotateVector(CurrentAccelerationWS);
-				CurrentFacingWS = FMath::QInterpConstantTo(CurrentFacingWS, CurrentAccelerationCS.ToOrientationQuat(), SecondsPerPredictionSample, TrajectoryData.RotateTowardsMovementSpeed);
+				const float Time = Trajectory.Samples[Index].AccumulatedSeconds;
+				Trajectory.Samples[Index].Position += TrajectoryDataDerived.Gravity * (Time * Time);
 			}
 		}
 	}
 }
 
 void UPoseSearchTrajectoryLibrary::PoseSearchGenerateTrajectory(
-	const UAnimInstance* InAnimInstance, UPARAM(ref) const FPoseSearchTrajectoryData& InTrajectoryData, float InDeltaTime,
-	UPARAM(ref) FPoseSearchQueryTrajectory& InOutTrajectory, UPARAM(ref) float& InOutDesiredControllerYawLastUpdate, FPoseSearchQueryTrajectory& OutTrajectory,
-	float InHistorySamplingInterval, int32 InTrajectoryHistoryCount, float InPredictionSamplingInterval, int32 InTrajectoryPredictionCount)
+	const UAnimInstance* InAnimInstance, UPARAM(ref)
+	const FPoseSearchTrajectoryData& InTrajectoryData,
+	float InDeltaTime,
+	UPARAM(ref) FPoseSearchQueryTrajectory& InOutTrajectory,
+	UPARAM(ref) float& InOutDesiredControllerYawLastUpdate,
+	FPoseSearchQueryTrajectory& OutTrajectory,
+	float InHistorySamplingInterval,
+	int32 InTrajectoryHistoryCount,
+	float InPredictionSamplingInterval,
+	int32 InTrajectoryPredictionCount,
+	bool bAlwaysApplyGravity)
 {
 	FPoseSearchTrajectoryData::FSampling TrajectoryDataSampling;
 	TrajectoryDataSampling.NumHistorySamples = InTrajectoryHistoryCount;
@@ -423,35 +371,39 @@ void UPoseSearchTrajectoryLibrary::PoseSearchGenerateTrajectory(
 	InTrajectoryData.UpdateData(InDeltaTime, InAnimInstance, TrajectoryDataDerived, TrajectoryDataState);
 	InitTrajectorySamples(InOutTrajectory, InTrajectoryData, TrajectoryDataDerived, TrajectoryDataSampling);
 	UpdateHistory_TransformHistory(InOutTrajectory, InTrajectoryData, TrajectoryDataDerived, TrajectoryDataSampling, InDeltaTime);
-	UpdatePrediction_SimulateCharacterMovement(InOutTrajectory, InTrajectoryData, TrajectoryDataDerived, TrajectoryDataSampling);
+	UpdatePrediction_SimulateCharacterMovement(InOutTrajectory, InTrajectoryData, TrajectoryDataDerived, TrajectoryDataSampling, bAlwaysApplyGravity);
 
 	InOutDesiredControllerYawLastUpdate = TrajectoryDataState.DesiredControllerYawLastUpdate;
 
 	OutTrajectory = InOutTrajectory;
 }
 
-void UPoseSearchTrajectoryLibrary::PoseSearchGenerateTrajectoryAdvanced(const UAnimInstance* InAnimInstance, const FPoseSearchTrajectoryData& InTrajectoryData, float InDeltaTime,
-	FPoseSearchQueryTrajectory& InOutTrajectory, float& InOutDesiredControllerYawLastUpdate, FPoseSearchQueryTrajectory& OutTrajectory, float InHistorySamplingInterval,
-	int32 InTrajectoryHistoryCount, float InPredictionSamplingInterval, int32 InTrajectoryPredictionCount)
+void UPoseSearchTrajectoryLibrary::HandleTrajectoryWorldCollisions(const UObject* WorldContextObject, UPARAM(ref) const FPoseSearchQueryTrajectory& InTrajectory, FPoseSearchQueryTrajectory& OutTrajectory,
+	ETraceTypeQuery TraceChannel, bool bTraceComplex, const TArray<AActor*>& ActorsToIgnore, EDrawDebugTrace::Type DrawDebugType, bool bIgnoreSelf, FLinearColor TraceColor, FLinearColor TraceHitColor, float DrawTime)
 {
-	FPoseSearchTrajectoryData::FSampling TrajectoryDataSampling;
-	TrajectoryDataSampling.NumHistorySamples = InTrajectoryHistoryCount;
-	TrajectoryDataSampling.SecondsPerHistorySample = InHistorySamplingInterval;
-	TrajectoryDataSampling.NumPredictionSamples = InTrajectoryPredictionCount;
-	TrajectoryDataSampling.SecondsPerPredictionSample = InPredictionSamplingInterval;
+	OutTrajectory = InTrajectory;
 
-	FPoseSearchTrajectoryData::FState TrajectoryDataState;
-	TrajectoryDataState.DesiredControllerYawLastUpdate = InOutDesiredControllerYawLastUpdate;
-
-	FPoseSearchTrajectoryData::FDerived TrajectoryDataDerived;
-	InTrajectoryData.UpdateData(InDeltaTime, InAnimInstance, TrajectoryDataDerived, TrajectoryDataState);
-	
-	InitTrajectorySamples(InOutTrajectory, InTrajectoryData, TrajectoryDataDerived, TrajectoryDataSampling);
-	UpdateHistory_TransformHistory(InOutTrajectory, InTrajectoryData, TrajectoryDataDerived, TrajectoryDataSampling, InDeltaTime);
-	UpdatePrediction_SimulateCharacterMovementAdvanced(InOutTrajectory, InTrajectoryData, TrajectoryDataDerived, TrajectoryDataSampling);
+	FHitResult OutHit;
+	const int32 NumSamples = OutTrajectory.Samples.Num();
+	for (int32 SampleIndex = 0; SampleIndex < NumSamples; ++SampleIndex)
+	{
+		FPoseSearchQueryTrajectorySample& Sample = OutTrajectory.Samples[SampleIndex];
+		if (Sample.AccumulatedSeconds > 0.f)
+		{
+			FVector End = Sample.Position;
+			FVector Start = Sample.Position;
+			Start.Z += 3000.f;
 
-	InOutDesiredControllerYawLastUpdate = TrajectoryDataState.DesiredControllerYawLastUpdate;
+			if (UKismetSystemLibrary::LineTraceSingle(WorldContextObject, Start, End, TraceChannel, bTraceComplex, ActorsToIgnore, DrawDebugType, OutHit, bIgnoreSelf, TraceColor, TraceHitColor, DrawTime))
+			{
+				Sample.Position.Z = OutHit.ImpactPoint.Z;
+			}
+		}
+	}
+}
 
-	OutTrajectory = InOutTrajectory;
+void UPoseSearchTrajectoryLibrary::GetTrajectorySampleAtTime(UPARAM(ref) const FPoseSearchQueryTrajectory& InTrajectory, float Time, FPoseSearchQueryTrajectorySample& OutTrajectorySample, bool bExtrapolate)
+{
+	OutTrajectorySample = InTrajectory.GetSampleAtTime(Time, bExtrapolate);
 }
 
diff --git a/Engine/Plugins/Experimental/Animation/PoseSearch/Source/Runtime/Public/PoseSearch/AnimNode_PoseSearchHistoryCollector.h b/Engine/Plugins/Experimental/Animation/PoseSearch/Source/Runtime/Public/PoseSearch/AnimNode_PoseSearchHistoryCollector.h
--- a/Engine/Plugins/Experimental/Animation/PoseSearch/Source/Runtime/Public/PoseSearch/AnimNode_PoseSearchHistoryCollector.h
+++ b/Engine/Plugins/Experimental/Animation/PoseSearch/Source/Runtime/Public/PoseSearch/AnimNode_PoseSearchHistoryCollector.h
@@ -82,7 +82,8 @@ public:
 	virtual void PreUpdate(const UAnimInstance* InAnimInstance) override;
 	// End of FAnimNode_Base interface
 
-	const UE::PoseSearch::IPoseHistory& GetPoseHistory() const { return PoseHistory; }
+	const UE::PoseSearch::FPoseHistory& GetPoseHistory() const { return PoseHistory; }
+	UE::PoseSearch::FPoseHistory& GetPoseHistory() { return PoseHistory; }
 
 protected:
 	UE::PoseSearch::FPoseHistory PoseHistory;
diff --git a/Engine/Plugins/Experimental/Animation/PoseSearch/Source/Runtime/Public/PoseSearch/PoseSearchHistory.h b/Engine/Plugins/Experimental/Animation/PoseSearch/Source/Runtime/Public/PoseSearch/PoseSearchHistory.h
--- a/Engine/Plugins/Experimental/Animation/PoseSearch/Source/Runtime/Public/PoseSearch/PoseSearchHistory.h
+++ b/Engine/Plugins/Experimental/Animation/PoseSearch/Source/Runtime/Public/PoseSearch/PoseSearchHistory.h
@@ -110,11 +110,6 @@ struct FPoseHistory : public IPoseHistory
 
 	void Initialize_AnyThread(int32 InNumPoses, float InSamplingInterval);
 	void CacheBones_AnyThread(const TArray<FBoneIndexType>& RequiredBones);
-
-	// if bGenerateTrajectory is true, this method will generate trajectory
-	// NoTe: InTrajectory.Samples[i].AccumulatedSeconds == 0 is the sample of the previous frame of simulation (since MM works by matching the previous character pose)
-	void Update_AnyThread(float DeltaTime, const FPoseSearchQueryTrajectory& InTrajectory, float InTrajectorySpeedMultiplier);
-
 	void EvaluateComponentSpace_AnyThread(float DeltaTime, FCSPose<FCompactPose>& ComponentSpacePose, bool bStoreScales);
 
 	// IPoseHistory interface
@@ -124,6 +119,7 @@ struct FPoseHistory : public IPoseHistory
 	virtual bool IsEmpty() const override;
 	virtual const FBoneToTransformMap& GetBoneToTransformMap() const override;
 	
+	void SetTrajectory(const FPoseSearchQueryTrajectory& InTrajectory, float InTrajectorySpeedMultiplier = 1.f);
 	virtual int32 GetNumEntries() const override;
 	virtual const FPoseHistoryEntry& GetEntry(int32 EntryIndex) const override;
 
diff --git a/Engine/Plugins/Experimental/Animation/PoseSearch/Source/Runtime/Public/PoseSearch/PoseSearchHistoryCollectorAnimNodeLibrary.h b/Engine/Plugins/Experimental/Animation/PoseSearch/Source/Runtime/Public/PoseSearch/PoseSearchHistoryCollectorAnimNodeLibrary.h
--- a/Engine/Plugins/Experimental/Animation/PoseSearch/Source/Runtime/Public/PoseSearch/PoseSearchHistoryCollectorAnimNodeLibrary.h
+++ b/Engine/Plugins/Experimental/Animation/PoseSearch/Source/Runtime/Public/PoseSearch/PoseSearchHistoryCollectorAnimNodeLibrary.h
@@ -40,4 +40,6 @@ public:
 	UFUNCTION(BlueprintPure, Category = "Animation|PoseHistory", meta = (BlueprintThreadSafe, DisplayName = "Get Pose History Node Trajectory"))
 	static void GetPoseHistoryNodeTrajectory(const FPoseSearchHistoryCollectorAnimNodeReference& PoseSearchHistoryCollectorNode, FPoseSearchQueryTrajectory& Trajectory);
 
+	UFUNCTION(BlueprintCallable, Category = "Animation|PoseHistory", meta = (BlueprintThreadSafe, DisplayName = "Set Pose History Node Trajectory"))
+	static void SetPoseHistoryNodeTrajectory(const FPoseSearchHistoryCollectorAnimNodeReference& PoseSearchHistoryCollectorNode, const FPoseSearchQueryTrajectory& Trajectory);
 };
\ No newline at end of file
diff --git a/Engine/Plugins/Experimental/Animation/PoseSearch/Source/Runtime/Public/PoseSearch/PoseSearchTrajectoryLibrary.h b/Engine/Plugins/Experimental/Animation/PoseSearch/Source/Runtime/Public/PoseSearch/PoseSearchTrajectoryLibrary.h
--- a/Engine/Plugins/Experimental/Animation/PoseSearch/Source/Runtime/Public/PoseSearch/PoseSearchTrajectoryLibrary.h
+++ b/Engine/Plugins/Experimental/Animation/PoseSearch/Source/Runtime/Public/PoseSearch/PoseSearchTrajectoryLibrary.h
@@ -4,6 +4,7 @@
 
 #include "Curves/CurveFloat.h"
 #include "PoseSearch/PoseSearchTrajectoryTypes.h"
+#include "Kismet/KismetSystemLibrary.h"
 #include "PoseSearchTrajectoryLibrary.generated.h"
 
 class UAnimInstance;
@@ -35,7 +36,7 @@ public:
 		FQuat Facing = FQuat::Identity;
 		FQuat MeshCompRelativeRotation = FQuat::Identity;
 		bool bOrientRotationToMovement = false;
-		bool bIsFalling = false;
+		bool bStepGroundPrediction = true;
 	};
 
 	struct FSampling
@@ -50,7 +51,6 @@ public:
 	void UpdateData(float DeltaTime, const FAnimInstanceProxy& AnimInstanceProxy, FDerived& TrajectoryDataDerived, FState& TrajectoryDataState) const;
 	void UpdateData(float DeltaTime, const UAnimInstance* AnimInstance, FDerived& TrajectoryDataDerived, FState& TrajectoryDataState) const;
 	FVector StepCharacterMovementGroundPrediction(float DeltaTime, const FVector& InVelocity, const FVector& InAcceleration, const FDerived& TrajectoryDataDerived) const;
-	FVector StepCharacterFallingPrediction(float DeltaTime, const FVector& InitialVelocity, const FDerived& TrajectoryDataDerived) const;
 	
 	// If the character is forward facing (i.e. bOrientRotationToMovement is true), this controls how quickly the trajectory will rotate
 	// to face acceleration. It's common for this to differ from the rotation rate of the character, because animations are often authored 
@@ -84,7 +84,7 @@ public:
 /**
  * Set of functions to help populate a FPoseSearchQueryTrajectory for motion matching.
  */
-UCLASS()
+UCLASS(Experimental)
 class POSESEARCH_API UPoseSearchTrajectoryLibrary : public UBlueprintFunctionLibrary
 {
 	GENERATED_BODY()
@@ -98,25 +98,22 @@ public:
 	static void UpdateHistory_TransformHistory(FPoseSearchQueryTrajectory& Trajectory, const FPoseSearchTrajectoryData& TrajectoryData, const FPoseSearchTrajectoryData::FDerived& TrajectoryDataDerived, const FPoseSearchTrajectoryData::FSampling& TrajectoryDataSampling, float DeltaTime);
 
 	// Update prediction by simulating the movement math for ground locomotion from UCharacterMovementComponent.
-	static void UpdatePrediction_SimulateCharacterMovement(FPoseSearchQueryTrajectory& Trajectory, const FPoseSearchTrajectoryData& TrajectoryData, const FPoseSearchTrajectoryData::FDerived& TrajectoryDataDerived, const FPoseSearchTrajectoryData::FSampling& TrajectoryDataSampling);
+	static void UpdatePrediction_SimulateCharacterMovement(FPoseSearchQueryTrajectory& Trajectory, const FPoseSearchTrajectoryData& TrajectoryData, const FPoseSearchTrajectoryData::FDerived& TrajectoryDataDerived, const FPoseSearchTrajectoryData::FSampling& TrajectoryDataSampling, bool bAlwaysApplyGravity = false);
 
-	// Update prediction by simulating the movement math for locomotion from UCharacterMovementComponent.
-	static void UpdatePrediction_SimulateCharacterMovementAdvanced(FPoseSearchQueryTrajectory& Trajectory, const FPoseSearchTrajectoryData& TrajectoryData, const FPoseSearchTrajectoryData::FDerived& TrajectoryDataDerived, const FPoseSearchTrajectoryData::FSampling& TrajectoryDataSampling);
-	
 	/** Get a Pose History node context from an anim node context (pure) */
-	UFUNCTION(BlueprintPure, Category = "Animation|PoseSearch", meta = (BlueprintThreadSafe, DisplayName = "Pose Search Generate Trajectory"))
-	static void PoseSearchGenerateTrajectory(
-		const UAnimInstance* InAnimInstance, UPARAM(ref) const FPoseSearchTrajectoryData& InTrajectoryData, float InDeltaTime,
+	UFUNCTION(BlueprintCallable, Category = "Animation|PoseSearch", meta = (BlueprintThreadSafe, DisplayName = "Pose Search Generate Trajectory"))
+	static void PoseSearchGenerateTrajectory(const UAnimInstance* InAnimInstance, UPARAM(ref) const FPoseSearchTrajectoryData& InTrajectoryData, float InDeltaTime,
 		UPARAM(ref) FPoseSearchQueryTrajectory& InOutTrajectory, UPARAM(ref) float& InOutDesiredControllerYawLastUpdate, FPoseSearchQueryTrajectory& OutTrajectory,
-		float InHistorySamplingInterval = 0.04f, int32 InTrajectoryHistoryCount = 10, float InPredictionSamplingInterval = 0.4f, int32 InTrajectoryPredictionCount = 8);
+		float InHistorySamplingInterval = 0.04f, int32 InTrajectoryHistoryCount = 10, float InPredictionSamplingInterval = 0.2f, int32 InTrajectoryPredictionCount = 8,
+		bool bAlwaysApplyGravity = false);
+
+	UFUNCTION(BlueprintCallable, Category="Animation|PoseSearch", meta=(bIgnoreSelf="true", WorldContext="WorldContextObject", AutoCreateRefTerm="ActorsToIgnore", AdvancedDisplay="TraceColor,TraceHitColor,DrawTime"))
+	static void HandleTrajectoryWorldCollisions(const UObject* WorldContextObject, UPARAM(ref) const FPoseSearchQueryTrajectory& InTrajectory, FPoseSearchQueryTrajectory& OutTrajectory,
+		ETraceTypeQuery TraceChannel, bool bTraceComplex, const TArray<AActor*>& ActorsToIgnore, EDrawDebugTrace::Type DrawDebugType, bool bIgnoreSelf, FLinearColor TraceColor = FLinearColor::Red, FLinearColor TraceHitColor = FLinearColor::Green, float DrawTime = 5.0f);
+
+	UFUNCTION(BlueprintPure, Category="Animation|PoseSearch")
+	static void GetTrajectorySampleAtTime(UPARAM(ref) const FPoseSearchQueryTrajectory& InTrajectory, float Time, FPoseSearchQueryTrajectorySample& OutTrajectorySample, bool bExtrapolate = true);
 
-	/** Generate trajectory to be used by a Pose History node. Supports more than on ground movement prediction (i.e. jumping, collisions etc). */
-	UFUNCTION(BlueprintPure, Category = "Animation|PoseSearch", meta = (BlueprintThreadSage, DisplayName = "Pose Search Generate Trajectory Advanced"))
-	static void PoseSearchGenerateTrajectoryAdvanced(
-		const UAnimInstance* InAnimInstance, UPARAM(ref) const FPoseSearchTrajectoryData& InTrajectoryData, float InDeltaTime,
-		UPARAM(ref) FPoseSearchQueryTrajectory& InOutTrajectory, UPARAM(ref) float& InOutDesiredControllerYawLastUpdate, FPoseSearchQueryTrajectory& OutTrajectory,
-		float InHistorySamplingInterval = 0.04f, int32 InTrajectoryHistoryCount = 10, float InPredictionSamplingInterval = 0.4f, int32 InTrajectoryPredictionCount = 8);
-	
 private:
 	static FVector RemapVectorMagnitudeWithCurve(const FVector& Vector, bool bUseCurve, const FRuntimeFloatCurve& Curve);
 };
\ No newline at end of file
