jasper - Improved pose search trajectory prediction for jumps

diff --git a/Engine/Plugins/Experimental/Animation/PoseSearch/Source/Runtime/Private/PoseSearchTrajectoryLibrary.cpp b/Engine/Plugins/Experimental/Animation/PoseSearch/Source/Runtime/Private/PoseSearchTrajectoryLibrary.cpp
--- a/Engine/Plugins/Experimental/Animation/PoseSearch/Source/Runtime/Private/PoseSearchTrajectoryLibrary.cpp
+++ b/Engine/Plugins/Experimental/Animation/PoseSearch/Source/Runtime/Private/PoseSearchTrajectoryLibrary.cpp
@@ -339,7 +339,7 @@ void UPoseSearchTrajectoryLibrary::UpdatePrediction_SimulateCharacterMovement(
 			for (int32 Index = NumHistorySamples + 1; Index < Trajectory.Samples.Num(); ++Index)
 			{
 				const float Time = Trajectory.Samples[Index].AccumulatedSeconds;
-				Trajectory.Samples[Index].Position += TrajectoryDataDerived.Gravity * (Time * Time);
+				Trajectory.Samples[Index].Position += 0.5f * TrajectoryDataDerived.Gravity * (Time * Time);
 			}
 		}
 	}
