motion matching - debug lines improvements

diff --git a/Engine/Plugins/Experimental/Animation/PoseSearch/Source/Runtime/Private/PoseSearchFeatureChannel_Position.cpp b/Engine/Plugins/Experimental/Animation/PoseSearch/Source/Runtime/Private/PoseSearchFeatureChannel_Position.cpp
--- a/Engine/Plugins/Experimental/Animation/PoseSearch/Source/Runtime/Private/PoseSearchFeatureChannel_Position.cpp
+++ b/Engine/Plugins/Experimental/Animation/PoseSearch/Source/Runtime/Private/PoseSearchFeatureChannel_Position.cpp
@@ -174,7 +174,7 @@ void UPoseSearchFeatureChannel_Position::DebugDraw(const UE::PoseSearch::FDebugD
 	DrawParams.DrawPoint(BonePos, Color);
 
 	const bool bDrawOrigin = !DeltaPos.IsNearlyZero() && (SchemaOriginBoneIdx != RootSchemaBoneIdx || !FMath::IsNearlyZero(OriginTimeOffset) ||
-							 SampleRole != OriginRole || PermutationTimeType != EPermutationTimeType::UseSampleTime);
+							 SampleRole != OriginRole || PermutationTimeType != EPermutationTimeType::UseSampleTime || bUseBlueprintQueryOverride);
 	if (bDrawOrigin)
 	{
 		DrawParams.DrawPoint(OriginBonePos, Color);
diff --git a/Engine/Plugins/Experimental/Animation/PoseSearch/Source/Runtime/Private/PoseSearchFeatureChannel_Trajectory.cpp b/Engine/Plugins/Experimental/Animation/PoseSearch/Source/Runtime/Private/PoseSearchFeatureChannel_Trajectory.cpp
--- a/Engine/Plugins/Experimental/Animation/PoseSearch/Source/Runtime/Private/PoseSearchFeatureChannel_Trajectory.cpp
+++ b/Engine/Plugins/Experimental/Animation/PoseSearch/Source/Runtime/Private/PoseSearchFeatureChannel_Trajectory.cpp
@@ -144,7 +144,7 @@ void UPoseSearchFeatureChannel_Trajectory::DebugDraw(const UE::PoseSearch::FDebu
 		}
 	}
 
-	if (Positions.Num() >= 2)
+	if (!Positions.IsEmpty())
 	{
 		Positions.Sort([](const UPoseSearchFeatureChannel_Position& a, const UPoseSearchFeatureChannel_Position& b)
 			{
@@ -165,9 +165,9 @@ void UPoseSearchFeatureChannel_Trajectory::DebugDraw(const UE::PoseSearch::FDebu
 			Color = FLinearColor::Blue.ToFColor(true);
 #endif // WITH_EDITORONLY_DATA
 
-			if (PrevTimeOffset * CurrTimeOffset < UE_KINDA_SMALL_NUMBER)
+			if (Positions.Num() == 1 || PrevTimeOffset * CurrTimeOffset < UE_KINDA_SMALL_NUMBER)
 			{
-				// we jumped from negative to positive time offset without having a zero time offset. so we add the zero
+				// we jumped from negative to positive time offset without having a zero time offset, or we only have one position point. so we add the zero
 				TrajSplinePos.Add(DrawParams.ExtractPosition(PoseVector, 0.f, RootSchemaBoneIdx, Positions[i]->OriginRole));
 				TrajSplineColor.Add(Color);
 			}
