motion matching - better fault tolerance for FPoseHistoryEntry::GetComponentSpaceTransform

diff --git a/Engine/Plugins/Experimental/Animation/PoseSearch/Source/Runtime/Private/PoseSearchHistory.cpp b/Engine/Plugins/Experimental/Animation/PoseSearch/Source/Runtime/Private/PoseSearchHistory.cpp
--- a/Engine/Plugins/Experimental/Animation/PoseSearch/Source/Runtime/Private/PoseSearchHistory.cpp
+++ b/Engine/Plugins/Experimental/Animation/PoseSearch/Source/Runtime/Private/PoseSearchHistory.cpp
@@ -212,6 +212,16 @@ void FPoseHistoryEntry::SetComponentSpaceTransform(int32 Index, const FTransform
 
 FTransform FPoseHistoryEntry::GetComponentSpaceTransform(int32 Index) const
 {
+#if WITH_EDITOR
+	if (Index < 0 || Index >= ComponentSpaceRotations.Num())
+	{
+		UE_LOG(LogPoseSearch, Error, TEXT("FPoseHistoryEntry::GetComponentSpaceTransform - Index %d out of bound [0, %d)"), Index, ComponentSpaceRotations.Num());
+		return FTransform::Identity;
+	}
+#endif // WITH_EDITOR
+
+	check(ComponentSpaceScales.IsEmpty() || ComponentSpaceRotations.Num() == ComponentSpaceScales.Num());
+
 	const FQuat Quat(ComponentSpaceRotations[Index]);
 	const FVector Scale(ComponentSpaceScales.IsEmpty() ? FVector3f::OneVector : ComponentSpaceScales[Index]);
 	return FTransform(Quat, ComponentSpacePositions[Index], Scale);
