motion matching - UPoseSearchDatabase AllowInvertedInterval on ExcludeFromDatabaseParameters and AdditionalExtrapolationTime

diff --git a/Engine/Plugins/Experimental/Animation/PoseSearch/Source/Runtime/Public/PoseSearch/PoseSearchDatabase.h b/Engine/Plugins/Experimental/Animation/PoseSearch/Source/Runtime/Public/PoseSearch/PoseSearchDatabase.h
--- a/Engine/Plugins/Experimental/Animation/PoseSearch/Source/Runtime/Public/PoseSearch/PoseSearchDatabase.h
+++ b/Engine/Plugins/Experimental/Animation/PoseSearch/Source/Runtime/Public/PoseSearch/PoseSearchDatabase.h
@@ -301,11 +301,11 @@ public:
 #if WITH_EDITORONLY_DATA
 	// These settings allow users to trim the start and end of animations in the database to preserve start/end frames for blending, and prevent the system from selecting the very last frames before it blends out.
 	// valid animation frames will be AnimationAssetTimeStart + ExcludeFromDatabaseParameters.Min, AnimationAssetTimeEnd + ExcludeFromDatabaseParameters.Max
-	UPROPERTY(EditAnywhere, Category = "Database")
+	UPROPERTY(EditAnywhere, Category = "Database", meta = (AllowInvertedInterval))
 	FFloatInterval ExcludeFromDatabaseParameters = FFloatInterval(0.f, -0.3f);
 
 	// extrapolation of animation assets will be clamped by AnimationAssetTimeStart + AdditionalExtrapolationTime.Min, AnimationAssetTimeEnd + AdditionalExtrapolationTime.Max
-	UPROPERTY(EditAnywhere, Category = "Database")
+	UPROPERTY(EditAnywhere, Category = "Database", meta = (AllowInvertedInterval))
 	FFloatInterval AdditionalExtrapolationTime = FFloatInterval(-100.f, 100.f);
 #endif // WITH_EDITORONLY_DATA
 
