motion matching - UPoseSearchFeatureChannel_TimeToEvent

diff --git a/Engine/Plugins/Experimental/Animation/PoseSearch/Source/Runtime/Private/PoseSearchAssetIndexer.cpp b/Engine/Plugins/Experimental/Animation/PoseSearch/Source/Runtime/Private/PoseSearchAssetIndexer.cpp
--- a/Engine/Plugins/Experimental/Animation/PoseSearch/Source/Runtime/Private/PoseSearchAssetIndexer.cpp
+++ b/Engine/Plugins/Experimental/Animation/PoseSearch/Source/Runtime/Private/PoseSearchAssetIndexer.cpp
@@ -231,7 +231,7 @@ void FAssetIndexer::Process(int32 AssetIdx)
 	bProcessFailed = false;
 
 	// Generate pose metadata
-	const float PlayLength = AssetSamplers.GetPlayLength();
+	const float PlayLength = GetPlayLength();
 	for (int32 SampleIdx = GetBeginSampleIdx(); SampleIdx != GetEndSampleIdx(); ++SampleIdx)
 	{
 		const float SampleTime = FMath::Min(CalculateSampleTime(SampleIdx), PlayLength);
@@ -293,7 +293,7 @@ void FAssetIndexer::ComputeStats()
 
 		for (int32 SampleIdx = GetBeginSampleIdx(); SampleIdx != GetEndSampleIdx(); ++SampleIdx)
 		{
-			const float SampleTime = FMath::Min(CalculateSampleTime(SampleIdx), AssetSamplers.GetPlayLength());
+			const float SampleTime = FMath::Min(CalculateSampleTime(SampleIdx), GetPlayLength());
 
 			bool AnyClamped = false;
 			const FTransform TrajTransformsPast = GetTransform(SampleTime - FiniteDelta, RoleIndex, AnyClamped, RootBoneReference);
@@ -329,7 +329,7 @@ void FAssetIndexer::ComputeStats()
 
 void FAssetIndexer::GetSampleInfo(float SampleTime, int32 RoleIndex, FTransform& OutRootTransform, float& OutClipTime, bool& bOutClamped) const
 {
-	const float PlayLength = AssetSamplers.GetPlayLength();
+	const float PlayLength = GetPlayLength();
 	const bool bCanWrap = AssetSamplers.IsLoopable();
 
 	float MainRelativeTime = SampleTime;
@@ -413,7 +413,7 @@ FAssetIndexer::FCachedEntry& FAssetIndexer::GetEntry(float SampleTime)
 		Entry->SampleTime = SampleTime;
 
 		const bool bLoopable = AssetSamplers.IsLoopable();
-		const float PlayLength = AssetSamplers.GetPlayLength();
+		const float PlayLength = GetPlayLength();
 		const int32 AssetSamplersNum = AssetSamplers.Num();
 
 		Entry->RootTransform.SetNum(AssetSamplersNum);
@@ -560,7 +560,7 @@ bool FAssetIndexer::GetSampleRotation(FQuat& OutSampleRotation, float SampleTime
 
 	if (SamplingAttributeId >= 0)
 	{
-		if (AssetSamplers.ProcessAllAnimNotifyEvents([&OutSampleRotation, &SampleRole, &OriginRole, SamplingAttributeId, this, OriginTime](const TConstArrayView<FAnimNotifyEvent> AnimNotifyEvents)
+		if (ProcessAllAnimNotifyEvents([&OutSampleRotation, &SampleRole, &OriginRole, SamplingAttributeId, this, OriginTime](const TConstArrayView<FAnimNotifyEvent> AnimNotifyEvents)
 			{
 				for (const FAnimNotifyEvent& AnimNotifyEvent : AnimNotifyEvents)
 				{
@@ -631,7 +631,7 @@ bool FAssetIndexer::GetSamplePositionInternal(FVector& OutSamplePosition, float
 
 	if (SamplingAttributeId >= 0)
 	{
-		if (AssetSamplers.ProcessAllAnimNotifyEvents([&OutSamplePosition, OriginTime, &bClamped, SchemaOriginBoneIdx, &SampleRole, OriginRole, SamplingAttributeId, this](const TConstArrayView<FAnimNotifyEvent> AnimNotifyEvents)
+		if (ProcessAllAnimNotifyEvents([&OutSamplePosition, OriginTime, &bClamped, SchemaOriginBoneIdx, &SampleRole, OriginRole, SamplingAttributeId, this](const TConstArrayView<FAnimNotifyEvent> AnimNotifyEvents)
 			{
 				for (const FAnimNotifyEvent& AnimNotifyEvent : AnimNotifyEvents)
 				{
@@ -715,7 +715,7 @@ bool FAssetIndexer::GetSampleVelocity(FVector& OutSampleVelocity, float SampleTi
 	
 	if (SamplingAttributeId >= 0)
 	{
-		if (AssetSamplers.ProcessAllAnimNotifyEvents([&OutSampleVelocity, SchemaSampleBoneIdx, SchemaOriginBoneIdx, bUseCharacterSpaceVelocities, &SampleRole, &OriginRole, SamplingAttributeId, this, OriginTime](const TConstArrayView<FAnimNotifyEvent> AnimNotifyEvents)
+		if (ProcessAllAnimNotifyEvents([&OutSampleVelocity, SchemaSampleBoneIdx, SchemaOriginBoneIdx, bUseCharacterSpaceVelocities, &SampleRole, &OriginRole, SamplingAttributeId, this, OriginTime](const TConstArrayView<FAnimNotifyEvent> AnimNotifyEvents)
 			{
 				for (const FAnimNotifyEvent& AnimNotifyEvent : AnimNotifyEvents)
 				{
@@ -798,6 +798,21 @@ bool FAssetIndexer::GetSampleVelocity(FVector& OutSampleVelocity, float SampleTi
 	return false;
 }
 
+bool FAssetIndexer::ProcessAllAnimNotifyEvents(TFunction<bool(TConstArrayView<FAnimNotifyEvent>)> ProcessAnimNotifyEvents) const
+{
+	return AssetSamplers.ProcessAllAnimNotifyEvents(ProcessAnimNotifyEvents);
+}
+
+const FString FAssetIndexer::GetAssetName() const
+{
+	return AssetSamplers.GetAssetName();
+}
+
+float FAssetIndexer::GetPlayLength() const
+{
+	return AssetSamplers.GetPlayLength();
+}
+
 int32 FAssetIndexer::GetBeginSampleIdx() const
 {
 	return SearchIndexAsset.GetBeginSampleIdx();
diff --git a/Engine/Plugins/Experimental/Animation/PoseSearch/Source/Runtime/Private/PoseSearchFeatureChannel_TimeToEvent.cpp b/Engine/Plugins/Experimental/Animation/PoseSearch/Source/Runtime/Private/PoseSearchFeatureChannel_TimeToEvent.cpp
new file mode 100644
--- /dev/null
+++ b/Engine/Plugins/Experimental/Animation/PoseSearch/Source/Runtime/Private/PoseSearchFeatureChannel_TimeToEvent.cpp
@@ -0,0 +1,92 @@
+// Copyright Epic Games, Inc. All Rights Reserved.
+
+#include "PoseSearchFeatureChannel_TimeToEvent.h"
+#include "Engine/BlueprintGeneratedClass.h"
+#include "PoseSearch/PoseSearchAnimNotifies.h"
+#include "PoseSearch/PoseSearchAssetIndexer.h"
+#include "PoseSearch/PoseSearchContext.h"
+#include "PoseSearch/PoseSearchSchema.h"
+
+UPoseSearchFeatureChannel_TimeToEvent::UPoseSearchFeatureChannel_TimeToEvent()
+{
+	bUseBlueprintQueryOverride = Cast<UBlueprintGeneratedClass>(GetClass()) != nullptr;
+}
+
+bool UPoseSearchFeatureChannel_TimeToEvent::Finalize(UPoseSearchSchema* Schema)
+{
+	ChannelDataOffset = Schema->SchemaCardinality;
+	ChannelCardinality = 1;
+	Schema->SchemaCardinality += ChannelCardinality;
+	return true;
+}
+
+void UPoseSearchFeatureChannel_TimeToEvent::BuildQuery(UE::PoseSearch::FSearchContext& SearchContext) const
+{
+	using namespace UE::PoseSearch;
+
+	if (!bUseBlueprintQueryOverride)
+	{
+		UE_LOG(LogPoseSearch, Warning, TEXT("UPoseSearchFeatureChannel_TimeToEvent::BuildQuery - UPoseSearchFeatureChannel_TimeToEvent is designed to work only as BluePrint overridden class. The query TimeToEvent value will be defaulted to 0"));
+		return;
+	}
+
+	if (SearchContext.GetAnimInstances().IsEmpty())
+	{
+		UE_LOG(LogPoseSearch, Error, TEXT("UPoseSearchFeatureChannel_Position::BuildQuery - no provided anim instance!"));
+	}
+	
+	const float TimeToEvent = BP_GetTimeToEvent(SearchContext.GetAnimInstances()[0]);
+	FFeatureVectorHelper::EncodeFloat(SearchContext.EditFeatureVector(), ChannelDataOffset, TimeToEvent);
+}
+
+#if WITH_EDITOR
+void UPoseSearchFeatureChannel_TimeToEvent::FillWeights(TArrayView<float> Weights) const
+{
+	for (int32 i = 0; i < ChannelCardinality; ++i)
+	{
+		Weights[ChannelDataOffset + i] = Weight;
+	}
+}
+
+bool UPoseSearchFeatureChannel_TimeToEvent::IndexAsset(UE::PoseSearch::FAssetIndexer& Indexer) const
+{
+	using namespace UE::PoseSearch;
+
+	float EventTime = Indexer.GetPlayLength();
+
+	if (SamplingAttributeId >= 0)
+	{
+		Indexer.ProcessAllAnimNotifyEvents([&Indexer, this, &EventTime](const TConstArrayView<FAnimNotifyEvent> AnimNotifyEvents)
+			{
+				for (const FAnimNotifyEvent& AnimNotifyEvent : AnimNotifyEvents)
+				{
+					if (const UAnimNotifyState_PoseSearchSamplingEvent* SamplingEvent = Cast<UAnimNotifyState_PoseSearchSamplingEvent>(AnimNotifyEvent.NotifyStateClass))
+					{
+						if (SamplingEvent->SamplingAttributeId == SamplingAttributeId)
+						{
+							EventTime = AnimNotifyEvent.GetTime();
+							return true;
+						}
+					}
+				}
+
+				return false;
+			});
+	}
+
+	for (int32 SampleIdx = Indexer.GetBeginSampleIdx(); SampleIdx != Indexer.GetEndSampleIdx(); ++SampleIdx)
+	{
+		FFeatureVectorHelper::EncodeFloat(Indexer.GetPoseVector(SampleIdx), ChannelDataOffset, EventTime - Indexer.CalculateSampleTime(SampleIdx));
+	}
+	
+	return true;
+}
+
+UE::PoseSearch::TLabelBuilder& UPoseSearchFeatureChannel_TimeToEvent::GetLabel(UE::PoseSearch::TLabelBuilder& LabelBuilder, UE::PoseSearch::ELabelFormat LabelFormat) const
+{
+	GetOuterLabel(LabelBuilder, LabelFormat);
+	AppendLabelSeparator(LabelBuilder, LabelFormat);
+	LabelBuilder.Appendf(TEXT("TimeToEvent_%d"), SamplingAttributeId);
+	return LabelBuilder;
+}
+#endif
\ No newline at end of file
diff --git a/Engine/Plugins/Experimental/Animation/PoseSearch/Source/Runtime/Private/PoseSearchFeatureChannel_TimeToEvent.h b/Engine/Plugins/Experimental/Animation/PoseSearch/Source/Runtime/Private/PoseSearchFeatureChannel_TimeToEvent.h
new file mode 100644
--- /dev/null
+++ b/Engine/Plugins/Experimental/Animation/PoseSearch/Source/Runtime/Private/PoseSearchFeatureChannel_TimeToEvent.h
@@ -0,0 +1,36 @@
+// Copyright Epic Games, Inc. All Rights Reserved.
+
+#pragma once
+
+#include "PoseSearch/PoseSearchFeatureChannel.h"
+#include "PoseSearchFeatureChannel_TimeToEvent.generated.h"
+
+UCLASS(EditInlineNew, Blueprintable, meta = (DisplayName = "Time To Event Channel"), CollapseCategories)
+class POSESEARCH_API UPoseSearchFeatureChannel_TimeToEvent : public UPoseSearchFeatureChannel
+{
+	GENERATED_BODY()
+
+public:
+	UPROPERTY(EditAnywhere, Category = "Settings")
+	float Weight = 0.f;
+
+	UPROPERTY(EditAnywhere, Category = "Settings")
+	int32 SamplingAttributeId = 0;
+
+	UFUNCTION(BlueprintPure, BlueprintImplementableEvent, meta=(BlueprintThreadSafe, DisplayName = "Get Time To Event"), Category = "Settings")
+	float BP_GetTimeToEvent(const UAnimInstance* AnimInstance) const;
+
+	bool bUseBlueprintQueryOverride = false;
+
+	UPoseSearchFeatureChannel_TimeToEvent();
+
+	// UPoseSearchFeatureChannel interface
+	virtual bool Finalize(UPoseSearchSchema* Schema) override;
+	virtual void BuildQuery(UE::PoseSearch::FSearchContext& SearchContext) const override;
+
+#if WITH_EDITOR
+	virtual void FillWeights(TArrayView<float> Weights) const override;
+	virtual bool IndexAsset(UE::PoseSearch::FAssetIndexer& Indexer) const override;
+	virtual UE::PoseSearch::TLabelBuilder& GetLabel(UE::PoseSearch::TLabelBuilder& LabelBuilder, UE::PoseSearch::ELabelFormat LabelFormat = UE::PoseSearch::ELabelFormat::Full_Horizontal) const override;
+#endif
+};
diff --git a/Engine/Plugins/Experimental/Animation/PoseSearch/Source/Runtime/Public/PoseSearch/PoseSearchAnimNotifies.h b/Engine/Plugins/Experimental/Animation/PoseSearch/Source/Runtime/Public/PoseSearch/PoseSearchAnimNotifies.h
--- a/Engine/Plugins/Experimental/Animation/PoseSearch/Source/Runtime/Public/PoseSearch/PoseSearchAnimNotifies.h
+++ b/Engine/Plugins/Experimental/Animation/PoseSearch/Source/Runtime/Public/PoseSearch/PoseSearchAnimNotifies.h
@@ -62,10 +62,10 @@ public:
 	float CostAddend = -1.0f;
 };
 
-// UPoseSearchFeatureChannel(s) can use this UAnimNotifyState_PoseSearchSamplingAttribute as animation space position, rotation, and linear velocity provider 
+// UPoseSearchFeatureChannel(s) can use this UAnimNotifyState_PoseSearchSamplingEvent to demarcate events identified by SamplingAttributeId
 // during database indexing by specifying their SamplingAttributeId property to match UAnimNotifyState_PoseSearchSamplingAttribute::SamplingAttributeId
-UCLASS(Blueprintable, meta = (DisplayName = "Pose Search: Sampling Attribute"))
-class POSESEARCH_API UAnimNotifyState_PoseSearchSamplingAttribute : public UAnimNotifyState_PoseSearchBase
+UCLASS(Blueprintable, meta = (DisplayName = "Pose Search: Sampling Event"))
+class POSESEARCH_API UAnimNotifyState_PoseSearchSamplingEvent : public UAnimNotifyState_PoseSearchBase
 {
 	GENERATED_BODY()
 
@@ -76,6 +76,20 @@ public:
 	UPROPERTY(EditAnywhere, Category = Config, meta = (ClampMin=0))
 	int32 SamplingAttributeId = 0;
 
+#endif // WITH_EDITORONLY_DATA
+};
+
+// UPoseSearchFeatureChannel(s) can use this UAnimNotifyState_PoseSearchSamplingAttribute as animation space position, rotation, and linear velocity provider 
+// during database indexing by specifying their SamplingAttributeId property to match UAnimNotifyState_PoseSearchSamplingAttribute::SamplingAttributeId
+UCLASS(Blueprintable, meta = (DisplayName = "Pose Search: Sampling Attribute"))
+class POSESEARCH_API UAnimNotifyState_PoseSearchSamplingAttribute : public UAnimNotifyState_PoseSearchSamplingEvent
+{
+	GENERATED_BODY()
+
+public:
+
+#if WITH_EDITORONLY_DATA
+	
 	UPROPERTY(EditAnywhere, Category = Config)
 	FBoneReference Bone;
 
diff --git a/Engine/Plugins/Experimental/Animation/PoseSearch/Source/Runtime/Public/PoseSearch/PoseSearchAssetIndexer.h b/Engine/Plugins/Experimental/Animation/PoseSearch/Source/Runtime/Public/PoseSearch/PoseSearchAssetIndexer.h
--- a/Engine/Plugins/Experimental/Animation/PoseSearch/Source/Runtime/Public/PoseSearch/PoseSearchAssetIndexer.h
+++ b/Engine/Plugins/Experimental/Animation/PoseSearch/Source/Runtime/Public/PoseSearch/PoseSearchAssetIndexer.h
@@ -83,6 +83,10 @@ public:
 	// if bUseCharacterSpaceVelocities is true, velocities will be computed in root bone space, rather than animation (world) space
 	bool GetSampleVelocity(FVector& OutSampleVelocity, float SampleTimeOffset, float OriginTimeOffset, int32 SampleIdx, int8 SchemaSampleBoneIdx, int8 SchemaOriginBoneIdx, const FRole& SampleRole, const FRole& OriginRole, bool bUseCharacterSpaceVelocities = true, EPermutationTimeType PermutationTimeType = EPermutationTimeType::UseSampleTime, int32 SamplingAttributeId = INDEX_NONE);
 
+	bool ProcessAllAnimNotifyEvents(TFunction<bool(TConstArrayView<FAnimNotifyEvent>)> ProcessAnimNotifyEvents) const;
+	const FString GetAssetName() const;
+	float GetPlayLength() const;
+
 	int32 GetBeginSampleIdx() const;
 	int32 GetEndSampleIdx() const;
 	int32 GetNumIndexedPoses() const;
