Move PoseSearch/BlendStack out of experimental folder

diff --git a/Engine/Plugins/Animation/PoseSearch/Source/Editor/Private/PoseSearchDatabaseViewModel.cpp b/Engine/Plugins/Animation/PoseSearch/Source/Editor/Private/PoseSearchDatabaseViewModel.cpp
new file mode 100644
--- /dev/null
+++ b/Engine/Plugins/Animation/PoseSearch/Source/Editor/Private/PoseSearchDatabaseViewModel.cpp
@@ -0,0 +1,950 @@
+// Copyright Epic Games, Inc. All Rights Reserved.
+
+#include "PoseSearchDatabaseViewModel.h"
+#include "AnimPreviewInstance.h"
+#include "Animation/AnimComposite.h"
+#include "Animation/AnimMontage.h"
+#include "Animation/AnimSequence.h"
+#include "Animation/BlendSpace.h"
+#include "Animation/DebugSkelMeshComponent.h"
+#include "Animation/MirrorDataTable.h"
+#include "Camera/CameraComponent.h"
+#include "Components/CapsuleComponent.h"
+#include "EngineUtils.h"
+#include "GameFramework/PlayerController.h"
+#include "InstancedStruct.h"
+#include "Modules/ModuleManager.h"
+#include "PoseSearch/PoseSearchAnimNotifies.h"
+#include "PoseSearch/PoseSearchContext.h"
+#include "PoseSearch/PoseSearchDatabase.h"
+#include "PoseSearch/PoseSearchDefines.h"
+#include "PoseSearch/PoseSearchDerivedData.h"
+#include "PoseSearch/PoseSearchHistory.h"
+#include "PoseSearch/PoseSearchSchema.h"
+#include "PoseSearch/PoseSearchTrajectoryTypes.h"
+#include "PoseSearchDatabaseAssetTreeNode.h"
+#include "PoseSearchDatabaseDataDetails.h"
+#include "PoseSearchDatabasePreviewScene.h"
+#include "PoseSearchEditor.h"
+#include "PropertyEditorModule.h"
+
+namespace UE::PoseSearch
+{
+#if ENABLE_ANIM_DEBUG
+static TAutoConsoleVariable<float> CVarDatabasePreviewDebugDrawSamplerSize(TEXT("a.DatabasePreview.DebugDrawSamplerSize"), 0.f, TEXT("Debug Draw Sampler Positions Size"));
+#endif
+
+constexpr float StepDeltaTime = 1.0f / 30.0f;
+
+// FDatabasePreviewActor
+bool FDatabasePreviewActor::SpawnPreviewActor(UWorld* World, const UPoseSearchDatabase* PoseSearchDatabase, int32 IndexAssetIdx, const FRole& Role, const FTransform& SamplerRootTransformOrigin, const FTransform* PrecalculatedRootTransformOrigin, int32 PoseIdxForTimeOffset)
+{
+	check(PoseSearchDatabase && PoseSearchDatabase->Schema);
+	const FSearchIndex& SearchIndex = PoseSearchDatabase->GetSearchIndex();
+
+	USkeleton* Skeleton = PoseSearchDatabase->Schema->GetSkeleton(Role);
+	if (!Skeleton)
+	{
+		UE_LOG(LogPoseSearchEditor, Error, TEXT("Couldn't spawn preview Actor because of missing Role '%s' in Schema '%s'"), *Role.ToString(), *PoseSearchDatabase->Schema->GetName());
+		return false;
+	}
+
+	const FSearchIndexAsset& IndexAsset = SearchIndex.Assets[IndexAssetIdx];
+
+	const FPoseSearchDatabaseAnimationAssetBase* DatabaseAnimationAsset = PoseSearchDatabase->GetAnimationAssetBase(IndexAsset.GetSourceAssetIdx());
+	UAnimationAsset* PreviewAsset = Cast<UAnimationAsset>(DatabaseAnimationAsset->GetAnimationAssetForRole(Role));
+	if (!PreviewAsset)
+	{
+		return false;
+	}
+
+	ActorRole = Role;
+	Sampler.Init(PreviewAsset, SamplerRootTransformOrigin, IndexAsset.GetBlendParameters());
+	Sampler.Process();
+
+	IndexAssetIndex = IndexAssetIdx;
+	CurrentPoseIndex = INDEX_NONE;
+
+	PlayTimeOffset = 0.f;
+	if (PoseIdxForTimeOffset >= 0)
+	{
+		PlayTimeOffset = PoseSearchDatabase->GetRealAssetTime(PoseIdxForTimeOffset) - IndexAsset.GetFirstSampleTime(PoseSearchDatabase->Schema->SampleRate);
+	}
+
+	FActorSpawnParameters Params;
+	Params.SpawnCollisionHandlingOverride = ESpawnActorCollisionHandlingMethod::AdjustIfPossibleButAlwaysSpawn;
+	ActorPtr = World->SpawnActor<AActor>(AActor::StaticClass(), FTransform::Identity, Params);
+	ActorPtr->SetFlags(RF_Transient);
+
+	UDebugSkelMeshComponent* Mesh = NewObject<UDebugSkelMeshComponent>(ActorPtr.Get());
+	Mesh->RegisterComponentWithWorld(World);
+
+	UAnimPreviewInstance* AnimInstance = NewObject<UAnimPreviewInstance>(Mesh);
+	Mesh->PreviewInstance = AnimInstance;
+	AnimInstance->InitializeAnimation();
+
+	USkeletalMesh* DatabasePreviewMesh = PoseSearchDatabase->PreviewMesh;
+	Mesh->SetSkeletalMesh(DatabasePreviewMesh ? DatabasePreviewMesh : PoseSearchDatabase->Schema->GetSkeleton(Role)->GetPreviewMesh(true));
+	Mesh->EnablePreview(true, PreviewAsset);
+		
+	AnimInstance->SetAnimationAsset(PreviewAsset, false, 0.0f);
+	AnimInstance->SetBlendSpacePosition(IndexAsset.GetBlendParameters());
+		
+	if (IndexAsset.IsMirrored() && PoseSearchDatabase->Schema)
+	{
+		const UMirrorDataTable* MirrorDataTable = PoseSearchDatabase->Schema->GetMirrorDataTable(Role);
+		AnimInstance->SetMirrorDataTable(MirrorDataTable);
+	}
+
+	const FMirrorDataCache MirrorDataCache(AnimInstance->GetMirrorDataTable(), AnimInstance->GetRequiredBonesOnAnyThread());
+	
+	if (PrecalculatedRootTransformOrigin)
+	{
+		// using the PrecalculatedRootTransformOrigin if provided. useful to synchronize multiple roles preview actors
+		RootTransformOrigin = *PrecalculatedRootTransformOrigin;
+	}
+	else
+	{
+		RootTransformOrigin = MirrorDataCache.MirrorTransform(Sampler.ExtractRootTransform(PlayTimeOffset));
+	}
+
+	AnimInstance->PlayAnim(false, 0.0f);
+
+	if (!ActorPtr->GetRootComponent())
+	{
+		ActorPtr->SetRootComponent(Mesh);
+	}
+
+	AnimInstance->SetPlayRate(0.f);
+
+	UE_LOG(LogPoseSearchEditor, Log, TEXT("Spawned preview Actor: %s"), *GetNameSafe(ActorPtr.Get()));
+	return true;
+}
+
+void FDatabasePreviewActor::UpdatePreviewActor(const UPoseSearchDatabase* PoseSearchDatabase, float PlayTime, bool bQuantizeAnimationToPoseData)
+{
+	check(PoseSearchDatabase);
+
+	const FSearchIndex& SearchIndex = PoseSearchDatabase->GetSearchIndex();
+
+	UAnimPreviewInstance* AnimInstance = GetAnimPreviewInstanceInternal();
+	if (!AnimInstance || IndexAssetIndex >= SearchIndex.Assets.Num())
+	{
+		return;
+	}
+
+	const UAnimationAsset* PreviewAsset = AnimInstance->GetAnimationAsset();
+	if (!PreviewAsset)
+	{
+		return;
+	}
+
+	CurrentTime = 0.f;
+	const FSearchIndexAsset& IndexAsset = SearchIndex.Assets[IndexAssetIndex];
+	float CurrentPlayTime = PlayTime + IndexAsset.GetFirstSampleTime(PoseSearchDatabase->Schema->SampleRate) + PlayTimeOffset;
+	FAnimationRuntime::AdvanceTime(false, CurrentPlayTime, CurrentTime, IndexAsset.GetLastSampleTime(PoseSearchDatabase->Schema->SampleRate));
+			 
+	// time to pose index
+	CurrentPoseIndex = PoseSearchDatabase->GetPoseIndexFromTime(CurrentTime, IndexAsset);
+
+	const float QuantizedTime = CurrentPoseIndex >= 0 ? PoseSearchDatabase->GetRealAssetTime(CurrentPoseIndex) : CurrentTime;
+	if (bQuantizeAnimationToPoseData)
+	{
+		CurrentTime = QuantizedTime;
+	}
+
+	// SetPosition is in [0..1] range for blendspaces
+	AnimInstance->SetPosition(Sampler.ToNormalizedTime(CurrentTime));
+	AnimInstance->SetPlayRate(0.f);
+	AnimInstance->SetBlendSpacePosition(IndexAsset.GetBlendParameters());
+
+	const FMirrorDataCache MirrorDataCache(AnimInstance->GetMirrorDataTable(), AnimInstance->GetRequiredBonesOnAnyThread());
+
+	// updating root transforms
+	RootTransformCurrentQuantizedTime = MirrorDataCache.MirrorTransform(Sampler.ExtractRootTransform(QuantizedTime));
+	RootTransformCurrentQuantizedTime.SetToRelativeTransform(RootTransformOrigin);
+
+	if (CurrentTime == QuantizedTime)
+	{
+		RootTransformCurrent = RootTransformCurrentQuantizedTime;
+	}
+	else
+	{
+		RootTransformCurrent = MirrorDataCache.MirrorTransform(Sampler.ExtractRootTransform(CurrentTime));
+		RootTransformCurrent.SetToRelativeTransform(RootTransformOrigin);
+	}
+
+	check(ActorPtr != nullptr);
+	ActorPtr->SetActorTransform(RootTransformCurrent);
+
+	// @todo: optimize this bone container, since we only need the root bone here...
+	const FBoneContainer& BoneContainer = AnimInstance->GetRequiredBonesOnAnyThread();
+	if (BoneContainer.GetNumBones() > 0)
+	{
+		FMemMark Mark(FMemStack::Get());
+		FCompactPose Pose;
+		Pose.SetBoneContainer(&BoneContainer);
+
+		Sampler.ExtractPose(QuantizedTime, Pose);
+		MirrorDataCache.MirrorPose(Pose);
+
+		RootBoneTransformCurrentQuantizedTime = Pose[FCompactPoseBoneIndex(RootBoneIndexType)];
+	}
+}
+
+void FDatabasePreviewActor::Destroy()
+{
+	if (ActorPtr != nullptr)
+	{
+		ActorPtr->Destroy();
+	}
+}
+
+bool FDatabasePreviewActor::DrawPreviewActors(TArrayView<FDatabasePreviewActor> PreviewActors, const UPoseSearchDatabase* PoseSearchDatabase, bool bDisplayRootMotionSpeed, bool bDisplayBlockTransition, TConstArrayView<float> QueryVector)
+{
+	using namespace UE::PoseSearch;
+
+	UWorld* CommonWorld = nullptr;
+	int32 CommonCurrentPoseIndex = INDEX_NONE;
+#if DO_CHECK
+	int32 CommonIndexAssetIndex = INDEX_NONE;
+#endif // DO_CHECK
+
+	TArray<const USkinnedMeshComponent*> Meshes;
+	FRoleToIndex RoleToIndex;
+	TArray<FArchivedPoseHistory> ArchivedPoseHistories;
+	TArray<const IPoseHistory*> PoseHistories;
+
+	const int32 NumPreviewActors = PreviewActors.Num();
+
+	Meshes.Reserve(NumPreviewActors);
+	RoleToIndex.Reserve(NumPreviewActors);
+	ArchivedPoseHistories.Reserve(NumPreviewActors);
+	PoseHistories.Reserve(NumPreviewActors);
+
+	for (FDatabasePreviewActor& PreviewActor : PreviewActors)
+	{
+		if (!PoseSearchDatabase->GetSearchIndex().IsValidPoseIndex(PreviewActor.GetCurrentPoseIndex()))
+		{
+			return false;
+		}
+
+		const UDebugSkelMeshComponent* Mesh = PreviewActor.GetDebugSkelMeshComponent();
+		if (!Mesh)
+		{
+			return false;
+		}
+
+		if (!CommonWorld)
+		{
+			CommonWorld = Mesh->GetWorld();
+		}
+		else if (CommonWorld != Mesh->GetWorld())
+		{
+			return false;
+		}
+
+		// making sure PreviewActors are consistent with each other
+		if (CommonCurrentPoseIndex == INDEX_NONE)
+		{
+			CommonCurrentPoseIndex = PreviewActor.GetCurrentPoseIndex();
+		}
+		else if (CommonCurrentPoseIndex != PreviewActor.GetCurrentPoseIndex())
+		{
+			checkNoEntry();
+			return false;
+		}
+
+#if DO_CHECK
+		if (CommonIndexAssetIndex == INDEX_NONE)
+		{
+			CommonIndexAssetIndex = PreviewActor.IndexAssetIndex;
+		}
+		else if (CommonIndexAssetIndex != PreviewActor.IndexAssetIndex)
+		{
+			checkNoEntry();
+			return false;
+		}
+#endif // DO_CHECK
+
+		RoleToIndex.Add(PreviewActor.ActorRole) = Meshes.Num();
+		Meshes.Add(Mesh);
+
+		const FTransform RootMotionTransform = PreviewActor.RootBoneTransformCurrentQuantizedTime * PreviewActor.RootTransformCurrentQuantizedTime;
+
+		// @todo: reconstruct ArchivedPoseHistory::BoneToTransformMap and ArchivedPoseHistory::Entries if needed
+		FArchivedPoseHistory& ArchivedPoseHistory = ArchivedPoseHistories.AddDefaulted_GetRef();
+		// @todo: reconstruct multiple trajectory samples if needed
+		// reconstructing the FPoseSearchQueryTrajectory with only one sample with AccumulatedSeconds at zero
+		FPoseSearchQueryTrajectorySample& TrajectorySample = ArchivedPoseHistory.Trajectory.Samples.AddDefaulted_GetRef();
+		TrajectorySample.SetTransform(RootMotionTransform);
+		TrajectorySample.AccumulatedSeconds = 0.f;
+
+		PoseHistories.Add(&ArchivedPoseHistory);
+	}
+
+	UE::PoseSearch::FDebugDrawParams DrawParams(Meshes, PoseHistories, RoleToIndex, PoseSearchDatabase);
+	DrawParams.DrawFeatureVector(CommonCurrentPoseIndex);
+
+	if (!QueryVector.IsEmpty())
+	{
+		DrawParams.DrawFeatureVector(QueryVector);
+	}
+
+	for (FDatabasePreviewActor& PreviewActor : PreviewActors)
+	{
+		const UDebugSkelMeshComponent* Mesh = PreviewActor.GetDebugSkelMeshComponent();
+		const FSearchIndex& SearchIndex = PoseSearchDatabase->GetSearchIndex();
+		const FSearchIndexAsset& IndexAsset = SearchIndex.Assets[PreviewActor.IndexAssetIndex];
+
+		const FMirrorDataCache MirrorDataCache(Mesh->PreviewInstance->GetMirrorDataTable(), Mesh->PreviewInstance->GetRequiredBonesOnAnyThread());
+		if (bDisplayRootMotionSpeed || bDisplayBlockTransition)
+		{
+			// initializing SampledRootMotion if required
+			if (PreviewActor.SampledRootMotion.IsEmpty())
+			{
+				const int NumPoses = IndexAsset.GetNumPoses();
+				if (NumPoses > 1)
+				{
+					PreviewActor.SampledRootMotion.SetNumUninitialized(NumPoses);
+					PreviewActor.SampledRootMotionSpeed.SetNumUninitialized(NumPoses);
+
+					for (int32 Index = 0; Index < NumPoses; ++Index)
+					{
+						const int32 IndexAssetPoseIdx = Index + IndexAsset.GetFirstPoseIdx();
+						const float IndexAssetPoseTime = IndexAsset.GetTimeFromPoseIndex(IndexAssetPoseIdx, PoseSearchDatabase->Schema->SampleRate);
+						FTransform IndexAssetPoseTransform = MirrorDataCache.MirrorTransform(PreviewActor.Sampler.ExtractRootTransform(IndexAssetPoseTime));
+						IndexAssetPoseTransform.SetToRelativeTransform(PreviewActor.RootTransformOrigin);
+
+						PreviewActor.SampledRootMotion[Index] = IndexAssetPoseTransform.GetTranslation();
+					}
+
+					for (int32 Index = 1; Index < NumPoses; ++Index)
+					{
+						const FVector& Start = PreviewActor.SampledRootMotion[Index - 1];
+						const FVector& End = PreviewActor.SampledRootMotion[Index];
+						PreviewActor.SampledRootMotionSpeed[Index] = (Start - End).Length() * PoseSearchDatabase->Schema->PermutationsSampleRate;
+					}
+					PreviewActor.SampledRootMotionSpeed[0] = PreviewActor.SampledRootMotionSpeed[1];
+				}
+			}
+		}
+
+		if (bDisplayRootMotionSpeed)
+		{
+			// drawing PreviewActor.SampledRootMotion
+			const int32 SampledRootMotionNum = PreviewActor.SampledRootMotion.Num();
+			if (SampledRootMotionNum > 1)
+			{
+				for (int32 Index = 0; Index < PreviewActor.SampledRootMotion.Num(); ++Index)
+				{
+					const FVector& EndDown = PreviewActor.SampledRootMotion[Index];
+					const FVector EndUp = EndDown + (PreviewActor.SampledRootMotionSpeed[Index] * FVector::UpVector);
+
+					DrawParams.DrawLine(EndDown, EndUp, FColor::Black);
+					if (Index > 0)
+					{
+						const FColor RootMotionColor = Index % 2 == 0 ? FColor::Purple : FColor::Orange;
+						const FVector& StartDown = PreviewActor.SampledRootMotion[Index - 1];
+						const FVector StartUp = StartDown + (PreviewActor.SampledRootMotionSpeed[Index - 1] * FVector::UpVector);
+						DrawParams.DrawLine(StartDown, EndDown, RootMotionColor);
+						DrawParams.DrawLine(StartUp, EndUp, RootMotionColor);
+					}
+				}
+			}
+		}
+
+		if (bDisplayBlockTransition)
+		{
+			const int NumPoses = IndexAsset.GetNumPoses();
+			if (NumPoses == PreviewActor.SampledRootMotion.Num())
+			{
+				for (int32 Index = 0; Index < NumPoses; ++Index)
+				{
+					const int32 IndexAssetPoseIdx = Index + IndexAsset.GetFirstPoseIdx();
+					if (SearchIndex.PoseMetadata[IndexAssetPoseIdx].IsBlockTransition())
+					{
+						DrawParams.DrawPoint(PreviewActor.SampledRootMotion[Index], FColor::Red);
+					}
+					else
+					{
+						DrawParams.DrawPoint(PreviewActor.SampledRootMotion[Index], FColor::Green);
+					}
+				}
+			}
+		}
+
+
+#if ENABLE_ANIM_DEBUG
+		const float DebugDrawSamplerSize = CVarDatabasePreviewDebugDrawSamplerSize.GetValueOnAnyThread();
+		if (DebugDrawSamplerSize > UE_KINDA_SMALL_NUMBER)
+		{
+			// drawing the pose extracted from the Sampler to visually compare with the pose features and the mesh drawing
+			FMemMark Mark(FMemStack::Get());
+			FCompactPose Pose;
+			Pose.SetBoneContainer(&PreviewActor.GetAnimPreviewInstance()->GetRequiredBonesOnAnyThread());
+
+			PreviewActor.Sampler.ExtractPose(PreviewActor.CurrentTime, Pose);
+			MirrorDataCache.MirrorPose(Pose);
+
+			const FTransform RootTransform = MirrorDataCache.MirrorTransform(PreviewActor.Sampler.ExtractRootTransform(PreviewActor.CurrentTime));
+
+			FCSPose<FCompactPose> ComponentSpacePose;
+			ComponentSpacePose.InitPose(MoveTemp(Pose));
+
+			for (int32 BoneIndex = 0; BoneIndex < ComponentSpacePose.GetPose().GetNumBones(); ++BoneIndex)
+			{
+				const FTransform BoneWorldTransforms = ComponentSpacePose.GetComponentSpaceTransform(FCompactPoseBoneIndex(BoneIndex)) * RootTransform;
+				DrawParams.DrawPoint(BoneWorldTransforms.GetTranslation(), FColor::Red, DebugDrawSamplerSize);
+			}
+		}
+#endif // ENABLE_ANIM_DEBUG
+	}
+	return true;
+}
+
+const UDebugSkelMeshComponent* FDatabasePreviewActor::GetDebugSkelMeshComponent() const
+{
+	if (ActorPtr != nullptr)
+	{
+		return Cast<UDebugSkelMeshComponent>(ActorPtr->GetRootComponent());
+	}
+	return nullptr;
+}
+
+const UAnimPreviewInstance* FDatabasePreviewActor::GetAnimPreviewInstance() const
+{
+	if (const UDebugSkelMeshComponent* Mesh = GetDebugSkelMeshComponent())
+	{
+		return Mesh->PreviewInstance.Get();
+	}
+	return nullptr;
+}
+
+UAnimPreviewInstance* FDatabasePreviewActor::GetAnimPreviewInstanceInternal()
+{
+	if (ActorPtr != nullptr)
+	{
+		if (UDebugSkelMeshComponent* Mesh = Cast<UDebugSkelMeshComponent>(ActorPtr->GetRootComponent()))
+		{
+			return Mesh->PreviewInstance.Get();
+		}
+	}
+	return nullptr;
+}
+
+// FDatabaseViewModel
+void FDatabaseViewModel::AddReferencedObjects(FReferenceCollector& Collector)
+{
+	Collector.AddReferencedObject(PoseSearchDatabasePtr);
+}
+
+void FDatabaseViewModel::Initialize(UPoseSearchDatabase* InPoseSearchDatabase, const TSharedRef<FDatabasePreviewScene>& InPreviewScene, const TSharedRef<SDatabaseDataDetails>& InDatabaseDataDetails)
+{
+	PoseSearchDatabasePtr = InPoseSearchDatabase;
+	PreviewScenePtr = InPreviewScene;
+	DatabaseDataDetails = InDatabaseDataDetails;
+
+	RemovePreviewActors();
+}
+
+void FDatabaseViewModel::BuildSearchIndex()
+{
+	using namespace UE::PoseSearch;
+	FAsyncPoseSearchDatabasesManagement::RequestAsyncBuildIndex(GetPoseSearchDatabase(), ERequestAsyncBuildFlag::NewRequest);
+}
+
+void FDatabaseViewModel::PreviewBackwardEnd()
+{
+	SetPlayTime(MinPreviewPlayLength, false);
+}
+
+void FDatabaseViewModel::PreviewBackwardStep()
+{
+	const float NewPlayTime = FMath::Clamp(PlayTime - StepDeltaTime, MinPreviewPlayLength, MaxPreviewPlayLength);
+	SetPlayTime(NewPlayTime, false);
+}
+
+void FDatabaseViewModel::PreviewBackward()
+{
+	DeltaTimeMultiplier = -1.f;
+}
+
+void FDatabaseViewModel::PreviewPause()
+{
+	DeltaTimeMultiplier = 0.f;
+}
+
+void FDatabaseViewModel::PreviewForward()
+{
+	DeltaTimeMultiplier = 1.f;
+}
+
+void FDatabaseViewModel::PreviewForwardStep()
+{
+	const float NewPlayTime = FMath::Clamp(PlayTime + StepDeltaTime, MinPreviewPlayLength, MaxPreviewPlayLength);
+	SetPlayTime(NewPlayTime, false);
+}
+
+void FDatabaseViewModel::PreviewForwardEnd()
+{
+	SetPlayTime(MaxPreviewPlayLength, false);
+}
+
+UWorld* FDatabaseViewModel::GetWorld()
+{
+	check(PreviewScenePtr.IsValid());
+	return PreviewScenePtr.Pin()->GetWorld();
+}
+
+void FDatabaseViewModel::OnPreviewActorClassChanged()
+{
+	// todo: implement
+}
+
+void FDatabaseViewModel::Tick(float DeltaSeconds)
+{
+	const float DeltaPlayTime = DeltaSeconds * DeltaTimeMultiplier;
+	if (!FMath::IsNearlyZero(DeltaPlayTime))
+	{
+		PlayTime += DeltaPlayTime;
+		PlayTime = FMath::Clamp(PlayTime, MinPreviewPlayLength, MaxPreviewPlayLength);
+
+		if (const UPoseSearchDatabase* Database = GetPoseSearchDatabase())
+		{
+			if (EAsyncBuildIndexResult::Success == FAsyncPoseSearchDatabasesManagement::RequestAsyncBuildIndex(Database, ERequestAsyncBuildFlag::ContinueRequest))
+			{
+				for (TArray<FDatabasePreviewActor>& PreviewActorGroup : PreviewActors)
+				{
+					for (FDatabasePreviewActor& PreviewActor : PreviewActorGroup)
+					{
+						PreviewActor.UpdatePreviewActor(Database, PlayTime, bQuantizeAnimationToPoseData);
+					}
+				}
+			}
+		}
+	}
+}
+
+void FDatabaseViewModel::RemovePreviewActors()
+{
+	PlayTime = 0.f;
+	DeltaTimeMultiplier = 1.f;
+	MaxPreviewPlayLength = 0.f;
+	MinPreviewPlayLength = 0.f;
+	bIsEditorSelection = true;
+	bDrawQueryVector = false;
+
+	for (TArray<FDatabasePreviewActor>& PreviewActorGroup : PreviewActors)
+	{
+		for (FDatabasePreviewActor& PreviewActor : PreviewActorGroup)
+		{
+			PreviewActor.Destroy();
+		}
+	}
+
+	PreviewActors.Reset();
+}
+
+void FDatabaseViewModel::AddSequenceToDatabase(UAnimSequence* AnimSequence)
+{
+	if (UPoseSearchDatabase* Database = GetPoseSearchDatabase())
+	{
+		FPoseSearchDatabaseSequence NewAsset;
+		NewAsset.Sequence = AnimSequence;
+		Database->AddAnimationAsset(FInstancedStruct::Make(NewAsset));
+	}
+}
+
+void FDatabaseViewModel::AddBlendSpaceToDatabase(UBlendSpace* BlendSpace)
+{
+	if (UPoseSearchDatabase* Database = GetPoseSearchDatabase())
+	{
+		FPoseSearchDatabaseBlendSpace NewAsset;
+		NewAsset.BlendSpace = BlendSpace;
+		Database->AddAnimationAsset(FInstancedStruct::Make(NewAsset));
+	}
+}
+
+void FDatabaseViewModel::AddAnimCompositeToDatabase(UAnimComposite* AnimComposite)
+{
+	if (UPoseSearchDatabase* Database = GetPoseSearchDatabase())
+	{
+		FPoseSearchDatabaseAnimComposite NewAsset;
+		NewAsset.AnimComposite = AnimComposite;
+		Database->AddAnimationAsset(FInstancedStruct::Make(NewAsset));
+	}
+}
+
+void FDatabaseViewModel::AddAnimMontageToDatabase(UAnimMontage* AnimMontage)
+{
+	if (UPoseSearchDatabase* Database = GetPoseSearchDatabase())
+	{
+		FPoseSearchDatabaseAnimMontage NewAsset;
+		NewAsset.AnimMontage = AnimMontage;
+		Database->AddAnimationAsset(FInstancedStruct::Make(NewAsset));
+	}
+}
+
+void FDatabaseViewModel::AddMultiSequenceToDatabase()
+{
+	if (UPoseSearchDatabase* Database = GetPoseSearchDatabase())
+	{
+		FPoseSearchDatabaseMultiSequence NewAsset;
+		Database->AddAnimationAsset(FInstancedStruct::Make(NewAsset));
+	}
+}
+
+bool FDatabaseViewModel::DeleteFromDatabase(int32 AnimationAssetIndex)
+{
+	if (UPoseSearchDatabase* Database = GetPoseSearchDatabase())
+	{
+		if (const FPoseSearchDatabaseAnimationAssetBase* DatabaseAnimationAssetBase = Database->GetAnimationAssetBase(AnimationAssetIndex))
+		{
+			if (UAnimSequenceBase* AnimSequenceBase = Cast<UAnimSequenceBase>(DatabaseAnimationAssetBase->GetAnimationAsset()))
+			{
+				bool bModified = false;
+				for (int32 NotifyIndex = AnimSequenceBase->Notifies.Num() - 1; NotifyIndex >= 0; --NotifyIndex)
+				{
+					const FAnimNotifyEvent& NotifyEvent = AnimSequenceBase->Notifies[NotifyIndex];
+					if (const UAnimNotifyState_PoseSearchBranchIn* PoseSearchBranchIn = Cast<UAnimNotifyState_PoseSearchBranchIn>(NotifyEvent.NotifyStateClass))
+					{
+						if (PoseSearchBranchIn->Database == Database)
+						{
+							AnimSequenceBase->Notifies.RemoveAt(NotifyIndex);
+							bModified = true;
+						}
+					}
+				}
+
+				if (bModified)
+				{
+					AnimSequenceBase->RefreshCacheData();	
+					AnimSequenceBase->Modify();
+				}
+			}
+
+			Database->RemoveAnimationAssetAt(AnimationAssetIndex);
+			Database->Modify();
+
+			return true;
+		}
+	}
+
+	return false;
+}
+
+void FDatabaseViewModel::SetDisableReselection(int32 AnimationAssetIndex, bool bEnabled)
+{
+	if (UPoseSearchDatabase* Database = GetPoseSearchDatabase())
+	{
+		if (FPoseSearchDatabaseAnimationAssetBase* DatabaseAnimationAsset = Database->GetMutableAnimationAssetBase(AnimationAssetIndex))
+		{
+			DatabaseAnimationAsset->SetDisableReselection(bEnabled);
+		}
+	}
+}
+
+bool FDatabaseViewModel::IsDisableReselection(int32 AnimationAssetIndex) const
+{
+	if (const UPoseSearchDatabase* Database = GetPoseSearchDatabase())
+	{
+		if (const FPoseSearchDatabaseAnimationAssetBase* DatabaseAnimationAsset = Database->GetAnimationAssetBase(AnimationAssetIndex))
+		{
+			return DatabaseAnimationAsset->IsDisableReselection();
+		}
+	}
+
+	return false;
+}
+
+void FDatabaseViewModel::SetIsEnabled(int32 AnimationAssetIndex, bool bEnabled)
+{
+	if (UPoseSearchDatabase* Database = GetPoseSearchDatabase())
+	{
+		if (FPoseSearchDatabaseAnimationAssetBase* DatabaseAnimationAsset = Database->GetMutableAnimationAssetBase(AnimationAssetIndex))
+		{
+			Database->Modify();
+
+			DatabaseAnimationAsset->SetIsEnabled(bEnabled);
+		}
+	}
+}
+
+bool FDatabaseViewModel::IsEnabled(int32 AnimationAssetIndex) const
+{
+	if (const UPoseSearchDatabase* Database = GetPoseSearchDatabase())
+	{
+		if (const FPoseSearchDatabaseAnimationAssetBase* DatabaseAnimationAsset = Database->GetAnimationAssetBase(AnimationAssetIndex))
+		{
+			return DatabaseAnimationAsset->IsEnabled();
+		}
+	}
+
+	return false;
+}
+
+int32 FDatabaseViewModel::SetSelectedNode(int32 PoseIdx, bool bClearSelection, bool bDrawQuery, TConstArrayView<float> InQueryVector)
+{
+	int32 SelectedSourceAssetIdx = INDEX_NONE;
+
+	if (bClearSelection)
+	{
+		RemovePreviewActors();
+	}
+
+	bIsEditorSelection = false;
+	bDrawQueryVector = bDrawQuery;
+	QueryVector = InQueryVector;
+
+	if (const UPoseSearchDatabase* Database = GetPoseSearchDatabase())
+	{
+		if (EAsyncBuildIndexResult::Success == FAsyncPoseSearchDatabasesManagement::RequestAsyncBuildIndex(Database, ERequestAsyncBuildFlag::ContinueRequest))
+		{
+			const FSearchIndex& SearchIndex = Database->GetSearchIndex();
+			if (SearchIndex.PoseMetadata.IsValidIndex(PoseIdx))
+			{
+				const uint32 IndexAssetIndex = SearchIndex.PoseMetadata[PoseIdx].GetAssetIndex();
+				if (SearchIndex.Assets.IsValidIndex(IndexAssetIndex))
+				{
+					const FSearchIndexAsset& IndexAsset = SearchIndex.Assets[IndexAssetIndex];
+					const FPoseSearchDatabaseAnimationAssetBase* DatabaseAnimationAsset = Database->GetAnimationAssetBase(IndexAsset.GetSourceAssetIdx());
+					check(DatabaseAnimationAsset);
+					int32 PreviewActorGroupIndex = INDEX_NONE;
+					bool bIsSynchronizedRootTransformOriginValid = false;
+					FTransform SynchronizedRootTransformOrigin = FTransform::Identity;
+					for (int32 RoleIndex = 0; RoleIndex < DatabaseAnimationAsset->GetNumRoles(); ++RoleIndex)
+					{
+						FDatabasePreviewActor PreviewActor;
+						const UE::PoseSearch::FRole Role = DatabaseAnimationAsset->GetRole(RoleIndex);
+						const FTransform& RootTransformOrigin = DatabaseAnimationAsset->GetRootTransformOriginForRole(Role);
+						if (PreviewActor.SpawnPreviewActor(GetWorld(), Database, IndexAssetIndex, Role, RootTransformOrigin, bIsSynchronizedRootTransformOriginValid ? &SynchronizedRootTransformOrigin : nullptr, PoseIdx))
+						{
+							if (PreviewActorGroupIndex == INDEX_NONE)
+							{
+								PreviewActorGroupIndex = PreviewActors.AddDefaulted();
+							}
+
+							if (!bIsSynchronizedRootTransformOriginValid)
+							{
+								bIsSynchronizedRootTransformOriginValid = true;
+								SynchronizedRootTransformOrigin = PreviewActor.GetRootTransformOrigin();
+							}
+
+							MaxPreviewPlayLength = FMath::Max(MaxPreviewPlayLength, IndexAsset.GetLastSampleTime(Database->Schema->SampleRate) - PreviewActor.GetPlayTimeOffset());
+							MinPreviewPlayLength = FMath::Min(MinPreviewPlayLength, IndexAsset.GetFirstSampleTime(Database->Schema->SampleRate) - PreviewActor.GetPlayTimeOffset());
+							PreviewActors[PreviewActorGroupIndex].Add(PreviewActor);
+							SelectedSourceAssetIdx = IndexAsset.GetSourceAssetIdx();
+						}
+					}
+				}
+			}
+
+			DatabaseDataDetails.Pin()->Reconstruct();
+
+			for (TArray<FDatabasePreviewActor>& PreviewActorGroup : PreviewActors)
+			{
+				for (FDatabasePreviewActor& PreviewActor : PreviewActorGroup)
+				{
+					PreviewActor.UpdatePreviewActor(Database, PlayTime, bQuantizeAnimationToPoseData);
+				}
+			}
+
+			SetPlayTime(0.f, false);
+		}
+	}
+
+	ProcessSelectedActor(nullptr);
+
+	return SelectedSourceAssetIdx;
+}
+
+void FDatabaseViewModel::SetSelectedNodes(const TArrayView<TSharedPtr<FDatabaseAssetTreeNode>>& InSelectedNodes)
+{
+	RemovePreviewActors();
+
+	if (const UPoseSearchDatabase* Database = GetPoseSearchDatabase())
+	{
+		if (EAsyncBuildIndexResult::Success == FAsyncPoseSearchDatabasesManagement::RequestAsyncBuildIndex(Database, ERequestAsyncBuildFlag::ContinueRequest))
+		{
+			TMap<int32, int32> AssociatedAssetIndices;
+			for (int32 i = 0; i < InSelectedNodes.Num(); ++i)
+			{
+				AssociatedAssetIndices.FindOrAdd(InSelectedNodes[i]->SourceAssetIdx) = i;
+			}
+
+			const FSearchIndex& SearchIndex = Database->GetSearchIndex();
+			for (int32 IndexAssetIndex = 0; IndexAssetIndex < SearchIndex.Assets.Num(); ++IndexAssetIndex)
+			{
+				const FSearchIndexAsset& IndexAsset = SearchIndex.Assets[IndexAssetIndex];
+				if (AssociatedAssetIndices.Find(IndexAsset.GetSourceAssetIdx()))
+				{
+					const FPoseSearchDatabaseAnimationAssetBase* DatabaseAnimationAsset = Database->GetAnimationAssetBase(IndexAsset.GetSourceAssetIdx());
+					check(DatabaseAnimationAsset);
+					int32 PreviewActorGroupIndex = INDEX_NONE;
+					bool bIsSynchronizedRootTransformOriginValid = false;
+					FTransform SynchronizedRootTransformOrigin = FTransform::Identity;
+					for (int32 RoleIndex = 0; RoleIndex < DatabaseAnimationAsset->GetNumRoles(); ++RoleIndex)
+					{
+						FDatabasePreviewActor PreviewActor;
+						const UE::PoseSearch::FRole Role = DatabaseAnimationAsset->GetRole(RoleIndex);
+						const FTransform& RootTransformOrigin = DatabaseAnimationAsset->GetRootTransformOriginForRole(Role);
+						if (PreviewActor.SpawnPreviewActor(GetWorld(), Database, IndexAssetIndex, Role, RootTransformOrigin, bIsSynchronizedRootTransformOriginValid ? &SynchronizedRootTransformOrigin : nullptr))
+						{
+							if (PreviewActorGroupIndex == INDEX_NONE)
+							{
+								PreviewActorGroupIndex = PreviewActors.AddDefaulted();
+							}
+
+							if (!bIsSynchronizedRootTransformOriginValid)
+							{
+								bIsSynchronizedRootTransformOriginValid = true;
+								SynchronizedRootTransformOrigin = PreviewActor.GetRootTransformOrigin();
+							}
+
+							MaxPreviewPlayLength = FMath::Max(MaxPreviewPlayLength, IndexAsset.GetLastSampleTime(Database->Schema->SampleRate) - IndexAsset.GetFirstSampleTime(Database->Schema->SampleRate));
+							PreviewActors[PreviewActorGroupIndex].Add(PreviewActor);
+						}
+					}
+				}
+			}
+
+			DatabaseDataDetails.Pin()->Reconstruct();
+			for (TArray<FDatabasePreviewActor>& PreviewActorGroup : PreviewActors)
+			{
+				for (FDatabasePreviewActor& PreviewActor : PreviewActorGroup)
+				{
+					PreviewActor.UpdatePreviewActor(Database, PlayTime, bQuantizeAnimationToPoseData);
+				}
+			}
+		}
+
+		ProcessSelectedActor(nullptr);
+	}
+}
+
+void FDatabaseViewModel::ProcessSelectedActor(AActor* Actor)
+{
+	SelectedActorIndexAssetIndex = INDEX_NONE;
+
+	for (const TArray<FDatabasePreviewActor>& PreviewActorGroup : PreviewActors)
+	{
+		for (const FDatabasePreviewActor& PreviewActor : PreviewActorGroup)
+		{
+			if (PreviewActor.GetActor() == Actor)
+			{
+				SelectedActorIndexAssetIndex = PreviewActor.GetIndexAssetIndex();
+				return;
+			}
+		}
+	}
+}
+
+void FDatabaseViewModel::SetDrawQueryVector(bool bValue)
+{
+	if (bDrawQueryVector != bValue)
+	{
+		bDrawQueryVector = bValue;
+		DatabaseDataDetails.Pin()->Reconstruct();
+	}
+}
+
+const FSearchIndexAsset* FDatabaseViewModel::GetSelectedActorIndexAsset() const
+{
+	if (const UPoseSearchDatabase* Database = GetPoseSearchDatabase())
+	{
+		if (EAsyncBuildIndexResult::Success == FAsyncPoseSearchDatabasesManagement::RequestAsyncBuildIndex(Database, ERequestAsyncBuildFlag::ContinueRequest))
+		{
+			const FSearchIndex& SearchIndex = Database->GetSearchIndex();
+			if (SearchIndex.Assets.IsValidIndex(SelectedActorIndexAssetIndex))
+			{
+				return &SearchIndex.Assets[SelectedActorIndexAssetIndex];
+			}
+		}
+	}
+	return nullptr;
+}
+
+TRange<double> FDatabaseViewModel::GetPreviewPlayRange() const
+{
+	constexpr double ViewRangeSlack = 0.2;
+	return TRange<double>(MinPreviewPlayLength - ViewRangeSlack, MaxPreviewPlayLength + ViewRangeSlack);
+}
+
+float FDatabaseViewModel::GetPlayTime() const
+{
+	return PlayTime;
+}
+
+void FDatabaseViewModel::SetPlayTime(float NewPlayTime, bool bInTickPlayTime)
+{
+	NewPlayTime = FMath::Clamp(NewPlayTime, MinPreviewPlayLength, MaxPreviewPlayLength);
+	DeltaTimeMultiplier = bInTickPlayTime ? DeltaTimeMultiplier : 0.0f;
+
+	if (!FMath::IsNearlyEqual(PlayTime, NewPlayTime))
+	{
+		PlayTime = NewPlayTime;
+		
+		if (const UPoseSearchDatabase* Database = GetPoseSearchDatabase())
+		{
+			if (EAsyncBuildIndexResult::Success == FAsyncPoseSearchDatabasesManagement::RequestAsyncBuildIndex(Database, ERequestAsyncBuildFlag::ContinueRequest))
+			{
+				for (TArray<FDatabasePreviewActor>& PreviewActorGroup : PreviewActors)
+				{
+					for (FDatabasePreviewActor& PreviewActor : PreviewActorGroup)
+					{
+						PreviewActor.UpdatePreviewActor(Database, PlayTime, bQuantizeAnimationToPoseData);
+					}
+				}
+			}
+		}
+	}
+}
+
+bool FDatabaseViewModel::GetAnimationTime(int32 SourceAssetIdx, float& CurrentPlayTime, FVector& BlendParameters) const
+{
+	if (const UPoseSearchDatabase* Database = GetPoseSearchDatabase())
+	{
+		if (EAsyncBuildIndexResult::Success == FAsyncPoseSearchDatabasesManagement::RequestAsyncBuildIndex(Database, ERequestAsyncBuildFlag::ContinueRequest))
+		{
+			const FSearchIndex& SearchIndex = Database->GetSearchIndex();
+			for (const TArray<FDatabasePreviewActor>& PreviewActorGroup : PreviewActors)
+			{
+				for (const FDatabasePreviewActor& PreviewActor : PreviewActorGroup)
+				{
+					if (PreviewActor.GetIndexAssetIndex() >= 0 && PreviewActor.GetIndexAssetIndex() < SearchIndex.Assets.Num())
+					{
+						const FSearchIndexAsset& IndexAsset = SearchIndex.Assets[PreviewActor.GetIndexAssetIndex()];
+						if (IndexAsset.GetSourceAssetIdx() == SourceAssetIdx)
+						{
+							CurrentPlayTime = PreviewActor.GetSampler().ToNormalizedTime(PlayTime + IndexAsset.GetFirstSampleTime(Database->Schema->SampleRate) + PreviewActor.GetPlayTimeOffset());
+							BlendParameters = IndexAsset.GetBlendParameters();
+							return true;
+						}
+					}
+				}
+			}
+
+			for (const FSearchIndexAsset& IndexAsset : SearchIndex.Assets)
+			{
+				if (IndexAsset.GetSourceAssetIdx() == SourceAssetIdx)
+				{
+					CurrentPlayTime = PlayTime + IndexAsset.GetFirstSampleTime(Database->Schema->SampleRate);
+					BlendParameters = IndexAsset.GetBlendParameters();
+
+					const bool bIsBlendSpace = Database->GetAnimationAssetStruct(IndexAsset).GetPtr<FPoseSearchDatabaseBlendSpace>() != nullptr;
+					if (bIsBlendSpace && !FMath::IsNearlyEqual(MaxPreviewPlayLength, MinPreviewPlayLength))
+					{
+						CurrentPlayTime = (CurrentPlayTime - MaxPreviewPlayLength) / (MaxPreviewPlayLength - MinPreviewPlayLength);
+					}
+					return true;
+				}
+			}
+		}
+	}
+
+	CurrentPlayTime = 0.f;
+	BlendParameters = FVector::ZeroVector;
+	return false;
+}
+}
diff --git a/Engine/Plugins/Animation/PoseSearch/Source/Runtime/Private/PoseSearchAssetIndexer.cpp b/Engine/Plugins/Animation/PoseSearch/Source/Runtime/Private/PoseSearchAssetIndexer.cpp
new file mode 100644
--- /dev/null
+++ b/Engine/Plugins/Animation/PoseSearch/Source/Runtime/Private/PoseSearchAssetIndexer.cpp
@@ -0,0 +1,936 @@
+// Copyright Epic Games, Inc. All Rights Reserved.
+
+#if WITH_EDITOR
+
+#include "PoseSearch/PoseSearchAssetIndexer.h"
+#include "AnimationRuntime.h"
+#include "Animation/MirrorDataTable.h"
+#include "PoseSearch/PoseSearchAnimNotifies.h"
+#include "PoseSearch/PoseSearchAssetSampler.h"
+#include "PoseSearch/PoseSearchContext.h"
+#include "PoseSearch/PoseSearchDefines.h"
+#include "PoseSearch/PoseSearchFeatureChannel.h"
+#include "PoseSearch/PoseSearchIndex.h"
+#include "PoseSearch/PoseSearchSchema.h"
+#include "PoseSearch/PoseSearchDatabase.h"
+
+namespace UE::PoseSearch
+{
+
+#if ENABLE_ANIM_DEBUG
+static TAutoConsoleVariable<bool> CVarMotionMatchTestDisableIndexerCaching(TEXT("a.MotionMatch.TestDisableIndexerCaching"), false, TEXT("Disable Motion Matching Indexer Caching"));
+static TAutoConsoleVariable<int32> CVarMotionMatchTestExtractPoseDeterminismNumIterations(TEXT("a.MotionMatch.TestExtractPoseDeterminismNumIterations"), 0, TEXT("Test Motion Matching ExtractPose Determinism via this NumIterations retries"));
+#endif // ENABLE_ANIM_DEBUG
+
+//////////////////////////////////////////////////////////////////////////
+// FSamplingParam helpers
+struct FSamplingParam
+{
+	float WrappedParam = 0.0f;
+	int32 NumCycles = 0;
+
+	// If the animation can't loop, WrappedParam contains the clamped value and whatever is left is stored here
+	float Extrapolation = 0.0f;
+};
+
+static FSamplingParam WrapOrClampSamplingParam(bool bCanWrap, float SamplingParamExtent, float SamplingParam)
+{
+	// This is a helper function used by both time and distance sampling. A schema may specify time or distance
+	// offsets that are multiple cycles of a clip away from the current pose being sampled.
+	// And that time or distance offset may before the beginning of the clip (SamplingParam < 0.0f)
+	// or after the end of the clip (SamplingParam > SamplingParamExtent). So this function
+	// helps determine how many cycles need to be applied and what the wrapped value should be, clamping
+	// if necessary.
+
+	FSamplingParam Result;
+
+	Result.WrappedParam = SamplingParam;
+
+	const bool bIsSamplingParamExtentKindaSmall = SamplingParamExtent <= UE_KINDA_SMALL_NUMBER;
+	if (!bIsSamplingParamExtentKindaSmall && bCanWrap)
+	{
+		if (SamplingParam < 0.0f)
+		{
+			while (Result.WrappedParam < 0.0f)
+			{
+				Result.WrappedParam += SamplingParamExtent;
+				++Result.NumCycles;
+			}
+		}
+
+		else
+		{
+			while (Result.WrappedParam > SamplingParamExtent)
+			{
+				Result.WrappedParam -= SamplingParamExtent;
+				++Result.NumCycles;
+			}
+		}
+	}
+
+	const float ParamClamped = FMath::Clamp(Result.WrappedParam, 0.0f, SamplingParamExtent);
+	if (ParamClamped != Result.WrappedParam)
+	{
+		check(bIsSamplingParamExtentKindaSmall || !bCanWrap);
+		Result.Extrapolation = Result.WrappedParam - ParamClamped;
+		Result.WrappedParam = ParamClamped;
+	}
+
+	return Result;
+}
+
+//////////////////////////////////////////////////////////////////////////
+// FAssetSamplingContext
+FAssetSamplingContext::FAssetSamplingContext(const UPoseSearchDatabase& Database)
+{
+	check(Database.Schema);
+	BaseCostBias = Database.BaseCostBias;
+	LoopingCostBias = Database.LoopingCostBias;
+}
+
+//////////////////////////////////////////////////////////////////////////
+// FAnimationAssetSamplers
+void FAnimationAssetSamplers::Reset()
+{
+	AnimationAssetSamplers.Reset();
+	MirrorDataCaches.Reset();
+}
+
+int32 FAnimationAssetSamplers::Num() const
+{
+	check(AnimationAssetSamplers.Num() == MirrorDataCaches.Num());
+	return AnimationAssetSamplers.Num();
+}
+
+float FAnimationAssetSamplers::GetPlayLength() const
+{
+	float PlayLength = 0.f;
+	for (const FAnimationAssetSampler* Sampler : AnimationAssetSamplers)
+	{
+		PlayLength = FMath::Max(PlayLength, Sampler->GetPlayLength());
+	}
+	return PlayLength;
+}
+
+bool FAnimationAssetSamplers::IsLoopable() const
+{
+	float CommonPlayLength = -1.f;
+	for (const FAnimationAssetSampler* Sampler : AnimationAssetSamplers)
+	{
+		if (!Sampler->IsLoopable())
+		{
+			return false;
+		}
+
+		if (CommonPlayLength < 0.f)
+		{
+			CommonPlayLength = Sampler->GetPlayLength();
+		}
+		else if (!FMath::IsNearlyEqual(CommonPlayLength, Sampler->GetPlayLength()))
+		{
+			return false;
+		}
+	}
+	return true;
+}
+
+void FAnimationAssetSamplers::ExtractPoseSearchNotifyStates(float Time, TFunction<bool(UAnimNotifyState_PoseSearchBase*)> ProcessPoseSearchBase) const
+{
+	for (const FAnimationAssetSampler* Sampler : AnimationAssetSamplers)
+	{
+		Sampler->ExtractPoseSearchNotifyStates(Time, ProcessPoseSearchBase);
+	}
+}
+
+bool FAnimationAssetSamplers::ProcessAllAnimNotifyEvents(TFunction<bool(TConstArrayView<FAnimNotifyEvent>)> ProcessAnimNotifyEvents) const
+{
+	for (const FAnimationAssetSampler* Sampler : AnimationAssetSamplers)
+	{
+		if (ProcessAnimNotifyEvents(Sampler->GetAllAnimNotifyEvents()))
+		{
+			return true;
+		}
+	}
+	return false;
+}
+
+const FString FAnimationAssetSamplers::GetAssetName() const
+{
+	FString Name;
+	Name.Reserve(256);
+	bool bAddComma = false;
+	for (const FAnimationAssetSampler* Sampler : AnimationAssetSamplers)
+	{
+		if (bAddComma)
+		{
+			Name += ", ";
+		}
+		else
+		{
+			bAddComma = true;
+		}
+
+		Name += GetNameSafe(Sampler->GetAsset());
+	}
+	return Name;
+}
+
+FTransform FAnimationAssetSamplers::ExtractRootTransform(float Time, int32 RoleIndex) const
+{
+	return AnimationAssetSamplers[RoleIndex]->ExtractRootTransform(Time);
+}
+
+FTransform FAnimationAssetSamplers::GetTotalRootTransform(int32 RoleIndex) const
+{
+	return AnimationAssetSamplers[RoleIndex]->GetTotalRootTransform();
+}
+
+void FAnimationAssetSamplers::ExtractPose(float Time, FCompactPose& OutPose, int32 RoleIndex) const
+{
+	AnimationAssetSamplers[RoleIndex]->ExtractPose(Time, OutPose);
+}
+
+FTransform FAnimationAssetSamplers::MirrorTransform(const FTransform& InTransform, int32 RoleIndex) const
+{
+	return MirrorDataCaches[RoleIndex]->MirrorTransform(InTransform);
+}
+
+void FAnimationAssetSamplers::MirrorPose(FCompactPose& Pose, int32 RoleIndex) const
+{
+	MirrorDataCaches[RoleIndex]->MirrorPose(Pose);
+}
+
+//////////////////////////////////////////////////////////////////////////
+// FAssetIndexer
+FAssetIndexer::FAssetIndexer(const TConstArrayView<FBoneContainer> InBoneContainers, const FSearchIndexAsset& InSearchIndexAsset, const FAssetSamplingContext& InSamplingContext,
+	const UPoseSearchSchema& InSchema, const FAnimationAssetSamplers& InAssetSamplers, const FRoleToIndex& InRoleToIndex, const FFloatInterval& InExtrapolationTimeInterval)
+: BoneContainers(InBoneContainers)
+, CachedEntries()
+, SearchIndexAsset(InSearchIndexAsset)
+, SamplingContext(InSamplingContext)
+, Schema(InSchema)
+, AssetSamplers(InAssetSamplers)
+, RoleToIndex(InRoleToIndex)
+, ExtrapolationTimeInterval(InExtrapolationTimeInterval)
+{
+	check(BoneContainers.Num() == AssetSamplers.Num() && BoneContainers.Num() == RoleToIndex.Num());
+	check(IsValid(RoleToIndex));
+}
+
+void FAssetIndexer::AssignWorkingData(int32 InStartPoseIdx, TArrayView<float> InOutFeatureVectorTable, TArrayView<FPoseMetadata> InOutPoseMetadata)
+{
+	const int32 NumIndexedPoses = GetNumIndexedPoses();
+
+	StartPoseIdx = InStartPoseIdx;
+	FeatureVectorTable = MakeArrayView(InOutFeatureVectorTable.GetData() + Schema.SchemaCardinality * StartPoseIdx, Schema.SchemaCardinality * NumIndexedPoses);
+	PoseMetadata = MakeArrayView(InOutPoseMetadata.GetData() + StartPoseIdx, NumIndexedPoses);
+}
+
+void FAssetIndexer::Process(int32 AssetIdx)
+{
+	bProcessFailed = false;
+
+	// Generate pose metadata
+	const float PlayLength = GetPlayLength();
+	for (int32 SampleIdx = GetBeginSampleIdx(); SampleIdx != GetEndSampleIdx(); ++SampleIdx)
+	{
+		const float SampleTime = FMath::Min(CalculateSampleTime(SampleIdx), PlayLength);
+		float CostAddend = SamplingContext.BaseCostBias;
+		bool bBlockTransition = false;
+
+		AssetSamplers.ExtractPoseSearchNotifyStates(SampleTime, [&bBlockTransition, &CostAddend](const UAnimNotifyState_PoseSearchBase* PoseSearchNotify)
+			{
+				if (PoseSearchNotify->GetClass()->IsChildOf<UAnimNotifyState_PoseSearchBlockTransition>())
+				{
+					bBlockTransition = true;
+				}
+				else if (const UAnimNotifyState_PoseSearchModifyCost* ModifyCostNotify = Cast<const UAnimNotifyState_PoseSearchModifyCost>(PoseSearchNotify))
+				{
+					CostAddend = ModifyCostNotify->CostAddend;
+				}
+				return true;
+			});
+
+		if (AssetSamplers.IsLoopable())
+		{
+			CostAddend += SamplingContext.LoopingCostBias;
+		}
+
+		const int32 ValueOffset = (StartPoseIdx + GetVectorIdx(SampleIdx)) * Schema.SchemaCardinality;
+		check(ValueOffset >= 0 && AssetIdx >= 0);
+		PoseMetadata[GetVectorIdx(SampleIdx)] = FPoseMetadata(ValueOffset, AssetIdx, bBlockTransition, CostAddend);
+	}
+
+	// Generate pose features data
+	if (Schema.SchemaCardinality > 0)
+	{
+		for (const TObjectPtr<UPoseSearchFeatureChannel>& ChannelPtr : Schema.GetChannels())
+		{
+			if (!ChannelPtr->IndexAsset(*this))
+			{
+				bProcessFailed = true;
+				break;
+			}
+		}
+	}
+
+	// Computing stats
+	if (!bProcessFailed)
+	{
+		ComputeStats();
+	}
+}
+
+void FAssetIndexer::ComputeStats()
+{
+	Stats = FStats();
+
+	for (const FRoleToIndexPair& RoleToIndexPair : RoleToIndex)
+	{
+		const FRole& Role = RoleToIndexPair.Key;
+		const int32 RoleIndex = RoleToIndexPair.Value;
+		const FBoneReference& RootBoneReference = Schema.GetBoneReferences(Role)[RootSchemaBoneIdx];
+
+		for (int32 SampleIdx = GetBeginSampleIdx(); SampleIdx != GetEndSampleIdx(); ++SampleIdx)
+		{
+			const float SampleTime = FMath::Min(CalculateSampleTime(SampleIdx), GetPlayLength());
+
+			bool AnyClamped = false;
+			const FTransform TrajTransformsPast = GetTransform(SampleTime - FiniteDelta, RoleIndex, AnyClamped, RootBoneReference);
+			if (!AnyClamped)
+			{
+				const FTransform TrajTransformsPresent = GetTransform(SampleTime, RoleIndex, AnyClamped, RootBoneReference);
+				if (!AnyClamped)
+				{
+					const FTransform TrajTransformsFuture = GetTransform(SampleTime + FiniteDelta, RoleIndex, AnyClamped, RootBoneReference);
+					if (!AnyClamped)
+					{
+						// if any transform is clamped we just skip the sample entirely
+						const FVector LinearVelocityPresent = (TrajTransformsPresent.GetTranslation() - TrajTransformsPast.GetTranslation()) / FiniteDelta;
+						const FVector LinearVelocityFuture = (TrajTransformsFuture.GetTranslation() - TrajTransformsPresent.GetTranslation()) / FiniteDelta;
+						const FVector LinearAcceleration = (LinearVelocityFuture - LinearVelocityPresent) / FiniteDelta;
+
+						const float Speed = LinearVelocityPresent.Length();
+						const float Acceleration = LinearAcceleration.Length();
+
+						Stats.AccumulatedSpeed += Speed;
+						Stats.MaxSpeed = FMath::Max(Stats.MaxSpeed, Speed);
+
+						Stats.AccumulatedAcceleration += Acceleration;
+						Stats.MaxAcceleration = FMath::Max(Stats.MaxAcceleration, Acceleration);
+
+						++Stats.NumAccumulatedSamples;
+					}
+				}
+			}
+		}
+	}
+}
+
+void FAssetIndexer::GetSampleInfo(float SampleTime, int32 RoleIndex, FTransform& OutRootTransform, float& OutClipTime, bool& bOutClamped) const
+{
+	const float PlayLength = GetPlayLength();
+	const bool bCanWrap = AssetSamplers.IsLoopable();
+
+	float MainRelativeTime = SampleTime;
+	if (SampleTime < 0.0f && bCanWrap)
+	{
+		// In this case we're sampling a loop backwards, so MainRelativeTime must adjust so the number of cycles is
+		// counted correctly.
+		MainRelativeTime += PlayLength;
+	}
+
+	const FSamplingParam SamplingParam = WrapOrClampSamplingParam(bCanWrap, PlayLength, MainRelativeTime);
+
+	if (FMath::Abs(SamplingParam.Extrapolation) > SMALL_NUMBER)
+	{
+		bOutClamped = true;
+		OutClipTime = SamplingParam.WrappedParam + SamplingParam.Extrapolation;
+		OutRootTransform = AssetSamplers.ExtractRootTransform(OutClipTime, RoleIndex);
+	}
+	else
+	{
+		bOutClamped = false;
+		OutClipTime = SamplingParam.WrappedParam;
+		OutRootTransform = FTransform::Identity;
+
+		// Find the remaining motion deltas after wrapping
+		FTransform RootMotionRemainder = AssetSamplers.ExtractRootTransform(OutClipTime, RoleIndex);
+
+  		const bool bNegativeSampleTime = SampleTime < 0.f;
+		if (SamplingParam.NumCycles > 0 || bNegativeSampleTime)
+		{
+			const FTransform RootMotionLast = AssetSamplers.GetTotalRootTransform(RoleIndex);
+
+			// Determine how to accumulate motion for every cycle of the anim. If the sample
+			// had to be clamped, this motion will end up not getting applied below.
+			// Also invert the accumulation direction if the requested sample was wrapped backwards.
+			FTransform RootMotionPerCycle = RootMotionLast;
+
+			if (bNegativeSampleTime)
+			{
+				RootMotionPerCycle = RootMotionPerCycle.Inverse();
+			}
+			
+			// Invert motion deltas if we wrapped backwards
+			if (bNegativeSampleTime)
+			{
+				RootMotionRemainder.SetToRelativeTransform(RootMotionLast);
+			}
+
+			// Note if the sample was clamped, no motion will be applied here because NumCycles will be zero
+			int32 CyclesRemaining = SamplingParam.NumCycles;
+			while (CyclesRemaining--)
+			{
+				OutRootTransform = RootMotionPerCycle * OutRootTransform;
+			}
+		}
+
+		OutRootTransform = RootMotionRemainder * OutRootTransform;
+	}
+}
+
+FTransform FAssetIndexer::MirrorTransform(const FTransform& Transform, int32 RoleIndex) const
+{
+	return SearchIndexAsset.IsMirrored() ? AssetSamplers.MirrorTransform(Transform, RoleIndex) : Transform;
+}
+
+FAssetIndexer::FCachedEntry& FAssetIndexer::GetEntry(float SampleTime)
+{
+	using namespace UE::Anim;
+	
+	bool bDisableCaching = false;
+#if ENABLE_ANIM_DEBUG
+	bDisableCaching = CVarMotionMatchTestDisableIndexerCaching.GetValueOnAnyThread();
+#endif // ENABLE_ANIM_DEBUG
+
+	SampleTime = FMath::Clamp(SampleTime, ExtrapolationTimeInterval.Min, ExtrapolationTimeInterval.Max);
+
+	FCachedEntry* Entry = bDisableCaching ? nullptr : CachedEntries.Find(SampleTime);
+	if (!Entry)
+	{
+		Entry = &CachedEntries.Add(SampleTime);
+		Entry->SampleTime = SampleTime;
+
+		const bool bLoopable = AssetSamplers.IsLoopable();
+		const float PlayLength = GetPlayLength();
+		const int32 AssetSamplersNum = AssetSamplers.Num();
+
+		Entry->RootTransform.SetNum(AssetSamplersNum);
+		Entry->ComponentSpacePose.SetNum(AssetSamplersNum);
+		for (int32 RoleIndex = 0; RoleIndex < AssetSamplers.AnimationAssetSamplers.Num(); ++RoleIndex)
+		{
+			if (!BoneContainers[RoleIndex].IsValid())
+			{
+				UE_LOG(LogPoseSearch,
+					Warning,
+					TEXT("Invalid BoneContainer encountered in FAssetIndexer::GetEntry. Asset: %s. Schema: %s. BoneContainerAsset: %s. NumBoneIndices: %d"),
+					*AssetSamplers.GetAssetName(),
+					*GetNameSafe(&Schema),
+					*GetNameSafe(BoneContainers[RoleIndex].GetAsset()),
+					BoneContainers[RoleIndex].GetCompactPoseNumBones());
+			}
+
+			FTransform SampleRootTransform;
+			bool bSampleClamped;
+			float CurrentTime;
+
+			GetSampleInfo(SampleTime, RoleIndex, SampleRootTransform, CurrentTime, bSampleClamped);
+
+			if (!bLoopable)
+			{
+				CurrentTime = FMath::Clamp(CurrentTime, 0.f, PlayLength);
+			}
+
+			FMemMark Mark(FMemStack::Get());
+			FCompactPose Pose;
+			Pose.SetBoneContainer(&BoneContainers[RoleIndex]);
+			AssetSamplers.ExtractPose(CurrentTime, Pose, RoleIndex);
+
+#if ENABLE_ANIM_DEBUG
+			const int32 NumIterations = CVarMotionMatchTestExtractPoseDeterminismNumIterations.GetValueOnAnyThread();
+			for (int32 IterationIndex = 0; IterationIndex < NumIterations; ++IterationIndex)
+			{
+				FCompactPose TestPose;
+				TestPose.SetBoneContainer(&BoneContainers[RoleIndex]);
+				AssetSamplers.ExtractPose(CurrentTime, TestPose, RoleIndex);
+
+				const TConstArrayView<FTransform> Bones = Pose.GetBones();
+				const TConstArrayView<FTransform> TestBones = TestPose.GetBones();
+				if (Bones.Num() != TestBones.Num())
+				{
+					UE_LOG(LogPoseSearch, Warning, TEXT("FAssetIndexer::GetEntry - ExtractPose is not deterministic"));
+				}
+				else
+				{
+					for (int32 BoneIndex = 0; BoneIndex < Bones.Num(); ++BoneIndex)
+					{
+						if (FMemory::Memcmp(&Bones[BoneIndex], &TestBones[BoneIndex], sizeof(FTransform)) != 0)
+						{
+							UE_LOG(LogPoseSearch, Warning, TEXT("FAssetIndexer::GetEntry - ExtractPose is not deterministic"));
+						}
+					}
+				}
+			}
+#endif // ENABLE_ANIM_DEBUG
+
+			if (SearchIndexAsset.IsMirrored())
+			{
+				AssetSamplers.MirrorPose(Pose, RoleIndex);
+			}
+
+			FCSPose<FCompactPose> StackComponentSpacePose;
+			StackComponentSpacePose.InitPose(MoveTemp(Pose));
+			Entry->ComponentSpacePose[RoleIndex].CopyPose(StackComponentSpacePose);
+
+			Entry->RootTransform[RoleIndex] = SampleRootTransform;
+			Entry->bClamped |= bSampleClamped;
+		}
+	}
+
+	return *Entry;
+}
+
+// returns the transform in component space for the bone indexed by Schema->BoneReferences[SchemaBoneIdx] at SampleTime seconds
+FTransform FAssetIndexer::GetComponentSpaceTransform(float SampleTime, const FRole& Role, bool& bClamped, int8 SchemaBoneIdx)
+{
+	using namespace UE::PoseSearch;
+
+	FCachedEntry& Entry = GetEntry(SampleTime);
+	bClamped = Entry.bClamped;
+
+	const int32 RoleIndex = RoleToIndex[Role];
+	const FBoneReference& BoneReference = Schema.GetBoneReferences(Role)[SchemaBoneIdx];
+	return CalculateComponentSpaceTransform(Entry, BoneReference, RoleIndex);
+}
+
+// returns the transform in animation space for the bone indexed by Schema->BoneReferences[SchemaBoneIdx] at SampleTime seconds
+FTransform FAssetIndexer::GetTransform(float SampleTime, const FRole& Role, bool& bClamped, int8 SchemaBoneIdx)
+{
+	using namespace UE::PoseSearch;
+
+	FCachedEntry& Entry = GetEntry(SampleTime);
+	bClamped = Entry.bClamped;
+
+	const int32 RoleIndex = RoleToIndex[Role];
+	const FBoneReference& BoneReference = Schema.GetBoneReferences(Role)[SchemaBoneIdx];
+	return CalculateComponentSpaceTransform(Entry, BoneReference, RoleIndex) * MirrorTransform(Entry.RootTransform[RoleIndex], RoleIndex);
+}
+
+// returns the transform in animation space for the BoneReference at SampleTime seconds
+FTransform FAssetIndexer::GetTransform(float SampleTime, int32 RoleIndex, bool& bClamped, const FBoneReference& BoneReference)
+{
+	FCachedEntry& Entry = GetEntry(SampleTime);
+	bClamped = Entry.bClamped;
+	return CalculateComponentSpaceTransform(Entry, BoneReference, RoleIndex) * MirrorTransform(Entry.RootTransform[RoleIndex], RoleIndex);
+}
+
+FTransform FAssetIndexer::CalculateComponentSpaceTransform(FAssetIndexer::FCachedEntry& Entry, const FBoneReference& BoneReference, int32 RoleIndex)
+{
+	const FCompactPoseBoneIndex CompactBoneIndex = BoneContainers[RoleIndex].MakeCompactPoseIndex(FMeshPoseBoneIndex(BoneReference.BoneIndex));
+	return Entry.ComponentSpacePose[RoleIndex].GetComponentSpaceTransform(CompactBoneIndex);
+}
+
+float FAssetIndexer::CalculateSampleTime(int32 SampleIdx) const
+{
+	return SampleIdx / float(Schema.SampleRate);
+}
+
+bool FAssetIndexer::GetSampleRotation(FQuat& OutSampleRotation, float SampleTimeOffset, float OriginTimeOffset, int32 SampleIdx, int8 SchemaSampleBoneIdx, int8 SchemaOriginBoneIdx, const FRole& SampleRole, const FRole& OriginRole, EPermutationTimeType PermutationTimeType, int32 SamplingAttributeId)
+{
+	using namespace UE::PoseSearch;
+
+	float PermutationSampleTimeOffset = 0.f;
+	float PermutationOriginTimeOffset = 0.f;
+	UPoseSearchFeatureChannel::GetPermutationTimeOffsets(PermutationTimeType, CalculatePermutationTimeOffset(), PermutationSampleTimeOffset, PermutationOriginTimeOffset);
+
+	const float Time = CalculateSampleTime(SampleIdx);
+	const float SampleTime = Time + SampleTimeOffset + PermutationSampleTimeOffset;
+	const float OriginTime = Time + OriginTimeOffset + PermutationOriginTimeOffset;
+
+	// @todo: add support for SchemaSampleBoneIdx
+	if (SchemaOriginBoneIdx != RootSchemaBoneIdx)
+	{
+		UE_LOG(LogPoseSearch,
+			Error,
+			TEXT("FAssetIndexer::GetSampleRotation: support for non root origin bones not implemented (bone: '%s', schema: '%s'"), 
+			*Schema.GetBoneReferences(OriginRole)[SchemaOriginBoneIdx].BoneName.ToString(),
+			*GetNameSafe(&Schema));
+	}
+
+	if (SamplingAttributeId >= 0)
+	{
+		if (ProcessAllAnimNotifyEvents([&OutSampleRotation, &SampleRole, &OriginRole, SamplingAttributeId, this, OriginTime](const TConstArrayView<FAnimNotifyEvent> AnimNotifyEvents)
+			{
+				for (const FAnimNotifyEvent& AnimNotifyEvent : AnimNotifyEvents)
+				{
+					if (const UAnimNotifyState_PoseSearchSamplingAttribute* SamplingAttribute = Cast<UAnimNotifyState_PoseSearchSamplingAttribute>(AnimNotifyEvent.NotifyStateClass))
+					{
+						if (SamplingAttribute->SamplingAttributeId == SamplingAttributeId)
+						{
+							bool bUnused;
+							if (SamplingAttribute->Bone.BoneName != NAME_None)
+							{
+								FBoneReference TempBoneReference = SamplingAttribute->Bone;
+								const int32 SampleRoleIndex = RoleToIndex[SampleRole];
+								TempBoneReference.Initialize(BoneContainers[SampleRoleIndex].GetSkeletonAsset());
+								if (TempBoneReference.HasValidSetup())
+								{
+									const float SamplingAttributeTime = AnimNotifyEvent.GetTime();
+									const FTransform RootBoneTransform = GetTransform(OriginTime, OriginRole, bUnused, RootSchemaBoneIdx);
+									const FTransform SamplingAttributeBoneTransform = GetTransform(SamplingAttributeTime, SampleRoleIndex, bUnused, TempBoneReference);
+									OutSampleRotation = RootBoneTransform.InverseTransformRotation(SamplingAttributeBoneTransform.GetRotation());
+									return true;
+								}
+
+								UE_LOG(LogPoseSearch, Error, TEXT("FAssetIndexer::GetSampleRotation: required UAnimNotifyState_PoseSearchSamplingAttribute in '%s' has an invalid Bone"), *AssetSamplers.GetAssetName());
+								return false;
+							}
+
+							const FTransform RootBoneTransform = GetTransform(OriginTime, OriginRole, bUnused, RootSchemaBoneIdx);
+							OutSampleRotation = RootBoneTransform.InverseTransformRotation(SamplingAttribute->Rotation);
+							return true;
+						}
+					}
+				}
+				return false;
+			}))
+		{
+			return true;
+		}
+
+		UE_LOG(LogPoseSearch, Error, TEXT("FAssetIndexer::GetSamplePositionInternal: required UAnimNotifyState_PoseSearchSamplingAttribute not found in '%s'"), *AssetSamplers.GetAssetName());
+		OutSampleRotation = FQuat::Identity;
+		return false;
+	}
+
+	bool bUnused;
+	const FTransform RootBoneTransform = GetTransform(OriginTime, OriginRole, bUnused, RootSchemaBoneIdx);
+	const FTransform SampleBoneTransform = GetTransform(SampleTime, SampleRole, bUnused, SchemaSampleBoneIdx);
+	OutSampleRotation = RootBoneTransform.InverseTransformRotation(SampleBoneTransform.GetRotation());
+	return true;
+}
+
+bool FAssetIndexer::GetSamplePosition(FVector& OutSamplePosition, float SampleTimeOffset, float OriginTimeOffset, int32 SampleIdx, int8 SchemaSampleBoneIdx, int8 SchemaOriginBoneIdx, const FRole& SampleRole, const FRole& OriginRole, EPermutationTimeType PermutationTimeType, int32 SamplingAttributeId)
+{
+	float PermutationSampleTimeOffset = 0.f;
+	float PermutationOriginTimeOffset = 0.f;
+	UPoseSearchFeatureChannel::GetPermutationTimeOffsets(PermutationTimeType, CalculatePermutationTimeOffset(), PermutationSampleTimeOffset, PermutationOriginTimeOffset);
+
+	const float Time = CalculateSampleTime(SampleIdx);
+	const float SampleTime = Time + SampleTimeOffset + PermutationSampleTimeOffset;
+	const float OriginTime = Time + OriginTimeOffset + PermutationOriginTimeOffset;
+	
+	bool bUnused;
+	return GetSamplePositionInternal(OutSamplePosition, SampleTime, OriginTime, bUnused, SchemaSampleBoneIdx, SchemaOriginBoneIdx, SampleRole, OriginRole, SamplingAttributeId);
+}
+
+bool FAssetIndexer::GetSamplePositionInternal(FVector& OutSamplePosition, float SampleTime, float OriginTime, bool& bClamped, int8 SchemaSampleBoneIdx, int8 SchemaOriginBoneIdx, const FRole& SampleRole, const FRole& OriginRole, int32 SamplingAttributeId)
+{
+	using namespace UE::PoseSearch;
+
+	if (SamplingAttributeId >= 0)
+	{
+		if (ProcessAllAnimNotifyEvents([&OutSamplePosition, OriginTime, &bClamped, SchemaOriginBoneIdx, &SampleRole, OriginRole, SamplingAttributeId, this](const TConstArrayView<FAnimNotifyEvent> AnimNotifyEvents)
+			{
+				for (const FAnimNotifyEvent& AnimNotifyEvent : AnimNotifyEvents)
+				{
+					if (const UAnimNotifyState_PoseSearchSamplingAttribute* SamplingAttribute = Cast<UAnimNotifyState_PoseSearchSamplingAttribute>(AnimNotifyEvent.NotifyStateClass))
+					{
+						if (SamplingAttribute->SamplingAttributeId == SamplingAttributeId)
+						{
+							bool bUnused;
+							if (SamplingAttribute->Bone.BoneName != NAME_None)
+							{
+								FBoneReference TempBoneReference = SamplingAttribute->Bone;
+								const int32 SampleRoleIndex = RoleToIndex[SampleRole];
+								TempBoneReference.Initialize(BoneContainers[SampleRoleIndex].GetSkeletonAsset());
+								if (TempBoneReference.HasValidSetup())
+								{
+									const float SamplingAttributeTime = AnimNotifyEvent.GetTime();
+									const FTransform RootBoneTransform = GetTransform(OriginTime, OriginRole, bUnused, RootSchemaBoneIdx);
+									const FTransform SamplingAttributeBoneTransform = GetTransform(SamplingAttributeTime, SampleRoleIndex, bClamped, TempBoneReference);
+									if (SchemaOriginBoneIdx == RootSchemaBoneIdx)
+									{
+										OutSamplePosition = RootBoneTransform.InverseTransformPosition(SamplingAttributeBoneTransform.GetTranslation());
+									}
+									else
+									{
+										bool bOriginClamped;
+										const FTransform OriginBoneTransform = GetTransform(OriginTime, OriginRole, bOriginClamped, SchemaOriginBoneIdx);
+										bClamped |= bOriginClamped;
+										const FVector DeltaBoneTranslation = SamplingAttributeBoneTransform.GetTranslation() - OriginBoneTransform.GetTranslation();
+										OutSamplePosition = RootBoneTransform.InverseTransformVector(DeltaBoneTranslation);
+									}
+									return true;
+								}
+
+								UE_LOG(LogPoseSearch, Error, TEXT("FAssetIndexer::GetSamplePositionInternal: required UAnimNotifyState_PoseSearchSamplingAttribute in '%s' has an invalid Bone"), *AssetSamplers.GetAssetName());
+								return false;
+							}
+
+							const FTransform RootBoneTransform = GetTransform(OriginTime, OriginRole, bUnused, RootSchemaBoneIdx);
+							OutSamplePosition = RootBoneTransform.InverseTransformPosition(SamplingAttribute->Position);
+							return true;
+						}
+					}
+				}
+				return false;
+			}))
+		{
+			return true;
+		}
+
+		UE_LOG(LogPoseSearch, Error, TEXT("FAssetIndexer::GetSamplePositionInternal: required UAnimNotifyState_PoseSearchSamplingAttribute not found in '%s'"), *AssetSamplers.GetAssetName());
+		OutSamplePosition = FVector::ZeroVector;
+		return false;
+	}
+
+	bool bUnused;
+	const FTransform RootBoneTransform = GetTransform(OriginTime, OriginRole, bUnused, RootSchemaBoneIdx);
+	const FTransform SampleBoneTransform = GetTransform(SampleTime, SampleRole, bClamped, SchemaSampleBoneIdx);
+	if (SchemaOriginBoneIdx == RootSchemaBoneIdx)
+	{
+		OutSamplePosition = RootBoneTransform.InverseTransformPosition(SampleBoneTransform.GetTranslation());
+		return true;
+	}
+
+	bool bOriginClamped;
+	const FTransform OriginBoneTransform = GetTransform(OriginTime, OriginRole, bOriginClamped, SchemaOriginBoneIdx);
+	bClamped |= bOriginClamped;
+	const FVector DeltaBoneTranslation = SampleBoneTransform.GetTranslation() - OriginBoneTransform.GetTranslation();
+	OutSamplePosition = RootBoneTransform.InverseTransformVector(DeltaBoneTranslation);
+	return true;
+}
+
+bool FAssetIndexer::GetSampleVelocity(FVector& OutSampleVelocity, float SampleTimeOffset, float OriginTimeOffset, int32 SampleIdx, int8 SchemaSampleBoneIdx, int8 SchemaOriginBoneIdx, const FRole& SampleRole, const FRole& OriginRole, bool bUseCharacterSpaceVelocities, EPermutationTimeType PermutationTimeType, int32 SamplingAttributeId)
+{
+	float PermutationSampleTimeOffset = 0.f;
+	float PermutationOriginTimeOffset = 0.f;
+	UPoseSearchFeatureChannel::GetPermutationTimeOffsets(PermutationTimeType, CalculatePermutationTimeOffset(), PermutationSampleTimeOffset, PermutationOriginTimeOffset);
+
+	const float Time = CalculateSampleTime(SampleIdx);
+	const float SampleTime = Time + SampleTimeOffset + PermutationSampleTimeOffset;
+	const float OriginTime = Time + OriginTimeOffset + PermutationOriginTimeOffset;
+	
+	if (SamplingAttributeId >= 0)
+	{
+		if (ProcessAllAnimNotifyEvents([&OutSampleVelocity, SchemaSampleBoneIdx, SchemaOriginBoneIdx, bUseCharacterSpaceVelocities, &SampleRole, &OriginRole, SamplingAttributeId, this, OriginTime](const TConstArrayView<FAnimNotifyEvent> AnimNotifyEvents)
+			{
+				for (const FAnimNotifyEvent& AnimNotifyEvent : AnimNotifyEvents)
+				{
+					if (const UAnimNotifyState_PoseSearchSamplingAttribute* SamplingAttribute = Cast<UAnimNotifyState_PoseSearchSamplingAttribute>(AnimNotifyEvent.NotifyStateClass))
+					{
+						if (SamplingAttribute->SamplingAttributeId == SamplingAttributeId)
+						{
+							bool bUnused, bClampedPast;
+							FVector BonePositionPast, BonePositionPresent;
+							if (SamplingAttribute->Bone.BoneName != NAME_None)
+							{
+								FBoneReference TempBoneReference = SamplingAttribute->Bone;
+								const int32 SampleRoleIndex = RoleToIndex[SampleRole];
+								TempBoneReference.Initialize(BoneContainers[SampleRoleIndex].GetSkeletonAsset());
+								if (TempBoneReference.HasValidSetup())
+								{
+									const float SamplingAttributeTime = AnimNotifyEvent.GetTime();
+
+									if (GetSamplePositionInternal(BonePositionPast, SamplingAttributeTime - FiniteDelta, bUseCharacterSpaceVelocities ? OriginTime - FiniteDelta : OriginTime, bClampedPast, SchemaSampleBoneIdx, SchemaOriginBoneIdx, SampleRole, OriginRole, SamplingAttributeId) &&
+										GetSamplePositionInternal(BonePositionPresent, SamplingAttributeTime, OriginTime, bUnused, SchemaSampleBoneIdx, SchemaOriginBoneIdx, SampleRole, OriginRole, SamplingAttributeId))
+									{
+										if (!bClampedPast)
+										{
+											OutSampleVelocity = (BonePositionPresent - BonePositionPast) / FiniteDelta;
+											return true;
+										}
+
+										FVector BonePositionFuture;
+										if (GetSamplePositionInternal(BonePositionFuture, SamplingAttributeTime + FiniteDelta, bUseCharacterSpaceVelocities ? OriginTime + FiniteDelta : OriginTime, bUnused, SchemaSampleBoneIdx, SchemaOriginBoneIdx, SampleRole, OriginRole, SamplingAttributeId))
+										{
+											OutSampleVelocity = (BonePositionFuture - BonePositionPresent) / FiniteDelta;
+											return true;
+										}
+									}
+
+									return false;
+								}
+
+								UE_LOG(LogPoseSearch, Error, TEXT("FAssetIndexer::GetSampleVelocity: required UAnimNotifyState_PoseSearchSamplingAttribute in '%s' has an invalid Bone"), *AssetSamplers.GetAssetName());
+								return false;
+							}
+
+							const FTransform RootBoneTransform = GetTransform(OriginTime, OriginRole, bUnused, RootSchemaBoneIdx);
+							OutSampleVelocity = RootBoneTransform.InverseTransformPosition(SamplingAttribute->LinearVelocity);
+							return true;
+						}
+					}
+				}
+				return false;
+			}))
+		{
+			return true;
+		}
+
+		UE_LOG(LogPoseSearch, Error, TEXT("FAssetIndexer::GetSampleVelocity: required UAnimNotifyState_PoseSearchSamplingAttribute not found in '%s'"), *AssetSamplers.GetAssetName());
+		OutSampleVelocity = FVector::ZeroVector;
+		return false;
+	}
+
+	bool bUnused, bClampedPast;
+	FVector BonePositionPast, BonePositionPresent;
+	if (GetSamplePositionInternal(BonePositionPast, SampleTime - FiniteDelta, bUseCharacterSpaceVelocities ? OriginTime - FiniteDelta : OriginTime, bClampedPast, SchemaSampleBoneIdx, SchemaOriginBoneIdx, SampleRole, OriginRole, INDEX_NONE) &&
+		GetSamplePositionInternal(BonePositionPresent, SampleTime, OriginTime, bUnused, SchemaSampleBoneIdx, SchemaOriginBoneIdx, SampleRole, OriginRole, INDEX_NONE))
+	{
+		if (!bClampedPast)
+		{
+			OutSampleVelocity = (BonePositionPresent - BonePositionPast) / FiniteDelta;
+			return true;
+		}
+
+		FVector BonePositionFuture;
+		if (GetSamplePositionInternal(BonePositionFuture, SampleTime + FiniteDelta, bUseCharacterSpaceVelocities ? OriginTime + FiniteDelta : OriginTime, bUnused, SchemaSampleBoneIdx, SchemaOriginBoneIdx, SampleRole, OriginRole, INDEX_NONE))
+		{
+			OutSampleVelocity = (BonePositionFuture - BonePositionPresent) / FiniteDelta;
+			return true;
+		}
+	}
+
+	OutSampleVelocity = FVector::ZeroVector;
+	return false;
+}
+
+bool FAssetIndexer::ProcessAllAnimNotifyEvents(TFunction<bool(TConstArrayView<FAnimNotifyEvent>)> ProcessAnimNotifyEvents) const
+{
+	return AssetSamplers.ProcessAllAnimNotifyEvents(ProcessAnimNotifyEvents);
+}
+
+const FString FAssetIndexer::GetAssetName() const
+{
+	return AssetSamplers.GetAssetName();
+}
+
+float FAssetIndexer::GetPlayLength() const
+{
+	return AssetSamplers.GetPlayLength();
+}
+
+int32 FAssetIndexer::GetBeginSampleIdx() const
+{
+	return SearchIndexAsset.GetBeginSampleIdx();
+}
+
+int32 FAssetIndexer::GetEndSampleIdx() const
+{
+	return SearchIndexAsset.GetEndSampleIdx();
+}
+
+int32 FAssetIndexer::GetNumIndexedPoses() const
+{
+	return SearchIndexAsset.GetNumPoses();
+}
+
+int32 FAssetIndexer::GetVectorIdx(int32 SampleIdx) const
+{
+	return SampleIdx - GetBeginSampleIdx();
+}
+
+TArrayView<float> FAssetIndexer::GetPoseVector(int32 SampleIdx) const
+{
+	return MakeArrayView(&FeatureVectorTable[GetVectorIdx(SampleIdx) * Schema.SchemaCardinality], Schema.SchemaCardinality);
+}
+
+const UPoseSearchSchema* FAssetIndexer::GetSchema() const
+{
+	return &Schema;
+}
+
+#if WITH_EDITOR
+float FAssetIndexer::CalculatePermutationTimeOffset() const
+{
+	check(Schema.PermutationsSampleRate > 0 && SearchIndexAsset.IsInitialized());
+	const float PermutationTimeOffset = Schema.PermutationsTimeOffset + SearchIndexAsset.GetPermutationIdx() / float(Schema.PermutationsSampleRate);
+	return PermutationTimeOffset;
+}
+#endif // WITH_EDITOR
+
+#if ENABLE_ANIM_DEBUG
+void FAssetIndexer::CompareCachedEntries(const FAssetIndexer& Other) const
+{
+	if (CachedEntries.Num() != Other.CachedEntries.Num())
+	{
+		UE_LOG(LogPoseSearch, Warning, TEXT("CompareCachedEntries - FAssetIndexer::CachedEntries::Num is not deterministic"));
+	}
+	else
+	{
+		for (const TPair<float, FCachedEntry>& Pair : CachedEntries)
+		{
+			if (const FCachedEntry* OtherEntry = Other.CachedEntries.Find(Pair.Key))
+			{
+				const FCachedEntry* Entry = &Pair.Value;
+				if (Entry->SampleTime != OtherEntry->SampleTime)
+				{
+					UE_LOG(LogPoseSearch, Warning, TEXT("CompareCachedEntries - FAssetIndexer::CachedEntries::SampleTime is not deterministic (%f, %f)"), Entry->SampleTime, OtherEntry->SampleTime);
+				}
+
+				if (Entry->bClamped != OtherEntry->bClamped)
+				{
+					UE_LOG(LogPoseSearch, Warning, TEXT("CompareCachedEntries - FAssetIndexer::CachedEntries::bClamped is not deterministic"));
+				}
+
+				if (Entry->RootTransform.Num() != OtherEntry->RootTransform.Num())
+				{
+					UE_LOG(LogPoseSearch, Warning, TEXT("CompareCachedEntries - FAssetIndexer::CachedEntries::RootTransform::Num is not deterministic"));
+				}
+				else
+				{
+					for (int32 RoleIndex = 0; RoleIndex < Entry->RootTransform.Num(); ++RoleIndex)
+					{
+						if (FMemory::Memcmp(&Entry->RootTransform[RoleIndex], &OtherEntry->RootTransform[RoleIndex], sizeof(FTransform)) != 0)
+						{
+							UE_LOG(LogPoseSearch, Warning, TEXT("CompareCachedEntries - FAssetIndexer::CachedEntries::RootTransform[%d] is not deterministic"), RoleIndex);
+						}
+					}
+				}
+
+				if (Entry->ComponentSpacePose.Num() != OtherEntry->ComponentSpacePose.Num())
+				{
+					UE_LOG(LogPoseSearch, Warning, TEXT("CompareCachedEntries - FAssetIndexer::CachedEntries::ComponentSpacePose::Num is not deterministic"));
+				}
+				else
+				{
+					for (int32 RoleIndex = 0; RoleIndex < Entry->ComponentSpacePose.Num(); ++RoleIndex)
+					{
+						if (Entry->ComponentSpacePose[RoleIndex].GetComponentSpaceFlags() != OtherEntry->ComponentSpacePose[RoleIndex].GetComponentSpaceFlags())
+						{
+							UE_LOG(LogPoseSearch, Warning, TEXT("CompareCachedEntries - FAssetIndexer::CachedEntries::ComponentSpacePose[%d]::ComponentSpaceFlags is not deterministic"), RoleIndex);
+						}
+
+						const TConstArrayView<FTransform> Bones = Entry->ComponentSpacePose[RoleIndex].GetPose().GetBones();
+						const TConstArrayView<FTransform> OtherBones = OtherEntry->ComponentSpacePose[RoleIndex].GetPose().GetBones();
+						if (Bones.Num() != OtherBones.Num())
+						{
+							UE_LOG(LogPoseSearch, Warning, TEXT("CompareCachedEntries - FAssetIndexer::CachedEntries::ComponentSpacePose[%d]::Bones is not deterministic"), RoleIndex);
+						}
+						else
+						{
+							for (int32 BoneIndex = 0; BoneIndex < Bones.Num(); ++BoneIndex)
+							{
+								if (FMemory::Memcmp(&Bones[BoneIndex], &OtherBones[BoneIndex], sizeof(FTransform)) != 0)
+								{
+									UE_LOG(LogPoseSearch, Warning, TEXT("CompareCachedEntries - FAssetIndexer::CachedEntries::ComponentSpacePose[%d]::Bones[%d] is not deterministic"), RoleIndex, BoneIndex);
+								}
+							}
+						}
+					}
+				}
+			}
+			else
+			{
+				UE_LOG(LogPoseSearch, Warning, TEXT("CompareCachedEntries - FAssetIndexer::CachedEntries is not deterministic. Missing CachedEntry at time %f"), Pair.Key);
+			}
+		}
+	}
+}
+#endif // ENABLE_ANIM_DEBUG
+
+} // namespace UE::PoseSearch
+#endif // WITH_EDITOR
