Move PoseSearch/BlendStack out of experimental folder

diff --git a/Engine/Plugins/Animation/PoseSearch/Source/Runtime/Private/PoseSearchHistory.cpp b/Engine/Plugins/Animation/PoseSearch/Source/Runtime/Private/PoseSearchHistory.cpp
new file mode 100644
--- /dev/null
+++ b/Engine/Plugins/Animation/PoseSearch/Source/Runtime/Private/PoseSearchHistory.cpp
@@ -0,0 +1,826 @@
+// Copyright Epic Games, Inc. All Rights Reserved.
+
+#include "PoseSearch/PoseSearchHistory.h"
+#include "AnimationRuntime.h"
+#include "Animation/AnimInstanceProxy.h"
+#include "Animation/AnimNodeBase.h"
+#include "Animation/AnimRootMotionProvider.h"
+#include "Animation/SkeletonRemapping.h"
+#include "Animation/SkeletonRemappingRegistry.h"
+#include "BonePose.h"
+#include "DrawDebugHelpers.h"
+#include "PoseSearch/PoseSearchResult.h"
+#include "PoseSearch/PoseSearchDatabase.h"
+#include "PoseSearch/PoseSearchDefines.h"
+
+IMPLEMENT_ANIMGRAPH_MESSAGE(UE::PoseSearch::IPoseHistoryProvider);
+
+namespace UE::PoseSearch
+{
+
+#if ENABLE_DRAW_DEBUG && ENABLE_ANIM_DEBUG
+TAutoConsoleVariable<bool> CVarAnimPoseHistoryDebugDrawPose(TEXT("a.AnimNode.PoseHistory.DebugDrawPose"), false, TEXT("Enable / Disable Pose History Pose DebugDraw"));
+TAutoConsoleVariable<bool> CVarAnimPoseHistoryDebugDrawTrajectory(TEXT("a.AnimNode.PoseHistory.DebugDrawTrajectory"), false, TEXT("Enable / Disable Pose History Trajectory DebugDraw"));
+#endif
+
+/**
+* Algo::LowerBound adapted to TIndexedContainerIterator for use with indexable but not necessarily contiguous containers. Used here with TRingBuffer.
+*
+* Performs binary search, resulting in position of the first element >= Value using predicate
+*
+* @param First TIndexedContainerIterator beginning of range to search through, must be already sorted by SortPredicate
+* @param Last TIndexedContainerIterator end of range
+* @param Value Value to look for
+* @param SortPredicate Predicate for sort comparison, defaults to <
+*
+* @returns Position of the first element >= Value, may be position after last element in range
+*/
+template <typename IteratorType, typename ValueType, typename ProjectionType, typename SortPredicateType>
+FORCEINLINE auto LowerBound(IteratorType First, IteratorType Last, const ValueType& Value, ProjectionType Projection, SortPredicateType SortPredicate) -> decltype(First.GetIndex())
+{
+	using SizeType = decltype(First.GetIndex());
+
+	check(First.GetIndex() <= Last.GetIndex());
+
+	// Current start of sequence to check
+	SizeType Start = First.GetIndex();
+
+	// Size of sequence to check
+	SizeType Size = Last.GetIndex() - Start;
+
+	// With this method, if Size is even it will do one more comparison than necessary, but because Size can be predicted by the CPU it is faster in practice
+	while (Size > 0)
+	{
+		const SizeType LeftoverSize = Size % 2;
+		Size = Size / 2;
+
+		const SizeType CheckIndex = Start + Size;
+		const SizeType StartIfLess = CheckIndex + LeftoverSize;
+
+		auto&& CheckValue = Invoke(Projection, *(First + CheckIndex));
+		Start = SortPredicate(CheckValue, Value) ? StartIfLess : Start;
+	}
+	return Start;
+}
+
+template <typename IteratorType, typename ValueType, typename SortPredicateType = TLess<>()>
+FORCEINLINE auto LowerBound(IteratorType First, IteratorType Last, const ValueType& Value, SortPredicateType SortPredicate) -> decltype(First.GetIndex())
+{
+	return LowerBound(First, Last, Value, FIdentityFunctor(), SortPredicate);
+}
+
+static FBoneIndexType GetRemappedBoneIndexType(FBoneIndexType BoneIndexType, const USkeleton* BoneIndexSkeleton, const USkeleton* LastUpdateSkeleton)
+{
+	// remapping BoneIndexType in case the skeleton used to store history (LastUpdateSkeleton) is different from BoneIndexSkeleton
+	if (LastUpdateSkeleton != nullptr && LastUpdateSkeleton != BoneIndexSkeleton)
+	{
+		const FSkeletonRemapping& SkeletonRemapping = UE::Anim::FSkeletonRemappingRegistry::Get().GetRemapping(BoneIndexSkeleton, LastUpdateSkeleton);
+		if (SkeletonRemapping.IsValid())
+		{
+			BoneIndexType = SkeletonRemapping.GetTargetSkeletonBoneIndex(BoneIndexType);
+		}
+	}
+
+	return BoneIndexType;
+}
+
+static FComponentSpaceTransformIndex GetRemappedComponentSpaceTransformIndex(const USkeleton* BoneIndexSkeleton, const USkeleton* LastUpdateSkeleton, const FBoneToTransformMap& BoneToTransformMap, FBoneIndexType BoneIndexType, bool& bSuccess)
+{
+	check(BoneIndexType != WorldSpaceIndexType);
+
+	FComponentSpaceTransformIndex BoneTransformIndex = FComponentSpaceTransformIndex(BoneIndexType);
+	if (BoneIndexType != ComponentSpaceIndexType)
+	{
+		BoneIndexType = GetRemappedBoneIndexType(BoneIndexType, BoneIndexSkeleton, LastUpdateSkeleton);
+
+		if (!BoneToTransformMap.IsEmpty())
+		{
+			if (const FComponentSpaceTransformIndex* FoundBoneTransformIndex = BoneToTransformMap.Find(BoneTransformIndex))
+			{
+				BoneTransformIndex = *FoundBoneTransformIndex;
+			}
+			else
+			{
+				BoneTransformIndex = RootBoneIndexType;
+				bSuccess = false;
+			}
+		}
+	}
+	return BoneTransformIndex;
+}
+
+static bool LerpEntries(float Time, bool bExtrapolate, const FPoseHistoryEntry& PrevEntry, const FPoseHistoryEntry& NextEntry, const USkeleton* BoneIndexSkeleton, const USkeleton* LastUpdateSkeleton,
+	const FBoneToTransformMap& BoneToTransformMap, FBoneIndexType BoneIndexType, FBoneIndexType ReferenceBoneIndexType, FTransform& OutBoneTransform)
+{
+	bool bSuccess = true;
+
+	const float Denominator = NextEntry.AccumulatedSeconds - PrevEntry.AccumulatedSeconds;
+	float LerpValue = 0.f;
+	if (!FMath::IsNearlyZero(Denominator))
+	{
+		const float Numerator = Time - PrevEntry.AccumulatedSeconds;
+		LerpValue = bExtrapolate ? Numerator / Denominator : FMath::Clamp(Numerator / Denominator, 0.f, 1.f);
+	}
+
+	const FComponentSpaceTransformIndex BoneTransformIndex = GetRemappedComponentSpaceTransformIndex(BoneIndexSkeleton, LastUpdateSkeleton, BoneToTransformMap, BoneIndexType, bSuccess);
+	const FComponentSpaceTransformIndex ReferenceBoneTransformIndex = GetRemappedComponentSpaceTransformIndex(BoneIndexSkeleton, LastUpdateSkeleton, BoneToTransformMap, ReferenceBoneIndexType, bSuccess);
+
+	if (BoneTransformIndex != ComponentSpaceIndexType)
+	{
+		if (ReferenceBoneTransformIndex == ComponentSpaceIndexType)
+		{
+			OutBoneTransform.Blend(
+				PrevEntry.GetComponentSpaceTransform(BoneTransformIndex),
+				NextEntry.GetComponentSpaceTransform(BoneTransformIndex),
+				LerpValue);
+		}
+		else
+		{
+			OutBoneTransform.Blend(
+				PrevEntry.GetComponentSpaceTransform(BoneTransformIndex) * PrevEntry.GetComponentSpaceTransform(ReferenceBoneTransformIndex).Inverse(),
+				NextEntry.GetComponentSpaceTransform(BoneTransformIndex) * NextEntry.GetComponentSpaceTransform(ReferenceBoneTransformIndex).Inverse(),
+				LerpValue);
+		}
+	}
+	else
+	{
+		OutBoneTransform = FTransform::Identity;
+		bSuccess = false;
+		unimplemented();
+	}
+
+	return bSuccess;
+}
+
+//////////////////////////////////////////////////////////////////////////
+// FPoseHistoryEntry
+void FPoseHistoryEntry::Update(float Time, FCSPose<FCompactPose>& ComponentSpacePose, const FBoneToTransformMap& BoneToTransformMap, bool bStoreScales)
+{
+	AccumulatedSeconds = Time;
+
+	const FBoneContainer& BoneContainer = ComponentSpacePose.GetPose().GetBoneContainer();
+	const USkeleton* SkeletonAsset = BoneContainer.GetSkeletonAsset();
+	check(SkeletonAsset);
+	const FReferenceSkeleton& RefSkeleton = SkeletonAsset->GetReferenceSkeleton();
+	const TArray<FTransform>& RefBonePose = RefSkeleton.GetRefBonePose();
+	const int32 NumSkeletonBones = RefSkeleton.GetNum();
+
+	if (BoneToTransformMap.IsEmpty())
+	{
+		// no mapping: we add all the transforms
+		SetNum(NumSkeletonBones, bStoreScales);
+		for (FSkeletonPoseBoneIndex SkeletonBoneIdx(0); SkeletonBoneIdx != NumSkeletonBones; ++SkeletonBoneIdx)
+		{
+			const FCompactPoseBoneIndex CompactBoneIdx = BoneContainer.GetCompactPoseIndexFromSkeletonPoseIndex(SkeletonBoneIdx);
+			SetComponentSpaceTransform(SkeletonBoneIdx.GetInt(), (CompactBoneIdx.IsValid() ? ComponentSpacePose.GetComponentSpaceTransform(CompactBoneIdx) : RefBonePose[SkeletonBoneIdx.GetInt()]));
+		}
+	}
+	else
+	{
+		SetNum(BoneToTransformMap.Num(), true);
+		for (const FBoneToTransformPair& BoneToTransformPair : BoneToTransformMap)
+		{
+			const FSkeletonPoseBoneIndex SkeletonBoneIdx(BoneToTransformPair.Key);
+			const FCompactPoseBoneIndex CompactBoneIdx = BoneContainer.GetCompactPoseIndexFromSkeletonPoseIndex(SkeletonBoneIdx);
+			SetComponentSpaceTransform(BoneToTransformPair.Value, (CompactBoneIdx.IsValid() ? ComponentSpacePose.GetComponentSpaceTransform(CompactBoneIdx) : RefBonePose[SkeletonBoneIdx.GetInt()]));
+		}
+	}
+}
+
+void FPoseHistoryEntry::SetNum(int32 Num, bool bStoreScales)
+{
+	ComponentSpaceRotations.SetNum(Num);
+	ComponentSpacePositions.SetNum(Num);
+	ComponentSpaceScales.SetNum(bStoreScales ? Num : 0);
+}
+
+int32 FPoseHistoryEntry::Num() const
+{
+	return ComponentSpaceRotations.Num();
+}
+
+void FPoseHistoryEntry::SetComponentSpaceTransform(int32 Index, const FTransform& Transform)
+{
+	ComponentSpaceRotations[Index] = FQuat4f(Transform.GetRotation());
+	ComponentSpacePositions[Index] = Transform.GetTranslation();
+	
+	if (!ComponentSpaceScales.IsEmpty())
+	{
+		ComponentSpaceScales[Index] = FVector3f(Transform.GetScale3D());
+	}
+}
+
+FTransform FPoseHistoryEntry::GetComponentSpaceTransform(int32 Index) const
+{
+#if WITH_EDITOR
+	if (Index < 0 || Index >= ComponentSpaceRotations.Num())
+	{
+		UE_LOG(LogPoseSearch, Error, TEXT("FPoseHistoryEntry::GetComponentSpaceTransform - Index %d out of bound [0, %d)"), Index, ComponentSpaceRotations.Num());
+		return FTransform::Identity;
+	}
+#endif // WITH_EDITOR
+
+	check(ComponentSpaceScales.IsEmpty() || ComponentSpaceRotations.Num() == ComponentSpaceScales.Num());
+
+	const FQuat Quat(ComponentSpaceRotations[Index]);
+	const FVector Scale(ComponentSpaceScales.IsEmpty() ? FVector3f::OneVector : ComponentSpaceScales[Index]);
+	return FTransform(Quat, ComponentSpacePositions[Index], Scale);
+}
+
+FArchive& operator<<(FArchive& Ar, FPoseHistoryEntry& Entry)
+{
+	Ar << Entry.ComponentSpaceRotations;
+	Ar << Entry.ComponentSpacePositions;
+	Ar << Entry.ComponentSpaceScales;
+	Ar << Entry.AccumulatedSeconds;
+	return Ar;
+}
+
+//////////////////////////////////////////////////////////////////////////
+// IPoseHistory
+#if ENABLE_DRAW_DEBUG && ENABLE_ANIM_DEBUG
+void IPoseHistory::DebugDraw(FAnimInstanceProxy& AnimInstanceProxy, FColor Color, float Time, float PointSize, bool bExtrapolate) const
+{
+	const FBoneContainer& BoneContainer = AnimInstanceProxy.GetRequiredBones();
+	if (Color.A > 0 && BoneContainer.IsValid())
+	{
+		const USkeleton* Skeleton = BoneContainer.GetSkeletonAsset();
+		FTransform OutBoneTransform;
+
+		const FBoneToTransformMap& BoneToTransformMap = GetBoneToTransformMap();
+		if (BoneToTransformMap.IsEmpty())
+		{
+			for (FSkeletonPoseBoneIndex SkeletonBoneIdx(0); SkeletonBoneIdx != BoneContainer.GetNumBones(); ++SkeletonBoneIdx)
+			{
+				const FCompactPoseBoneIndex CompactBoneIdx = BoneContainer.GetCompactPoseIndexFromSkeletonPoseIndex(SkeletonBoneIdx);
+				if (GetTransformAtTime(Time, OutBoneTransform, Skeleton, CompactBoneIdx.GetInt(), WorldSpaceIndexType, bExtrapolate))
+				{
+					AnimInstanceProxy.AnimDrawDebugPoint(OutBoneTransform.GetTranslation(), PointSize, Color, false, 0.f, ESceneDepthPriorityGroup::SDPG_Foreground);
+				}
+			}
+		}
+		else
+		{
+			for (const FBoneToTransformPair& BoneToTransformPair : BoneToTransformMap)
+			{
+				const FSkeletonPoseBoneIndex SkeletonBoneIdx(BoneToTransformPair.Key);
+				const FCompactPoseBoneIndex CompactBoneIdx = BoneContainer.GetCompactPoseIndexFromSkeletonPoseIndex(SkeletonBoneIdx);
+				if (GetTransformAtTime(Time, OutBoneTransform, Skeleton, CompactBoneIdx.GetInt(), WorldSpaceIndexType, bExtrapolate))
+				{
+					AnimInstanceProxy.AnimDrawDebugPoint(OutBoneTransform.GetTranslation(), PointSize, Color, false, 0.f, ESceneDepthPriorityGroup::SDPG_Foreground);
+				}
+			}
+		}
+	}
+}
+#endif // ENABLE_DRAW_DEBUG && ENABLE_ANIM_DEBUG
+
+//////////////////////////////////////////////////////////////////////////
+// FArchivedPoseHistory
+void FArchivedPoseHistory::InitFrom(const IPoseHistory* PoseHistory)
+{
+	Trajectory.Samples.Reset();
+	BoneToTransformMap.Reset();
+	Entries.Reset();
+
+	if (PoseHistory)
+	{
+		Trajectory = PoseHistory->GetTrajectory();
+		BoneToTransformMap = PoseHistory->GetBoneToTransformMap();
+			
+		const int32 NumEntries = PoseHistory->GetNumEntries();
+		Entries.SetNum(NumEntries);
+
+		for (int32 EntryIndex = 0; EntryIndex < NumEntries; ++EntryIndex)
+		{
+			Entries[EntryIndex] = PoseHistory->GetEntry(EntryIndex);
+			// validating input PoseHistory to have Entries properly sorted by time
+			check(EntryIndex == 0 || (Entries[EntryIndex - 1].AccumulatedSeconds <= Entries[EntryIndex].AccumulatedSeconds));
+		}
+	}
+}
+
+bool FArchivedPoseHistory::GetTransformAtTime(float Time, FTransform& OutBoneTransform, const USkeleton* BoneIndexSkeleton, FBoneIndexType BoneIndexType, FBoneIndexType ReferenceBoneIndexType, bool bExtrapolate) const
+{
+	static_assert(RootBoneIndexType == 0 && ComponentSpaceIndexType == FBoneIndexType(-1) && WorldSpaceIndexType == FBoneIndexType(-2)); // some assumptions
+	check(BoneIndexType != ComponentSpaceIndexType && BoneIndexType != WorldSpaceIndexType);
+	
+	bool bSuccess = false;
+	
+	const bool bApplyComponentToWorld = ReferenceBoneIndexType == WorldSpaceIndexType;
+	FTransform ComponentToWorld = FTransform::Identity;
+	if (bApplyComponentToWorld)
+	{
+		ComponentToWorld = Trajectory.GetSampleAtTime(Time, bExtrapolate).GetTransform();
+		ReferenceBoneIndexType = ComponentSpaceIndexType;
+	}
+
+	const int32 NumEntries = Entries.Num();
+	if (NumEntries > 0)
+	{
+		int32 NextIdx = 0;
+		int32 PrevIdx = 0;
+
+		if (NumEntries > 1)
+		{
+			const int32 LowerBoundIdx = Algo::LowerBound(Entries, Time, [](const FPoseHistoryEntry& Entry, float Value) { return Value > Entry.AccumulatedSeconds; });
+			NextIdx = FMath::Clamp(LowerBoundIdx, 1, NumEntries - 1);
+			PrevIdx = NextIdx - 1;
+		}
+	
+		const FPoseHistoryEntry& PrevEntry = Entries[PrevIdx];
+		const FPoseHistoryEntry& NextEntry = Entries[NextIdx];
+
+		bSuccess = LerpEntries(Time, bExtrapolate, PrevEntry, NextEntry, BoneIndexSkeleton, nullptr, BoneToTransformMap, BoneIndexType, ReferenceBoneIndexType, OutBoneTransform);
+		if (bApplyComponentToWorld)
+		{
+			OutBoneTransform *= ComponentToWorld;
+		}
+	}
+	else
+	{
+		OutBoneTransform = ComponentToWorld;
+	}
+	
+	return bSuccess;
+}
+
+#if ENABLE_DRAW_DEBUG && ENABLE_ANIM_DEBUG
+void FArchivedPoseHistory::DebugDraw(const UWorld* World, FColor Color) const
+{
+	if (Color.A > 0 && !Trajectory.Samples.IsEmpty())
+	{
+		TArray<FTransform, TInlineAllocator<128>> PrevGlobalTransforms;
+
+		for (int32 EntryIndex = 0; EntryIndex < Entries.Num(); ++EntryIndex)
+		{
+			const FPoseHistoryEntry& Entry = Entries[EntryIndex];
+
+			const int32 PrevGlobalTransformsNum = PrevGlobalTransforms.Num();
+			const int32 Max = FMath::Max(PrevGlobalTransformsNum, Entry.Num());
+
+			PrevGlobalTransforms.SetNum(Max, EAllowShrinking::No);
+
+			for (int32 i = 0; i < Entry.Num(); ++i)
+			{
+				const FTransform RootTransform = Trajectory.GetSampleAtTime(Entry.AccumulatedSeconds).GetTransform();
+				const FTransform GlobalTransforms = Entry.GetComponentSpaceTransform(i) * RootTransform;
+
+				if (i < PrevGlobalTransformsNum)
+				{
+					DrawDebugLine(World, PrevGlobalTransforms[i].GetTranslation(), GlobalTransforms.GetTranslation(), Color, false, -1.f, SDPG_Foreground);
+				}
+
+				PrevGlobalTransforms[i] = GlobalTransforms;
+			}
+		}
+	}
+}
+#endif // ENABLE_DRAW_DEBUG && ENABLE_ANIM_DEBUG
+
+FArchive& operator<<(FArchive& Ar, FArchivedPoseHistory& Entry)
+{
+	Ar << Entry.BoneToTransformMap;
+	Ar << Entry.Entries;
+	Ar << Entry.Trajectory;
+	return Ar;
+}
+
+//////////////////////////////////////////////////////////////////////////
+// FPoseHistory
+FPoseHistory::FPoseHistory(const FPoseHistory& Other)
+	: FPoseHistory()
+{
+	*this = Other;
+}
+
+FPoseHistory::FPoseHistory(FPoseHistory&& Other)
+	: FPoseHistory()
+{
+	*this = MoveTemp(Other);
+}
+
+FPoseHistory& FPoseHistory::operator=(const FPoseHistory& Other)
+{
+	CheckThreadSafetyWrite(ReadDataThreadSafeCounter);
+	CheckThreadSafetyWrite(WriteDataThreadSafeCounter);
+
+	ReadData = Other.ReadData;
+	WriteData = Other.WriteData;
+	return *this;
+}
+
+FPoseHistory& FPoseHistory::operator=(FPoseHistory&& Other)
+{
+	CheckThreadSafetyWrite(ReadDataThreadSafeCounter);
+	CheckThreadSafetyWrite(WriteDataThreadSafeCounter);
+
+	ReadData = MoveTemp(Other.ReadData);
+	WriteData = MoveTemp(Other.WriteData);
+	return *this;
+}
+
+void FPoseHistory::Initialize_AnyThread(int32 InNumPoses, float InSamplingInterval)
+{
+	CheckThreadSafetyWrite(WriteDataThreadSafeCounter);
+	check(InNumPoses >= 2 && InSamplingInterval > UE_KINDA_SMALL_NUMBER);
+
+	MaxNumPoses = InNumPoses;
+	SamplingInterval = InSamplingInterval;
+}
+
+void FPoseHistory::CacheBones_AnyThread(const TArray<FBoneIndexType>& RequiredBones)
+{
+	CheckThreadSafetyWrite(WriteDataThreadSafeCounter);
+
+	check(MaxNumPoses >= 2 && SamplingInterval > UE_KINDA_SMALL_NUMBER);
+
+	WriteData.BoneToTransformMap.Reset();
+	if (!RequiredBones.IsEmpty())
+	{
+		// making sure we always collect the root bone transform (by construction BoneToTransformMap[0] = 0)
+		WriteData.BoneToTransformMap.Add(RootBoneIndexType) = WriteData.BoneToTransformMap.Num();
+
+		for (int32 i = 0; i < RequiredBones.Num(); ++i)
+		{
+			// adding only unique RequiredBones to avoid oversizing Entries::ComponentSpaceTransforms
+			if (!WriteData.BoneToTransformMap.Find(RequiredBones[i]))
+			{
+				WriteData.BoneToTransformMap.Add(RequiredBones[i]) = WriteData.BoneToTransformMap.Num();
+			}
+		}
+	}
+
+	WriteData.Entries.Reset();
+	WriteData.Entries.Reserve(MaxNumPoses);
+}
+
+bool FPoseHistory::GetTransformAtTime(float Time, FTransform& OutBoneTransform, const USkeleton* BoneIndexSkeleton, FBoneIndexType BoneIndexType, FBoneIndexType ReferenceBoneIndexType, bool bExtrapolate) const
+{
+	CheckThreadSafetyRead(ReadDataThreadSafeCounter);
+
+	static_assert(RootBoneIndexType == 0 && ComponentSpaceIndexType == FBoneIndexType(-1) && WorldSpaceIndexType == FBoneIndexType(-2)); // some assumptions
+	check(BoneIndexType != ComponentSpaceIndexType && BoneIndexType != WorldSpaceIndexType);
+	
+	bool bSuccess = false;
+	
+	const bool bApplyComponentToWorld = ReferenceBoneIndexType == WorldSpaceIndexType;
+	FTransform ComponentToWorld = FTransform::Identity;
+	if (bApplyComponentToWorld)
+	{
+		ComponentToWorld = Trajectory.GetSampleAtTime(Time, bExtrapolate).GetTransform();
+		ReferenceBoneIndexType = ComponentSpaceIndexType;
+	}
+
+	const int32 NumEntries = ReadData.Entries.Num();
+	if (NumEntries > 0)
+	{
+		int32 NextIdx = 0;
+		int32 PrevIdx = 0;
+
+		if (NumEntries > 1)
+		{
+			const int32 LowerBoundIdx = LowerBound(ReadData.Entries.begin(), ReadData.Entries.end(), Time, [](const FPoseHistoryEntry& Entry, float Value) { return Value > Entry.AccumulatedSeconds; });
+			NextIdx = FMath::Clamp(LowerBoundIdx, 1, NumEntries - 1);
+			PrevIdx = NextIdx - 1;
+		}
+	
+		const FPoseHistoryEntry& PrevEntry = ReadData.Entries[PrevIdx];
+		const FPoseHistoryEntry& NextEntry = ReadData.Entries[NextIdx];
+
+		bSuccess = LerpEntries(Time, bExtrapolate, PrevEntry, NextEntry, BoneIndexSkeleton, ReadData.LastUpdateSkeleton.Get(), ReadData.BoneToTransformMap, BoneIndexType, ReferenceBoneIndexType, OutBoneTransform);
+		if (bApplyComponentToWorld)
+		{
+			OutBoneTransform *= ComponentToWorld;
+		}
+	}
+	else
+	{
+		OutBoneTransform = ComponentToWorld;
+	}
+	
+	return bSuccess;
+}
+
+const FPoseSearchQueryTrajectory& FPoseHistory::GetTrajectory() const
+{
+	CheckThreadSafetyRead(ReadDataThreadSafeCounter);
+	return Trajectory;
+}
+
+float FPoseHistory::GetTrajectorySpeedMultiplier() const
+{
+	CheckThreadSafetyRead(ReadDataThreadSafeCounter);
+	return TrajectorySpeedMultiplier;
+}
+
+bool FPoseHistory::IsEmpty() const
+{
+	CheckThreadSafetyRead(ReadDataThreadSafeCounter);
+	return ReadData.Entries.IsEmpty();
+}
+
+const FBoneToTransformMap& FPoseHistory::GetBoneToTransformMap() const
+{
+	CheckThreadSafetyRead(ReadDataThreadSafeCounter);
+	return ReadData.BoneToTransformMap;
+}
+
+int32 FPoseHistory::GetNumEntries() const
+{
+	CheckThreadSafetyRead(ReadDataThreadSafeCounter);
+	return ReadData.Entries.Num();
+}
+
+const FPoseHistoryEntry& FPoseHistory::GetEntry(int32 EntryIndex) const
+{
+	CheckThreadSafetyRead(ReadDataThreadSafeCounter);
+	return ReadData.Entries[EntryIndex];
+}
+
+void FPoseHistory::PreUpdate(const UAnimInstance* AnimInstance, float DeltaTime, bool bGenerateTrajectory, const FPoseSearchTrajectoryData& TrajectoryData, const FPoseSearchTrajectoryData::FSampling& TrajectoryDataSampling, bool bNeedsReset)
+{
+	// we're writing both ReadData, WriteData. checking for thread safety
+	CheckThreadSafetyWrite(ReadDataThreadSafeCounter);
+	CheckThreadSafetyWrite(WriteDataThreadSafeCounter);
+
+	check(IsInGameThread());
+
+	if (bNeedsReset)
+	{
+		WriteData.Entries.Reset();
+	}
+
+	if (bGenerateTrajectory)
+	{
+		// @todo: Synchronize the FPoseSearchQueryTrajectorySample::AccumulatedSeconds of the generated trajectory with the FPoseHistoryEntry::AccumulatedSeconds of the captured poses
+		FPoseSearchTrajectoryData::FDerived TrajectoryDataDerived;
+		TrajectoryData.UpdateData(DeltaTime, AnimInstance, TrajectoryDataDerived, TrajectoryDataState);
+		UPoseSearchTrajectoryLibrary::InitTrajectorySamples(Trajectory, TrajectoryData, TrajectoryDataDerived, TrajectoryDataSampling);
+		UPoseSearchTrajectoryLibrary::UpdateHistory_TransformHistory(Trajectory, TrajectoryData, TrajectoryDataDerived, TrajectoryDataSampling, DeltaTime);
+		UPoseSearchTrajectoryLibrary::UpdatePrediction_SimulateCharacterMovement(Trajectory, TrajectoryData, TrajectoryDataDerived, TrajectoryDataSampling);
+
+		// @todo: support TrajectorySpeedMultiplier
+		//TrajectorySpeedMultiplier = 1.f;
+	}
+	
+	ReadData = WriteData;
+}
+
+void FPoseHistory::SetTrajectory(const FPoseSearchQueryTrajectory& InTrajectory, float InTrajectorySpeedMultiplier)
+{
+	if (!InTrajectory.Samples.IsEmpty())
+	{
+		CheckThreadSafetyWrite(ReadDataThreadSafeCounter);
+
+		// @todo: THIS IS NOT THREAD SAFE! in the contex of multi character motion matching (CheckThreadSafetyWrite will assert in case of improper usage)
+		Trajectory = InTrajectory;
+
+		TrajectorySpeedMultiplier = InTrajectorySpeedMultiplier;
+
+		if (!FMath::IsNearlyEqual(TrajectorySpeedMultiplier, 1.f))
+		{
+			const float TrajectorySpeedMultiplierInv = FMath::IsNearlyZero(TrajectorySpeedMultiplier) ? 1.f : 1.f / TrajectorySpeedMultiplier;
+			for (FPoseSearchQueryTrajectorySample& Sample : Trajectory.Samples)
+			{
+				Sample.AccumulatedSeconds *= TrajectorySpeedMultiplierInv;
+			}
+		}
+	}
+}
+
+void FPoseHistory::EvaluateComponentSpace_AnyThread(float DeltaTime, FCSPose<FCompactPose>& ComponentSpacePose, bool bStoreScales)
+{
+	CheckThreadSafetyWrite(WriteDataThreadSafeCounter);
+
+	check(MaxNumPoses >= 2 && SamplingInterval > UE_KINDA_SMALL_NUMBER);
+
+	const USkeleton* Skeleton = ComponentSpacePose.GetPose().GetBoneContainer().GetSkeletonAsset();
+	if (WriteData.LastUpdateSkeleton != Skeleton)
+	{
+		// @todo: support a different USkeleton per FPoseHistoryEntry if required
+		WriteData.Entries.Reset();
+		WriteData.LastUpdateSkeleton = Skeleton;
+	}
+
+	// Age our elapsed times
+	for (FPoseHistoryEntry& Entry : WriteData.Entries)
+	{
+		Entry.AccumulatedSeconds -= DeltaTime;
+	}
+
+	if (WriteData.Entries.Num() != MaxNumPoses)
+	{
+		// Consume every pose until the queue is full
+		WriteData.Entries.Emplace();
+	}
+	else if (WriteData.Entries[WriteData.Entries.Num() - 2].AccumulatedSeconds <= -SamplingInterval)
+	{
+		FPoseHistoryEntry EntryTemp = MoveTemp(WriteData.Entries.First());
+		WriteData.Entries.PopFront();
+		WriteData.Entries.Emplace(MoveTemp(EntryTemp));
+	}
+
+	// Regardless of the retention policy, we always update the most recent Entry
+	WriteData.Entries.Last().Update(0.f, ComponentSpacePose, WriteData.BoneToTransformMap, bStoreScales);
+}
+
+#if ENABLE_DRAW_DEBUG && ENABLE_ANIM_DEBUG
+void FPoseHistory::DebugDraw(FAnimInstanceProxy& AnimInstanceProxy, FColor Color) const
+{
+	CheckThreadSafetyRead(ReadDataThreadSafeCounter);
+
+	if (CVarAnimPoseHistoryDebugDrawTrajectory.GetValueOnAnyThread())
+	{
+		Trajectory.DebugDrawTrajectory(AnimInstanceProxy);
+	}
+
+	if (Color.A > 0 && CVarAnimPoseHistoryDebugDrawPose.GetValueOnAnyThread())
+	{
+		const bool bValidTrajectory = !Trajectory.Samples.IsEmpty();
+		TArray<FTransform, TInlineAllocator<128>> PrevGlobalTransforms;
+
+		for (int32 EntryIndex = 0; EntryIndex < ReadData.Entries.Num(); ++EntryIndex)
+		{
+			const FPoseHistoryEntry& Entry = ReadData.Entries[EntryIndex];
+
+			const int32 PrevGlobalTransformsNum = PrevGlobalTransforms.Num();
+			const int32 Max = FMath::Max(PrevGlobalTransformsNum, Entry.Num());
+
+			PrevGlobalTransforms.SetNum(Max, EAllowShrinking::No);
+
+			for (int32 i = 0; i < Entry.Num(); ++i)
+			{
+				const FTransform RootTransform = bValidTrajectory ? Trajectory.GetSampleAtTime(Entry.AccumulatedSeconds).GetTransform() : AnimInstanceProxy.GetComponentTransform();
+				const FTransform GlobalTransforms = Entry.GetComponentSpaceTransform(i) * RootTransform;
+
+				if (i < PrevGlobalTransformsNum)
+				{
+					AnimInstanceProxy.AnimDrawDebugLine(PrevGlobalTransforms[i].GetTranslation(), GlobalTransforms.GetTranslation(), Color, false, 0.f, ESceneDepthPriorityGroup::SDPG_Foreground);
+				}
+
+				PrevGlobalTransforms[i] = GlobalTransforms;
+			}
+		}
+	}
+}
+#endif // ENABLE_DRAW_DEBUG && ENABLE_ANIM_DEBUG
+
+//////////////////////////////////////////////////////////////////////////
+// FMemStackPoseHistory
+void FMemStackPoseHistory::Init(const IPoseHistory* InPoseHistory)
+{
+	check(InPoseHistory);
+	PoseHistory = InPoseHistory;
+}
+
+bool FMemStackPoseHistory::GetTransformAtTime(float Time, FTransform& OutBoneTransform, const USkeleton* BoneIndexSkeleton, FBoneIndexType BoneIndexType, FBoneIndexType ReferenceBoneIndexType, bool bExtrapolate) const
+{
+	check(PoseHistory);
+	if (Time > 0.f)
+	{
+		const int32 Num = FutureEntries.Num();
+		if (Num > 0)
+		{
+			const bool bApplyComponentToWorld = ReferenceBoneIndexType == WorldSpaceIndexType;
+			FTransform ComponentToWorld = FTransform::Identity;
+			if (bApplyComponentToWorld)
+			{
+				ComponentToWorld = GetTrajectory().GetSampleAtTime(Time, bExtrapolate).GetTransform();
+				ReferenceBoneIndexType = ComponentSpaceIndexType;
+			}
+
+			const int32 LowerBoundIdx = Algo::LowerBound(FutureEntries, Time, [](const FPoseHistoryEntry& Entry, float Value) { return Value > Entry.AccumulatedSeconds; });
+			const int32 NextIdx = FMath::Min(LowerBoundIdx, Num - 1);
+			const FPoseHistoryEntry& NextEntry = FutureEntries[NextIdx];
+			const FPoseHistoryEntry& PrevEntry = NextIdx > 0 ? FutureEntries[NextIdx - 1] : PoseHistory->GetNumEntries() > 0 ? PoseHistory->GetEntry(PoseHistory->GetNumEntries() - 1) : NextEntry;
+						
+			const bool bSuccess = LerpEntries(Time, bExtrapolate, PrevEntry, NextEntry, BoneIndexSkeleton, nullptr, GetBoneToTransformMap(), BoneIndexType, ReferenceBoneIndexType, OutBoneTransform);
+			if (bApplyComponentToWorld)
+			{
+				OutBoneTransform *= ComponentToWorld;
+			}
+			return bSuccess;
+		}
+	}
+	
+	return PoseHistory->GetTransformAtTime(Time, OutBoneTransform, BoneIndexSkeleton, BoneIndexType, ReferenceBoneIndexType, bExtrapolate);
+}
+
+void FMemStackPoseHistory::AddFutureRootBone(float Time, const FTransform& FutureRootBoneTransform, bool bStoreScales)
+{
+	// we don't allow to add "past" or "present" poses to FutureEntries
+	check(Time > 0.f);
+
+	const int32 LowerBoundIdx = Algo::LowerBound(FutureEntries, Time, [](const FPoseHistoryEntry& Entry, float Value) { return Value > Entry.AccumulatedSeconds; });
+	FPoseHistoryEntry& FutureEntry = FutureEntries.InsertDefaulted_GetRef(LowerBoundIdx);
+	FutureEntry.SetNum(1, bStoreScales);
+	FutureEntry.SetComponentSpaceTransform(RootBoneIndexType, FutureRootBoneTransform);
+	FutureEntry.AccumulatedSeconds = Time;
+}
+
+void FMemStackPoseHistory::AddFuturePose(float Time, FCSPose<FCompactPose>& ComponentSpacePose)
+{
+	// we don't allow to add "past" or "present" poses to FutureEntries
+	check(Time > 0.f);
+	check(PoseHistory);	
+	const int32 LowerBoundIdx = Algo::LowerBound(FutureEntries, Time, [](const FPoseHistoryEntry& Entry, float Value) { return Value > Entry.AccumulatedSeconds; });
+	FutureEntries.InsertDefaulted_GetRef(LowerBoundIdx).Update(Time, ComponentSpacePose, GetBoneToTransformMap(), true);
+}
+
+int32 FMemStackPoseHistory::GetNumEntries() const
+{
+	check(PoseHistory);
+	return PoseHistory->GetNumEntries() + FutureEntries.Num();
+}
+
+const FPoseHistoryEntry& FMemStackPoseHistory::GetEntry(int32 EntryIndex) const
+{
+	check(PoseHistory);
+
+	const int32 PoseHistoryNumEntries = PoseHistory->GetNumEntries();
+	if (EntryIndex < PoseHistoryNumEntries)
+	{
+		return PoseHistory->GetEntry(EntryIndex);
+	}
+	return FutureEntries[EntryIndex - PoseHistoryNumEntries];
+}
+
+
+#if ENABLE_DRAW_DEBUG && ENABLE_ANIM_DEBUG
+void FMemStackPoseHistory::DebugDraw(FAnimInstanceProxy& AnimInstanceProxy, FColor Color) const
+{
+	check(PoseHistory);
+
+	if (Color.A > 0 && !FutureEntries.IsEmpty() && CVarAnimPoseHistoryDebugDrawPose.GetValueOnAnyThread())
+	{
+		const FPoseSearchQueryTrajectory& Trajectory = GetTrajectory();
+		const bool bValidTrajectory = !Trajectory.Samples.IsEmpty();
+		TArray<FTransform, TInlineAllocator<128>> PrevGlobalTransforms;
+
+		int32 EntriesNum = FutureEntries.Num();
+		if (PoseHistory->GetNumEntries() > 0)
+		{
+			// connecting the future entries with the past entries
+			++EntriesNum;
+		}
+
+		for (int32 EntryIndex = 0; EntryIndex < EntriesNum; ++EntryIndex)
+		{
+			const FPoseHistoryEntry& Entry = (EntryIndex == FutureEntries.Num()) ? PoseHistory->GetEntry(PoseHistory->GetNumEntries() - 1) : FutureEntries[EntryIndex];
+
+			const int32 PrevGlobalTransformsNum = PrevGlobalTransforms.Num();
+			const int32 Max = FMath::Max(PrevGlobalTransformsNum, Entry.Num());
+
+			PrevGlobalTransforms.SetNum(Max, EAllowShrinking::No);
+
+			for (int32 i = 0; i < Entry.Num(); ++i)
+			{
+				const FTransform RootTransform = bValidTrajectory ? Trajectory.GetSampleAtTime(Entry.AccumulatedSeconds).GetTransform() : AnimInstanceProxy.GetComponentTransform();
+				const FTransform GlobalTransforms = Entry.GetComponentSpaceTransform(i) * RootTransform;
+
+				if (i < PrevGlobalTransformsNum)
+				{
+					AnimInstanceProxy.AnimDrawDebugLine(PrevGlobalTransforms[i].GetTranslation(), GlobalTransforms.GetTranslation(), Color, false, 0.f, ESceneDepthPriorityGroup::SDPG_Foreground);
+				}
+
+				PrevGlobalTransforms[i] = GlobalTransforms;
+			}
+		}
+
+		// no need to DebugDraw PoseHistory since it'll be drawn anyways by the history collectors
+		//PoseHistory->DebugDraw(AnimInstanceProxy, Color);
+	}
+}
+#endif // ENABLE_DRAW_DEBUG && ENABLE_ANIM_DEBUG
+
+//////////////////////////////////////////////////////////////////////////
+// FPoseIndicesHistory
+void FPoseIndicesHistory::Update(const FSearchResult& SearchResult, float DeltaTime, float MaxTime)
+{
+	if (MaxTime > 0.f)
+	{
+		for (auto It = IndexToTime.CreateIterator(); It; ++It)
+		{
+			It.Value() += DeltaTime;
+			if (It.Value() > MaxTime)
+			{
+				It.RemoveCurrent();
+			}
+		}
+
+		if (SearchResult.IsValid())
+		{
+			FHistoricalPoseIndex HistoricalPoseIndex;
+			HistoricalPoseIndex.PoseIndex = SearchResult.PoseIdx;
+			HistoricalPoseIndex.DatabaseKey = FObjectKey(SearchResult.Database.Get());
+			IndexToTime.Add(HistoricalPoseIndex, 0.f);
+		}
+	}
+	else
+	{
+		IndexToTime.Reset();
+	}
+}
+
+} // namespace UE::PoseSearch
diff --git a/Engine/Plugins/Animation/PoseSearch/Source/Runtime/Private/PoseSearchHistoryCollectorAnimNodeLibrary.cpp b/Engine/Plugins/Animation/PoseSearch/Source/Runtime/Private/PoseSearchHistoryCollectorAnimNodeLibrary.cpp
new file mode 100644
--- /dev/null
+++ b/Engine/Plugins/Animation/PoseSearch/Source/Runtime/Private/PoseSearchHistoryCollectorAnimNodeLibrary.cpp
@@ -0,0 +1,25 @@
+// Copyright Epic Games, Inc. All Rights Reserved.
+
+#include "PoseSearch/PoseSearchHistoryCollectorAnimNodeLibrary.h"
+#include "PoseSearch/AnimNode_PoseSearchHistoryCollector.h"
+
+FPoseSearchHistoryCollectorAnimNodeReference UPoseSearchHistoryCollectorAnimNodeLibrary::ConvertToPoseHistoryNode(const FAnimNodeReference& Node, EAnimNodeReferenceConversionResult& Result)
+{
+	return FAnimNodeReference::ConvertToType<FPoseSearchHistoryCollectorAnimNodeReference>(Node, Result);
+}
+
+void UPoseSearchHistoryCollectorAnimNodeLibrary::GetPoseHistoryNodeTrajectory(const FPoseSearchHistoryCollectorAnimNodeReference& PoseSearchHistoryCollectorNode, FPoseSearchQueryTrajectory& Trajectory)
+{
+	if (FAnimNode_PoseSearchHistoryCollector* PoseSearchHistoryCollectorNodePtr = PoseSearchHistoryCollectorNode.GetAnimNodePtr<FAnimNode_PoseSearchHistoryCollector>())
+	{
+		Trajectory = PoseSearchHistoryCollectorNodePtr->GetPoseHistory().GetTrajectory();
+	}
+}
+
+void UPoseSearchHistoryCollectorAnimNodeLibrary::SetPoseHistoryNodeTrajectory(const FPoseSearchHistoryCollectorAnimNodeReference& PoseSearchHistoryCollectorNode, const FPoseSearchQueryTrajectory& Trajectory)
+{
+	if (FAnimNode_PoseSearchHistoryCollector* PoseSearchHistoryCollectorNodePtr = PoseSearchHistoryCollectorNode.GetAnimNodePtr<FAnimNode_PoseSearchHistoryCollector>())
+	{
+		PoseSearchHistoryCollectorNodePtr->GetPoseHistory().SetTrajectory(Trajectory);
+	}
+}
diff --git a/Engine/Plugins/Animation/PoseSearch/Source/Runtime/Private/PoseSearchMultiSequence.cpp b/Engine/Plugins/Animation/PoseSearch/Source/Runtime/Private/PoseSearchMultiSequence.cpp
new file mode 100644
--- /dev/null
+++ b/Engine/Plugins/Animation/PoseSearch/Source/Runtime/Private/PoseSearchMultiSequence.cpp
@@ -0,0 +1,121 @@
+// Copyright Epic Games, Inc. All Rights Reserved.
+
+#include "PoseSearch/PoseSearchMultiSequence.h"
+#include "Animation/AnimSequenceBase.h"
+
+bool UPoseSearchMultiSequence::IsLooping() const
+{
+	float CommonPlayLength = -1.f;
+	for (const FPoseSearchMultiSequenceItem& Item : Items)
+	{
+		if (Item.Sequence)
+		{
+			if (!Item.Sequence->bLoop)
+			{
+				return false;
+			}
+
+			if (CommonPlayLength < 0.f)
+			{
+				CommonPlayLength = Item.Sequence->GetPlayLength();
+			}
+			else if (!FMath::IsNearlyEqual(CommonPlayLength, Item.Sequence->GetPlayLength()))
+			{
+				return false;
+			}
+		}
+	}
+	return true;
+}
+
+const FString UPoseSearchMultiSequence::GetName() const
+{
+	FString Name;
+	Name.Reserve(256);
+	bool bAddComma = false;
+	for (const FPoseSearchMultiSequenceItem& Item : Items)
+	{
+		if (bAddComma)
+		{
+			Name += ", ";
+		}
+		else
+		{
+			bAddComma = true;
+		}
+
+		Name += "[";
+		Name += Item.Role.ToString();
+		Name += "] ";
+		Name += GetNameSafe(Item.Sequence);
+	}
+	return Name;
+}
+
+bool UPoseSearchMultiSequence::HasRootMotion() const
+{
+	bool bHasAtLeastOneValidItem = false;
+	bool bHasRootMotion = true;
+
+	for (const FPoseSearchMultiSequenceItem& Item : Items)
+	{
+		if (Item.Sequence)
+		{
+			bHasRootMotion &= Item.Sequence->HasRootMotion();
+			bHasAtLeastOneValidItem = true;
+		}
+	}
+
+	return bHasAtLeastOneValidItem && bHasRootMotion;
+}
+
+float UPoseSearchMultiSequence::GetPlayLength() const
+{
+	float PlayLength = 0.f;
+	for (const FPoseSearchMultiSequenceItem& Item : Items)
+	{
+		if (Item.Sequence)
+		{
+			PlayLength = FMath::Max(PlayLength, Item.Sequence->GetPlayLength());
+		}
+	}
+	return PlayLength;
+}
+
+#if WITH_EDITOR
+int32 UPoseSearchMultiSequence::GetFrameAtTime(float Time) const
+{
+	for (const FPoseSearchMultiSequenceItem& Item : Items)
+	{
+		if (Item.Sequence)
+		{
+			return Item.Sequence->GetFrameAtTime(Time);
+		}
+	}
+	return 0;
+}
+#endif // WITH_EDITOR
+
+UAnimSequenceBase* UPoseSearchMultiSequence::GetSequence(const UE::PoseSearch::FRole& Role) const
+{
+	for (const FPoseSearchMultiSequenceItem& Item : Items)
+	{
+		if (Item.Role == Role)
+		{
+			return Item.Sequence;
+		}
+	}
+	return nullptr;
+}
+
+const FTransform& UPoseSearchMultiSequence::GetOrigin(const UE::PoseSearch::FRole& Role) const
+{
+	for (const FPoseSearchMultiSequenceItem& Item : Items)
+	{
+		if (Item.Role == Role)
+		{
+			return Item.Origin;
+		}
+	}
+	return FTransform::Identity;
+}
\ No newline at end of file
diff --git a/Engine/Plugins/Animation/PoseSearch/Source/Runtime/Private/PoseSearchSchema.cpp b/Engine/Plugins/Animation/PoseSearch/Source/Runtime/Private/PoseSearchSchema.cpp
new file mode 100644
--- /dev/null
+++ b/Engine/Plugins/Animation/PoseSearch/Source/Runtime/Private/PoseSearchSchema.cpp
@@ -0,0 +1,352 @@
+// Copyright Epic Games, Inc. All Rights Reserved.
+
+#include "PoseSearch/PoseSearchSchema.h"
+#include "AnimationRuntime.h"
+#include "PoseSearch/PoseSearchContext.h"
+#include "PoseSearch/PoseSearchDefines.h"
+#include "PoseSearch/PoseSearchResult.h"
+#include "PoseSearchFeatureChannel_Padding.h"
+#include "PoseSearchFeatureChannel_PermutationTime.h"
+#include "PoseSearchFeatureChannel_Pose.h"
+#include "PoseSearchFeatureChannel_Trajectory.h"
+#include "UObject/ObjectSaveContext.h"
+
+void UPoseSearchSchema::AddChannel(UPoseSearchFeatureChannel* Channel)
+{
+	Channels.Add(Channel);
+}
+
+void UPoseSearchSchema::AddTemporaryChannel(UPoseSearchFeatureChannel* TemporaryChannel)
+{
+	TemporaryChannel->Finalize(this);
+	FinalizedChannels.Add(TemporaryChannel);
+}
+
+TConstArrayView<float> UPoseSearchSchema::BuildQuery(UE::PoseSearch::FSearchContext& SearchContext) const
+{
+	QUICK_SCOPE_CYCLE_COUNTER(STAT_PoseSearch_BuildQuery);
+
+	SearchContext.AddNewFeatureVectorBuilder(this);
+
+	for (const TObjectPtr<UPoseSearchFeatureChannel>& ChannelPtr : GetChannels())
+	{
+		ChannelPtr->BuildQuery(SearchContext);
+	}
+
+	return SearchContext.EditFeatureVector();
+}
+
+void UPoseSearchSchema::AddSkeleton(USkeleton* Skeleton, UMirrorDataTable* MirrorDataTable, const UE::PoseSearch::FRole& Role)
+{
+	FPoseSearchRoledSkeleton& RoledSkeleton = Skeletons.AddDefaulted_GetRef();
+	RoledSkeleton.Skeleton = Skeleton;
+	RoledSkeleton.MirrorDataTable = MirrorDataTable;
+	RoledSkeleton.Role = Role;
+}
+
+bool UPoseSearchSchema::AreSkeletonsCompatible(const UPoseSearchSchema* Other) const
+{
+	if (Skeletons.Num() != Other->Skeletons.Num())
+	{
+		return false;
+	}
+
+	for (int32 SkeletonIndex = 0; SkeletonIndex < Skeletons.Num(); ++SkeletonIndex)
+	{
+		if (Skeletons[SkeletonIndex].Skeleton != Other->Skeletons[SkeletonIndex].Skeleton)
+		{
+			return false;
+		}
+
+		if (Skeletons[SkeletonIndex].Role != Other->Skeletons[SkeletonIndex].Role)
+		{
+			return false;
+		}
+	}
+
+	return true;
+}
+
+void UPoseSearchSchema::AddDefaultChannels()
+{
+	// defaulting UPoseSearchSchema for a meaningful locomotion setup
+	AddChannel(NewObject<UPoseSearchFeatureChannel_Trajectory>(this, NAME_None, RF_Transactional));
+	AddChannel(NewObject<UPoseSearchFeatureChannel_Pose>(this, NAME_None, RF_Transactional));
+}
+
+void UPoseSearchSchema::InitBoneContainersFromRoledSkeleton(TMap<FName, FBoneContainer>& RoledBoneContainers) const
+{
+	RoledBoneContainers.Reset();
+	RoledBoneContainers.Reserve(Skeletons.Num());
+
+	for (const FPoseSearchRoledSkeleton& RoledSkeleton : Skeletons)
+	{
+		RoledBoneContainers.Add(RoledSkeleton.Role).InitializeTo(RoledSkeleton.BoneIndicesWithParents, UE::Anim::FCurveFilterSettings(UE::Anim::ECurveFilterMode::DisallowAll), *RoledSkeleton.Skeleton);
+	}
+}
+
+bool UPoseSearchSchema::AllRoledSkeletonHaveMirrorDataTable() const
+{
+	for (const FPoseSearchRoledSkeleton& RoledSkeleton : Skeletons)
+	{
+		if (!RoledSkeleton.MirrorDataTable)
+		{
+			return false;
+		}
+	}
+	return true;
+}
+
+const FPoseSearchRoledSkeleton* UPoseSearchSchema::GetRoledSkeleton(const UE::PoseSearch::FRole& Role) const
+{
+	for (const FPoseSearchRoledSkeleton& RoledSkeleton : Skeletons)
+	{
+		if (RoledSkeleton.Role == Role)
+		{
+			return &RoledSkeleton;
+		}
+	}
+	return nullptr;
+}
+
+FPoseSearchRoledSkeleton* UPoseSearchSchema::GetRoledSkeleton(const UE::PoseSearch::FRole& Role)
+{
+	for (FPoseSearchRoledSkeleton& RoledSkeleton : Skeletons)
+	{
+		if (RoledSkeleton.Role == Role)
+		{
+			return &RoledSkeleton;
+		}
+	}
+	return nullptr;
+}
+
+#if WITH_EDITOR
+const UE::PoseSearch::FRole UPoseSearchSchema::GetDefaultRole() const
+{
+	if (!Skeletons.IsEmpty())
+	{
+		return Skeletons[0].Role;
+	}
+	return UE::PoseSearch::DefaultRole;
+}
+#endif // WITH_EDITOR
+
+USkeleton* UPoseSearchSchema::GetSkeleton(const UE::PoseSearch::FRole& Role) const
+{
+	if (const FPoseSearchRoledSkeleton* RoledSkeleton = GetRoledSkeleton(Role))
+	{
+		return RoledSkeleton->Skeleton.Get();
+	}
+	return nullptr;
+}
+
+UMirrorDataTable* UPoseSearchSchema::GetMirrorDataTable(const UE::PoseSearch::FRole& Role) const
+{
+	if (const FPoseSearchRoledSkeleton* RoledSkeleton = GetRoledSkeleton(Role))
+	{
+		return RoledSkeleton->MirrorDataTable.Get();
+	}
+	return nullptr;
+}
+
+TConstArrayView<FBoneReference> UPoseSearchSchema::GetBoneReferences(const UE::PoseSearch::FRole& Role) const
+{
+	const FPoseSearchRoledSkeleton* RoledSkeleton = GetRoledSkeleton(Role);
+	check(RoledSkeleton);
+	return RoledSkeleton->BoneReferences;
+}
+
+int8 UPoseSearchSchema::AddBoneReference(const FBoneReference& BoneReference, const UE::PoseSearch::FRole& Role)
+{
+	using namespace UE::PoseSearch;
+
+	FPoseSearchRoledSkeleton* RoledSkeleton = GetRoledSkeleton(Role);
+	if (!RoledSkeleton)
+	{
+		UE_LOG(LogPoseSearch, Error, TEXT("UPoseSearchSchema::AddBoneReference: couldn't find data for the requested Role '%s' in UPoseSearchSchema '%s'"), *Role.ToString(), *GetNameSafe(this));
+		return -1;
+	}
+
+	int32 SchemaBoneIdx = 0;
+	const USkeleton* Skeleton = RoledSkeleton->Skeleton;
+	if (!Skeleton)
+	{
+		UE_LOG(LogPoseSearch, Error, TEXT("UPoseSearchSchema::AddBoneReference: couldn't find Skeleton with Role '%s' in UPoseSearchSchema '%s'"), *Role.ToString(), *GetNameSafe(this));
+		return -1;
+	}
+
+	bool bDefaultToRootBone = true;
+	FBoneReference TempBoneReference = BoneReference;
+	if (TempBoneReference.BoneName != NAME_None)
+	{
+		TempBoneReference.Initialize(Skeleton);
+		if (!TempBoneReference.HasValidSetup())
+		{
+			UE_LOG(LogPoseSearch, Error, TEXT("UPoseSearchSchema::AddBoneReference: couldn't initialize FBoneReference '%s' with Skeleton '%s' with Role '%s' in UPoseSearchSchema '%s'"),
+				*TempBoneReference.BoneName.ToString(), *GetNameSafe(Skeleton), *Role.ToString(), *GetNameSafe(this));
+			return -1;
+		}
+	}
+	else
+	{
+		TempBoneReference.BoneName = Skeleton->GetReferenceSkeleton().GetBoneName(int32(RootBoneIndexType));
+		TempBoneReference.Initialize(Skeleton);
+		check(TempBoneReference.HasValidSetup());
+	}
+
+	SchemaBoneIdx = RoledSkeleton->BoneReferences.AddUnique(TempBoneReference);
+	check(SchemaBoneIdx >= 0 && SchemaBoneIdx < 128);
+	return int8(SchemaBoneIdx);
+}
+
+void UPoseSearchSchema::ResetFinalize()
+{
+	for (FPoseSearchRoledSkeleton& RoledSkeleton : Skeletons)
+	{
+		RoledSkeleton.BoneReferences.Reset();
+		RoledSkeleton.BoneIndicesWithParents.Reset();
+	}
+
+	FinalizedChannels.Reset();
+	SchemaCardinality = 0;
+}
+
+void UPoseSearchSchema::Finalize()
+{
+	using namespace UE::PoseSearch;
+
+	ResetFinalize();
+
+	// adding as first bone reference the root bone
+	for (int32 RoledSkeletonIndex = 0; RoledSkeletonIndex < Skeletons.Num(); ++RoledSkeletonIndex)
+	{
+		const FPoseSearchRoledSkeleton& RoledSkeleton = Skeletons[RoledSkeletonIndex];
+		for (int32 ComparisonIndex = RoledSkeletonIndex + 1; ComparisonIndex < Skeletons.Num(); ++ComparisonIndex)
+		{
+			if (Skeletons[ComparisonIndex].Role == RoledSkeleton.Role)
+			{
+				UE_LOG(LogPoseSearch, Error, TEXT("UPoseSearchSchema::Finalize: couldn't Finalize '%s' because of duplicate Role '%s' in Skeletons"), *GetNameSafe(this), *RoledSkeleton.Role.ToString());
+
+				ResetFinalize();
+				return;
+			}
+		}
+
+		const int8 SchemaBoneIdx = AddBoneReference(FBoneReference(), RoledSkeleton.Role);
+		if (SchemaBoneIdx != RootSchemaBoneIdx)
+		{
+			UE_LOG(LogPoseSearch, Error, TEXT("UPoseSearchSchema::Finalize: couldn't Finalize '%s' because couldn't initialize root bone properly"), *GetNameSafe(this));
+
+			ResetFinalize();
+			return;
+		}
+	}
+
+	for (const TObjectPtr<UPoseSearchFeatureChannel>& ChannelPtr : Channels)
+	{
+		if (ChannelPtr)
+		{
+			FinalizedChannels.Add(ChannelPtr);
+			if (!ChannelPtr->Finalize(this))
+			{
+				#if WITH_EDITOR
+				TLabelBuilder LabelBuilder;
+				FString Label = ChannelPtr->GetLabel(LabelBuilder).ToString();
+				UE_LOG(LogPoseSearch, Error, TEXT("UPoseSearchSchema::Finalize: couldn't Finalize '%s' because of Channel '%s'"), *GetNameSafe(this), *Label);
+				#else // WITH_EDITOR
+				UE_LOG(LogPoseSearch, Error, TEXT("UPoseSearchSchema::Finalize: couldn't Finalize '%s' because of Channel '%s'"), *GetNameSafe(this), *GetNameSafe(ChannelPtr));
+				#endif // WITH_EDITOR
+
+				ResetFinalize();
+				return;
+			}
+		}
+	}
+
+	// AddDependentChannels can add channels to FinalizedChannels, so we need a while loop
+	int32 ChannelIndex = 0;
+	while (ChannelIndex < FinalizedChannels.Num())
+	{
+		FinalizedChannels[ChannelIndex]->AddDependentChannels(this);
+		++ChannelIndex;
+	}
+
+	for (const TObjectPtr<UPoseSearchFeatureChannel>& ChannelPtr : FinalizedChannels)
+	{
+		check(ChannelPtr);
+		if (ChannelPtr->GetPermutationTimeType() != EPermutationTimeType::UseSampleTime)
+		{
+			// there's at least one channel that uses UsePermutationTime or UseSampleToPermutationTime: we automatically add a UPoseSearchFeatureChannel_PermutationTime if not already in the schema
+			UPoseSearchFeatureChannel_PermutationTime::FindOrAddToSchema(this);
+			break;
+		}
+	}
+	
+	// adding padding if required
+	if (bAddDataPadding)
+	{
+		// calculating how many floats of padding are required to make the data 16 bytes padded
+		const int32 PaddingSize = SchemaCardinality % (16 / sizeof(float));
+		if (PaddingSize > 0)
+		{
+			UPoseSearchFeatureChannel_Padding::AddToSchema(this, PaddingSize);
+		}
+	}
+
+	// Initialize references to obtain bone indices and fill out bone index array
+	for (FPoseSearchRoledSkeleton& RoledSkeleton : Skeletons)
+	{
+		for (FBoneReference& BoneRef : RoledSkeleton.BoneReferences)
+		{
+			check(BoneRef.HasValidSetup());
+			RoledSkeleton.BoneIndicesWithParents.Add(BoneRef.BoneIndex);
+		}
+
+		// Build separate index array with parent indices guaranteed to be present. Sort for EnsureParentsPresent.
+		check(!RoledSkeleton.BoneIndicesWithParents.IsEmpty());
+		RoledSkeleton.BoneIndicesWithParents.Sort();
+		FAnimationRuntime::EnsureParentsPresent(RoledSkeleton.BoneIndicesWithParents, RoledSkeleton.Skeleton->GetReferenceSkeleton());
+	}
+}
+
+void UPoseSearchSchema::PreSave(FObjectPreSaveContext ObjectSaveContext)
+{
+	Finalize();
+	Super::PreSave(ObjectSaveContext);
+}
+
+void UPoseSearchSchema::PostLoad()
+{
+	Super::PostLoad();
+
+	if (Skeleton_DEPRECATED)
+	{
+		Skeletons.AddDefaulted_GetRef().Skeleton = Skeleton_DEPRECATED;
+		Skeleton_DEPRECATED = nullptr;
+	}
+
+	if (MirrorDataTable_DEPRECATED)
+	{
+		if (Skeletons.IsEmpty())
+		{
+			Skeletons.AddDefaulted_GetRef().MirrorDataTable = MirrorDataTable_DEPRECATED;
+		}
+		else
+		{
+			Skeletons[0].MirrorDataTable = MirrorDataTable_DEPRECATED;
+		}
+		MirrorDataTable_DEPRECATED = nullptr;
+	}
+
+	Finalize();
+}
+
+#if WITH_EDITOR
+void UPoseSearchSchema::PostEditChangeProperty(FPropertyChangedEvent& PropertyChangedEvent)
+{
+	Finalize();
+	Super::PostEditChangeProperty(PropertyChangedEvent);
+}
+
+#endif
\ No newline at end of file
diff --git a/Engine/Plugins/Animation/PoseSearch/Source/Runtime/Private/PoseSearchTrajectoryTypes.cpp b/Engine/Plugins/Animation/PoseSearch/Source/Runtime/Private/PoseSearchTrajectoryTypes.cpp
new file mode 100644
--- /dev/null
+++ b/Engine/Plugins/Animation/PoseSearch/Source/Runtime/Private/PoseSearchTrajectoryTypes.cpp
@@ -0,0 +1,120 @@
+// Copyright Epic Games, Inc. All Rights Reserved.
+
+#include "PoseSearch/PoseSearchTrajectoryTypes.h"
+#include "Animation/AnimInstanceProxy.h"
+#include "Animation/AnimTypes.h"
+#include "Components/SceneComponent.h"
+#include "DrawDebugHelpers.h"
+#include "PoseSearch/PoseSearchDefines.h"
+
+FPoseSearchQueryTrajectorySample FPoseSearchQueryTrajectorySample::Lerp(const FPoseSearchQueryTrajectorySample& Other, float Alpha) const
+{
+	check(Facing.IsNormalized());
+	check(Other.Facing.IsNormalized());
+
+	FPoseSearchQueryTrajectorySample Result;
+	
+	Result.Facing = FQuat::FastLerp(Facing, Other.Facing, Alpha).GetNormalized();
+	Result.Position = FMath::Lerp(Position, Other.Position, Alpha);
+	Result.AccumulatedSeconds = FMath::Lerp(AccumulatedSeconds, Other.AccumulatedSeconds, Alpha);
+
+	return Result;
+}
+
+void FPoseSearchQueryTrajectorySample::SetTransform(const FTransform& Transform)
+{
+	Position = Transform.GetTranslation();
+	Facing = Transform.GetRotation();
+}
+
+FPoseSearchQueryTrajectorySample FPoseSearchQueryTrajectory::GetSampleAtTime(float Time, bool bExtrapolate) const
+{
+	const int32 Num = Samples.Num();
+	if (Num > 1)
+	{
+		const int32 LowerBoundIdx = Algo::LowerBound(Samples, Time, [](const FPoseSearchQueryTrajectorySample& TrajectorySample, float Value)
+			{
+				return Value > TrajectorySample.AccumulatedSeconds;
+			});
+
+		const int32 NextIdx = FMath::Clamp(LowerBoundIdx, 1, Samples.Num() - 1);
+		const int32 PrevIdx = NextIdx - 1;
+
+		const float Denominator = Samples[NextIdx].AccumulatedSeconds - Samples[PrevIdx].AccumulatedSeconds;
+		if (!FMath::IsNearlyZero(Denominator))
+		{
+			const float Numerator = Time - Samples[PrevIdx].AccumulatedSeconds;
+			const float LerpValue = bExtrapolate ? Numerator / Denominator : FMath::Clamp(Numerator / Denominator, 0.f, 1.f);
+			return Samples[PrevIdx].Lerp(Samples[NextIdx], LerpValue);
+		}
+
+		return Samples[PrevIdx];
+	}
+
+	if (Num > 0)
+	{
+		return Samples[0];
+	}
+
+	return FPoseSearchQueryTrajectorySample();
+}
+
+#if ENABLE_ANIM_DEBUG
+void FPoseSearchQueryTrajectory::DebugDrawTrajectory(const UWorld* World, float HeightOffset) const
+{
+	const FVector OffsetVector = FVector::UpVector * HeightOffset;
+
+	const int32 LastIndex = Samples.Num() - 1;
+	if (LastIndex >= 0)
+	{
+		for (int32 Index = 0; ; ++Index)
+		{
+			DrawDebugSphere(World, Samples[Index].Position + OffsetVector, 2.f /*Radius*/, 4 /*Segments*/, FColor::Black, false /*bPersistentLines*/, -1.f /*LifeTime*/, SDPG_Foreground);
+			DrawDebugCoordinateSystem(World, Samples[Index].Position + OffsetVector, FRotator(Samples[Index].Facing), 12.f /*Scale*/, false /*bPersistentLines*/, -1.f /*LifeTime*/, SDPG_Foreground);
+
+			if (Index == LastIndex)
+			{
+				break;
+			}
+			
+			DrawDebugLine(World, Samples[Index].Position + OffsetVector, Samples[Index + 1].Position + OffsetVector, FColor::Black, false /*bPersistentLines*/, -1.f /*LifeTime*/, SDPG_Foreground);
+		}
+	}
+}
+
+void FPoseSearchQueryTrajectory::DebugDrawTrajectory(FAnimInstanceProxy& AnimInstanceProxy, float HeightOffset) const
+{
+	const FVector OffsetVector = FVector::UpVector * HeightOffset;
+
+	const int32 LastIndex = Samples.Num() - 1;
+	if (LastIndex >= 0)
+	{
+		for (int32 Index = 0; ; ++Index)
+		{
+			AnimInstanceProxy.AnimDrawDebugSphere(Samples[Index].Position + OffsetVector, 2.f /*Radius*/, 4 /*Segments*/, FColor::Black, false /*bPersistentLines*/, -1.f /*LifeTime*/, 0.f /*Thickness*/, SDPG_Foreground);
+			AnimInstanceProxy.AnimDrawDebugCoordinateSystem(Samples[Index].Position + OffsetVector, FRotator(Samples[Index].Facing), 12.f /*Scale*/, false /*bPersistentLines*/, -1.f /*LifeTime*/, 0.f /*Thickness*/, SDPG_Foreground);
+
+			if (Index == LastIndex)
+			{
+				break;
+			}
+			
+			AnimInstanceProxy.AnimDrawDebugLine(Samples[Index].Position + OffsetVector, Samples[Index + 1].Position + OffsetVector, FColor::Black, false /*bPersistentLines*/, -1.f /*LifeTime*/, 0.f /*Thickness*/, SDPG_Foreground);
+		}
+	}
+}
+#endif // ENABLE_ANIM_DEBUG
+
+FArchive& operator<<(FArchive& Ar, FPoseSearchQueryTrajectorySample& TrajectorySample)
+{
+	Ar << TrajectorySample.Facing;
+	Ar << TrajectorySample.Position;
+	Ar << TrajectorySample.AccumulatedSeconds;
+	return Ar;
+}
+
+FArchive& operator<<(FArchive& Ar, FPoseSearchQueryTrajectory& Trajectory)
+{
+	Ar << Trajectory.Samples;
+	return Ar;
+}
\ No newline at end of file
diff --git a/Engine/Plugins/Animation/PoseSearch/Source/Runtime/Public/PoseSearch/PoseSearchAssetIndexer.h b/Engine/Plugins/Animation/PoseSearch/Source/Runtime/Public/PoseSearch/PoseSearchAssetIndexer.h
new file mode 100644
--- /dev/null
+++ b/Engine/Plugins/Animation/PoseSearch/Source/Runtime/Public/PoseSearch/PoseSearchAssetIndexer.h
@@ -0,0 +1,165 @@
+// Copyright Epic Games, Inc. All Rights Reserved.
+
+#pragma once
+
+#if WITH_EDITOR
+
+#include "BonePose.h"
+#include "PoseSearch/PoseSearchMirrorDataCache.h"
+#include "PoseSearch/PoseSearchDefines.h"
+#include "PoseSearch/PoseSearchFeatureChannel.h"
+#include "PoseSearch/PoseSearchSchema.h"
+
+class UAnimNotifyState_PoseSearchBase;
+class UPoseSearchDatabase;
+
+namespace UE::PoseSearch
+{
+
+struct FAnimationAssetSampler;
+struct FSearchIndexAsset;
+struct FPoseMetadata;
+
+struct FAssetSamplingContext
+{
+	float BaseCostBias = 0.f;
+	float LoopingCostBias = 0.f;
+
+	FAssetSamplingContext(const UPoseSearchDatabase& Database);
+};
+
+struct FAnimationAssetSamplers
+{
+	void Reset();
+	int32 Num() const;
+
+	float GetPlayLength() const;
+	bool IsLoopable() const;
+	void ExtractPoseSearchNotifyStates(float Time, TFunction<bool(UAnimNotifyState_PoseSearchBase*)> ProcessPoseSearchBase) const;
+	bool ProcessAllAnimNotifyEvents(TFunction<bool(TConstArrayView<FAnimNotifyEvent>)> ProcessAnimNotifyEvents) const;
+	const FString GetAssetName() const;
+
+	FTransform ExtractRootTransform(float Time, int32 RoleIndex) const;
+	FTransform GetTotalRootTransform(int32 RoleIndex) const;
+	void ExtractPose(float Time, FCompactPose& OutPose, int32 RoleIndex) const;
+	FTransform MirrorTransform(const FTransform& InTransform, int32 RoleIndex) const;
+	void MirrorPose(FCompactPose& Pose, int32 RoleIndex) const;
+
+	TArray<const FAnimationAssetSampler*> AnimationAssetSamplers;
+	TArray<const FMirrorDataCache*> MirrorDataCaches;
+};
+
+class FAssetIndexer
+{
+public:
+	struct FStats
+	{
+		int32 NumAccumulatedSamples = 0;
+		float AccumulatedSpeed = 0.f;
+		float MaxSpeed = 0.f;
+		float AccumulatedAcceleration = 0.f;
+		float MaxAcceleration = 0.f;
+	};
+
+	FAssetIndexer(const TConstArrayView<FBoneContainer> InBoneContainers, const FSearchIndexAsset& InSearchIndexAsset, const FAssetSamplingContext& InSamplingContext,
+		const UPoseSearchSchema& InSchema, const FAnimationAssetSamplers& InAssetSamplers, const FRoleToIndex& InRoleToIndex, const FFloatInterval& InExtrapolationTimeInterval);
+	void AssignWorkingData(int32 InStartPoseIdx, TArrayView<float> InOutFeatureVectorTable, TArrayView<FPoseMetadata> InOutPoseMetadata);
+	void Process(int32 AssetIdx);
+	const FStats& GetStats() const { return Stats; }
+
+	// Returns OutSampleRotation as the rotation of the bone Schema.BoneReferences[SchemaSampleBoneIdx] at time CalculateSampleTime(SampleIdx) + SampleTimeOffset relative to the
+	// transform of the bone Schema.BoneReferences[SchemaOriginBoneIdx] at time CalculateSampleTime(SampleIdx) + OriginTimeOffset 
+	// Times will be processed by GetPermutationTimeOffsets(PermutationTimeType, ...)
+	bool GetSampleRotation(FQuat& OutSampleRotation, float SampleTimeOffset, float OriginTimeOffset, int32 SampleIdx, int8 SchemaSampleBoneIdx, int8 SchemaOriginBoneIdx, const FRole& SampleRole, const FRole& OriginRole, EPermutationTimeType PermutationTimeType = EPermutationTimeType::UseSampleTime, int32 SamplingAttributeId = INDEX_NONE);
+
+	// Returns OutSamplePosition as the position of the bone Schema.BoneReferences[SchemaSampleBoneIdx] at time CalculateSampleTime(SampleIdx) + SampleTimeOffset relative to the
+	// transform of the bone Schema.BoneReferences[SchemaOriginBoneIdx] at time CalculateSampleTime(SampleIdx) + OriginTimeOffset.
+	// Times will be processed by GetPermutationTimeOffsets(PermutationTimeType, ...)
+	bool GetSamplePosition(FVector& OutSamplePosition, float SampleTimeOffset, float OriginTimeOffset, int32 SampleIdx, int8 SchemaSampleBoneIdx, int8 SchemaOriginBoneIdx, const FRole& SampleRole, const FRole& OriginRole, EPermutationTimeType PermutationTimeType = EPermutationTimeType::UseSampleTime, int32 SamplingAttributeId = INDEX_NONE);
+
+	// Returns OutSampleVelocity as the delta velocity of the bone Schema.BoneReferences[SchemaSampleBoneIdx] velocity at time CalculateSampleTime(SampleIdx) + SampleTimeOffset minus
+	// the velocity of the bone Schema.BoneReferences[SchemaOriginBoneIdx] at time CalculateSampleTime(SampleIdx) + OriginTimeOffset.
+	// Times will be processed by GetPermutationTimeOffsets(PermutationTimeType, ...)
+	// if bUseCharacterSpaceVelocities is true, velocities will be computed in root bone space, rather than animation (world) space
+	bool GetSampleVelocity(FVector& OutSampleVelocity, float SampleTimeOffset, float OriginTimeOffset, int32 SampleIdx, int8 SchemaSampleBoneIdx, int8 SchemaOriginBoneIdx, const FRole& SampleRole, const FRole& OriginRole, bool bUseCharacterSpaceVelocities = true, EPermutationTimeType PermutationTimeType = EPermutationTimeType::UseSampleTime, int32 SamplingAttributeId = INDEX_NONE);
+
+	bool ProcessAllAnimNotifyEvents(TFunction<bool(TConstArrayView<FAnimNotifyEvent>)> ProcessAnimNotifyEvents) const;
+	const FString GetAssetName() const;
+	float GetPlayLength() const;
+
+	int32 GetBeginSampleIdx() const;
+	int32 GetEndSampleIdx() const;
+	int32 GetNumIndexedPoses() const;
+	
+	TArrayView<float> GetPoseVector(int32 SampleIdx) const;
+	const UPoseSearchSchema* GetSchema() const;
+	float CalculateSampleTime(int32 SampleIdx) const;
+	bool IsProcessFailed() const { return bProcessFailed; }
+
+#if WITH_EDITOR
+	float CalculatePermutationTimeOffset() const;
+#endif //WITH_EDITOR
+
+#if ENABLE_ANIM_DEBUG
+	void CompareCachedEntries(const FAssetIndexer& Other) const;
+#endif // ENABLE_ANIM_DEBUG
+
+private:
+	int32 GetVectorIdx(int32 SampleIdx) const;
+
+	// Returns the animation (world) space transform of the bone Schema.BoneReferences[SchemaBoneIdx] at time SampleTime
+	// bClamped will be true if SampleTime is outside the animation duration boundaries
+	FTransform GetTransform(float SampleTime, const FRole& Role, bool& bClamped, int8 SchemaBoneIdx = RootSchemaBoneIdx);
+	FTransform GetTransform(float SampleTime, int32 RoleIndex, bool& bClamped, const FBoneReference& BoneReference);
+
+	// Returns the component space transform of the bone Schema.BoneReferences[SchemaBoneIdx] at time SampleTime
+	// bClamped will be true if SampleTime is outside the animation duration boundaries
+	FTransform GetComponentSpaceTransform(float SampleTime, const FRole& Role, bool& bClamped, int8 SchemaBoneIdx = RootSchemaBoneIdx);
+
+	// Returns OutSamplePosition as the position of the bone Schema.BoneReferences[SchemaSampleBoneIdx] at time SampleTime relative to the
+	// transform of the bone Schema.BoneReferences[SchemaOriginBoneIdx] at time OriginTime.
+	// bClamped will be true if SampleTime or OriginTime are outside the animation duration boundaries
+	bool GetSamplePositionInternal(FVector& OutSamplePosition, float SampleTime, float OriginTime, bool& bClamped, int8 SchemaSampleBoneIdx, int8 SchemaOriginBoneIdx, const FRole& SampleRole, const FRole& OriginRole, int32 SamplingAttributeId);
+
+	struct FCachedEntry
+	{
+		float SampleTime = 0.f;
+		bool bClamped = false;
+
+		TArray<FTransform> RootTransform;
+		TArray<FCSPose<FCompactHeapPose>> ComponentSpacePose;
+	};
+
+	void GetSampleInfo(float SampleTime, int32 RoleIndex, FTransform& OutRootTransform, float& OutClipTime, bool& bOutClamped) const;
+	FTransform MirrorTransform(const FTransform& Transform, int32 RoleIndex) const;
+	FCachedEntry& GetEntry(float SampleTime);
+	FTransform CalculateComponentSpaceTransform(FCachedEntry& Entry, const FBoneReference& BoneReference, int32 RoleIndex);
+
+	void ComputeStats();
+
+	TArray<FBoneContainer> BoneContainers;
+	TMap<float, FCachedEntry> CachedEntries;
+	const FSearchIndexAsset& SearchIndexAsset;
+	const FAssetSamplingContext& SamplingContext;
+	const UPoseSearchSchema& Schema;
+	const FAnimationAssetSamplers AssetSamplers;
+
+	// NoTe: mapping Role to the index of the associated asset that this FAssetIndexer is indexing.
+	// NOT the index of the UPoseSearchSchema::Skeletons! Use UPoseSearchSchema::GetRoledSkeleton API to resolve that Role to FPoseSearchRoledSkeleton 
+	const FRoleToIndex RoleToIndex;
+
+	FFloatInterval ExtrapolationTimeInterval = FFloatInterval(-UE_BIG_NUMBER, UE_BIG_NUMBER);
+	
+	int32 StartPoseIdx = 0;
+	
+	TArrayView<float> FeatureVectorTable;
+	TArrayView<FPoseMetadata> PoseMetadata;
+
+	FStats Stats;
+
+	bool bProcessFailed = false;
+};
+
+} // namespace UE::PoseSearch
+
+#endif // WITH_EDITOR
diff --git a/Engine/Plugins/Animation/PoseSearch/Source/Runtime/Public/PoseSearch/PoseSearchAssetSampler.h b/Engine/Plugins/Animation/PoseSearch/Source/Runtime/Public/PoseSearch/PoseSearchAssetSampler.h
new file mode 100644
--- /dev/null
+++ b/Engine/Plugins/Animation/PoseSearch/Source/Runtime/Public/PoseSearch/PoseSearchAssetSampler.h
@@ -0,0 +1,69 @@
+// Copyright Epic Games, Inc. All Rights Reserved.
+
+#pragma once
+
+#include "BonePose.h"
+
+struct FAnimationPoseData;
+struct FAnimExtractContext;
+class UAnimationAsset;
+class UAnimNotifyState_PoseSearchBase;
+
+namespace UE::PoseSearch
+{
+
+/**
+ * Helper for sampling data from animation assets
+ */
+struct POSESEARCH_API FAnimationAssetSampler
+{
+	FAnimationAssetSampler(TObjectPtr<const UAnimationAsset> InAnimationAsset = nullptr, const FTransform& InRootTransformOrigin = FTransform::Identity, const FVector& InBlendParameters = FVector::ZeroVector, int32 InRootTransformSamplingRate = 30);
+	void Init(TObjectPtr<const UAnimationAsset> InAnimationAsset, const FTransform& InRootTransformOrigin = FTransform::Identity, const FVector& InBlendParameters = FVector::ZeroVector, int32 InRootTransformSamplingRate = 30);
+
+	bool IsInitialized() const;
+	float GetPlayLength() const;
+	float ToRealTime(float NormalizedTime) const;
+	float ToNormalizedTime(float RealTime) const;
+	bool IsLoopable() const;
+
+	// Gets the final root transformation at the end of the asset's playback time
+	FTransform GetTotalRootTransform() const;
+
+	// Extracts pose for this asset for a given context
+	void ExtractPose(const FAnimExtractContext& ExtractionCtx, FAnimationPoseData& OutAnimPoseData) const;
+
+	// Extracts pose for this asset at a given Time
+	void ExtractPose(float Time, FCompactPose& OutPose) const;
+
+	// Extracts root transform at the given time, using the extremities of the sequence to extrapolate beyond the 
+	// sequence limits when Time is less than zero or greater than the sequence length.
+	FTransform ExtractRootTransform(float Time) const;
+
+	// Extracts notify states inheriting from UAnimNotifyState_PoseSearchBase present in the sequence at Time.
+	// The function does not empty NotifyStates before adding new notifies!
+	void ExtractPoseSearchNotifyStates(float Time, TFunction<bool(UAnimNotifyState_PoseSearchBase*)> ProcessPoseSearchBase) const;
+	
+	TConstArrayView<FAnimNotifyEvent> GetAllAnimNotifyEvents() const;
+
+	const UAnimationAsset* GetAsset() const;
+
+	void Process();
+
+	static float GetPlayLength(const UAnimationAsset* AnimAsset, const FVector& BlendParameters = FVector::ZeroVector);
+
+protected:
+	TWeakObjectPtr<const UAnimationAsset> AnimationAssetPtr;
+	FTransform RootTransformOrigin = FTransform::Identity;
+
+	// members used to sample blend spaces only!
+	FVector BlendParameters = FVector::ZeroVector;
+	int32 RootTransformSamplingRate = 30;
+	float CachedPlayLength = 0.f;
+	TArray<FTransform> AccumulatedRootTransform;
+
+	const float ExtrapolationSampleTime = 1.f / 30.f;
+	const float ExtractionInterval = 1.0f / 120.0f;
+};
+
+} // namespace UE::PoseSearch
+
diff --git a/Engine/Plugins/Animation/PoseSearch/Source/Runtime/Public/PoseSearch/PoseSearchContext.h b/Engine/Plugins/Animation/PoseSearch/Source/Runtime/Public/PoseSearch/PoseSearchContext.h
new file mode 100644
--- /dev/null
+++ b/Engine/Plugins/Animation/PoseSearch/Source/Runtime/Public/PoseSearch/PoseSearchContext.h
@@ -0,0 +1,345 @@
+// Copyright Epic Games, Inc. All Rights Reserved.
+
+#pragma once
+
+#include "DrawDebugHelpers.h"
+#include "Animation/AnimInstance.h"
+#include "PoseSearch/PoseSearchDefines.h"
+#include "PoseSearch/PoseSearchFeatureChannel.h"
+#include "PoseSearch/PoseSearchIndex.h"
+#include "PoseSearch/PoseSearchResult.h"
+
+struct FPoseSearchQueryTrajectory;
+class UPoseSearchDatabase;
+class UPoseSearchFeatureChannel_Position;
+
+namespace UE::PoseSearch
+{
+
+struct FPoseIndicesHistory;
+struct IPoseHistory;
+
+enum class EDebugDrawFlags : uint32
+{
+	None = 0,
+
+	// used to differenciate channels debug drawing of the query
+	DrawQuery = 1 << 0,
+};
+ENUM_CLASS_FLAGS(EDebugDrawFlags);
+
+enum class EPoseCandidateFlags : uint32
+{
+	None = 0,
+
+	Valid_Pose = 1 << 0,
+	Valid_ContinuingPose = 1 << 1,
+	Valid_CurrentPose = 1 << 2,
+
+	AnyValidMask = Valid_Pose | Valid_ContinuingPose | Valid_CurrentPose,
+
+	DiscardedBy_PoseJumpThresholdTime = 1 << 3,
+	DiscardedBy_PoseReselectHistory = 1 << 4,
+	DiscardedBy_BlockTransition = 1 << 5,
+	DiscardedBy_PoseFilter = 1 << 6,
+	DiscardedBy_AssetIdxFilter = 1 << 7,
+	DiscardedBy_Search = 1 << 8,
+
+	AnyDiscardedMask = DiscardedBy_PoseJumpThresholdTime | DiscardedBy_PoseReselectHistory | DiscardedBy_BlockTransition | DiscardedBy_PoseFilter | DiscardedBy_AssetIdxFilter | DiscardedBy_Search,
+};
+ENUM_CLASS_FLAGS(EPoseCandidateFlags);
+
+#if ENABLE_DRAW_DEBUG
+struct POSESEARCH_API FDebugDrawParams
+{
+	FDebugDrawParams(TArrayView<FAnimInstanceProxy*> InAnimInstanceProxies, TConstArrayView<const IPoseHistory*> InPoseHistories, const FRoleToIndex& InRoleToIndex, const UPoseSearchDatabase* InDatabase, EDebugDrawFlags InFlags = EDebugDrawFlags::None);
+	FDebugDrawParams(TArrayView<const USkinnedMeshComponent*> InMeshes, TConstArrayView<const IPoseHistory*> InPoseHistories, const FRoleToIndex& InRoleToIndex, const UPoseSearchDatabase* InDatabase, EDebugDrawFlags InFlags = EDebugDrawFlags::None);
+
+	const FSearchIndex* GetSearchIndex() const;
+	const UPoseSearchSchema* GetSchema() const;
+
+	FVector ExtractPosition(TConstArrayView<float> PoseVector, const UPoseSearchFeatureChannel_Position* Position) const;
+	FVector ExtractPosition(TConstArrayView<float> PoseVector, float SampleTimeOffset, int8 SchemaBoneIdx, const FRole& Role, EPermutationTimeType PermutationTimeType = EPermutationTimeType::UseSampleTime, int32 SamplingAttributeId = INDEX_NONE) const;
+
+	FQuat ExtractRotation(TConstArrayView<float> PoseVector, float SampleTimeOffset, int8 SchemaBoneIdx, const FRole& Role, EPermutationTimeType PermutationTimeType = EPermutationTimeType::UseSampleTime, int32 SamplingAttributeId = INDEX_NONE) const;
+
+	FTransform GetRootTransform(const FRole& Role, float SampleTimeOffset = 0.f) const;
+
+	void DrawLine(const FVector& LineStart, const FVector& LineEnd, const FColor& Color, float Thickness = 0.f) const;
+	void DrawPoint(const FVector& Position, const FColor& Color, float Thickness = 6.f) const;
+	void DrawCircle(const FMatrix& TransformMatrix, float Radius, int32 Segments, const FColor& Color, float Thickness = 1.f) const;
+	void DrawCentripetalCatmullRomSpline(TConstArrayView<FVector> Points, TConstArrayView<FColor> Colors, float Alpha, int32 NumSamplesPerSegment, float Thickness = 1.f) const;
+	
+	void DrawFeatureVector(TConstArrayView<float> PoseVector);
+	void DrawFeatureVector(int32 PoseIdx);
+
+private:
+	bool CanDraw() const;
+
+	const TArrayView<FAnimInstanceProxy*> AnimInstanceProxies;
+	const TArrayView<const USkinnedMeshComponent*> Meshes;
+	const TConstArrayView<const IPoseHistory*> PoseHistories;
+
+	// NoTe: mapping Role to the index of the associated asset that this FDebugDrawParams is drawing.
+	// NOT the index of the UPoseSearchSchema::Skeletons! Use UPoseSearchSchema::GetRoledSkeleton API to resolve that Role to FPoseSearchRoledSkeleton 
+	const FRoleToIndex& RoleToIndex;
+
+	const UPoseSearchDatabase* Database = nullptr;
+	EDebugDrawFlags Flags = EDebugDrawFlags::None;
+};
+
+#endif // ENABLE_DRAW_DEBUG
+
+// float buffer of features according to a UPoseSearchSchema layout. Used to build search queries at runtime
+struct FCachedQuery
+{
+public:
+	explicit FCachedQuery(const UPoseSearchSchema* InSchema);
+	const UPoseSearchSchema* GetSchema() const { return Schema; }
+	TArrayView<float> EditValues() { return Values; }
+	TConstArrayView<float> GetValues() const { return Values; }
+
+private:
+	TStackAlignedArray<float> Values;
+	
+	// no need for a TWeakObjectPtr since it doesn't persist across multiple frames (same lifespan as FSearchContext)
+	const UPoseSearchSchema* Schema;
+};
+
+// CachedChannels uses hashed unique identifiers to determine channels that can share feature vector data during the building of the query
+struct FCachedChannel
+{
+	// no need for a TWeakObjectPtr since it doesn't persist across multiple frames (same lifespan as FSearchContext)
+	const UPoseSearchFeatureChannel* Channel = nullptr;
+
+	// index of the associated query in FSearchContext::CachedQueries
+	int32 CachedQueryIndex = INDEX_NONE;
+};
+
+struct POSESEARCH_API FSearchContext
+{
+	FSearchContext(float InDesiredPermutationTimeOffset = 0.f, const FPoseIndicesHistory* InPoseIndicesHistory = nullptr,
+		const FSearchResult& InCurrentResult = FSearchResult(), const FFloatInterval& InPoseJumpThresholdTime = FFloatInterval(0.f, 0.f), bool bInUseCachedChannelData = false);
+
+	// deleting copies and moves since members could reference other members (e.g.: CurrentResultPoseVector could point to CurrentResultPoseVectorData, so it'll require proper copies and movers implementations)
+	FSearchContext(const FSearchContext& Other) = delete;
+	FSearchContext(FSearchContext&& Other) = delete;
+	FSearchContext& operator=(const FSearchContext& Other) = delete;
+	FSearchContext& operator=(FSearchContext&& Other) = delete;
+
+	void AddRole(const FRole& Role, const UAnimInstance* AnimInstance, const IPoseHistory* PoseHistory);
+
+	// Returns the rotation of the bone Schema.BoneReferences[SchemaSampleBoneIdx] at an offset time of SampleTimeOffset relative to the
+	// transform of the bone Schema.BoneReferences[SchemaOriginBoneIdx] at an offset time of time OriginTimeOffset 
+	// Times will be processed by GetPermutationTimeOffsets(PermutationTimeType, ...)
+	FQuat GetSampleRotation(float SampleTimeOffset, float OriginTimeOffset, int8 SchemaSampleBoneIdx, int8 SchemaOriginBoneIdx, const FRole& SampleRole, const FRole& OriginRole, EPermutationTimeType PermutationTimeType = EPermutationTimeType::UseSampleTime, const FQuat* SampleBoneRotationWorldOverride = nullptr);
+	
+	// Returns the position of the bone Schema.BoneReferences[SchemaSampleBoneIdx] at an offset time of SampleTimeOffset relative to the
+	// transform of the bone Schema.BoneReferences[SchemaOriginBoneIdx] at an offset time of time OriginTimeOffset 
+	// Times will be processed by GetPermutationTimeOffsets(PermutationTimeType, ...)
+	FVector GetSamplePosition(float SampleTimeOffset, float OriginTimeOffset, int8 SchemaSampleBoneIdx, int8 SchemaOriginBoneIdx, const FRole& SampleRole, const FRole& OriginRole, EPermutationTimeType PermutationTimeType = EPermutationTimeType::UseSampleTime, const FVector* SampleBonePositionWorldOverride = nullptr);
+	
+	// Returns the delta velocity of the velocity of the bone Schema.BoneReferences[SchemaSampleBoneIdx] at an offset time of SampleTimeOffset minus
+	// the velocity of the bone Schema.BoneReferences[SchemaOriginBoneIdx] at an offset time of time OriginTimeOffset 
+	// Times will be processed by GetPermutationTimeOffsets(PermutationTimeType, ...)
+	// if bUseCharacterSpaceVelocities is true, velocities will be computed in root bone space, rather than world space
+	FVector GetSampleVelocity(float SampleTimeOffset, float OriginTimeOffset, int8 SchemaSampleBoneIdx, int8 SchemaOriginBoneIdx, const FRole& SampleRole, const FRole& OriginRole, bool bUseCharacterSpaceVelocities = true, EPermutationTimeType PermutationTimeType = EPermutationTimeType::UseSampleTime, const FVector* SampleBoneVelocityWorldOverride = nullptr);
+
+	void ResetCurrentBestCost();
+	void UpdateCurrentBestCost(const FPoseSearchCost& PoseSearchCost);
+	float GetCurrentBestTotalCost() const { return CurrentBestTotalCost; }
+
+	TConstArrayView<float> GetOrBuildQuery(const UPoseSearchSchema* Schema);
+	TConstArrayView<float> GetCachedQuery(const UPoseSearchSchema* Schema) const;
+
+	bool IsCurrentResultFromDatabase(const UPoseSearchDatabase* Database) const;
+	bool CanUseCurrentResult() const;
+
+	TConstArrayView<float> GetCurrentResultPoseVector() const { return CurrentResultPoseVector; }
+
+	void UpdateCurrentResultPoseVector();
+	const FSearchResult& GetCurrentResult() const { return CurrentResult; }
+	const FFloatInterval& GetPoseJumpThresholdTime() const { return PoseJumpThresholdTime; }
+	const FPoseIndicesHistory* GetPoseIndicesHistory() const { return PoseIndicesHistory; }
+	
+	bool ArePoseHistoriesValid() const;
+	const TArray<const IPoseHistory*, TInlineAllocator<PreallocatedRolesNum, TMemStackAllocator<>>>& GetPoseHistories() const { return PoseHistories; }
+	const IPoseHistory* GetPoseHistory(const FRole& Role) const { return PoseHistories[RoleToIndex[Role]]; }
+
+	float GetDesiredPermutationTimeOffset() const { return DesiredPermutationTimeOffset; }
+	const UAnimInstance* GetAnimInstance(const FRole& Role) const { return AnimInstances[RoleToIndex[Role]]; }
+	const TArray<const UAnimInstance*, TInlineAllocator<PreallocatedRolesNum, TMemStackAllocator<>>>& GetAnimInstances() const { return AnimInstances; }
+
+	const FRoleToIndex& GetRoleToIndex() const { return RoleToIndex; }
+
+	void SetAssetsToConsider(TConstArrayView<const UObject*> InAssetsToConsider) { AssetsToConsider = InAssetsToConsider; }
+	TConstArrayView<const UObject*> GetAssetsToConsider() const { return AssetsToConsider; }
+	
+	// returns the world space transform of the bone SchemaBoneIdx at time SampleTime
+	FTransform GetWorldBoneTransformAtTime(float SampleTime, const FRole& SampleRole, int8 SchemaBoneIdx);
+
+#if WITH_EDITOR
+	void SetAsyncBuildIndexInProgress() { bAsyncBuildIndexInProgress = true; }
+	void ResetAsyncBuildIndexInProgress() { bAsyncBuildIndexInProgress = false; }
+	bool IsAsyncBuildIndexInProgress() const { return bAsyncBuildIndexInProgress; }
+#endif // WITH_EDITOR
+
+	bool AnyCachedQuery() const { return !CachedQueries.IsEmpty(); }
+	void AddNewFeatureVectorBuilder(const UPoseSearchSchema* Schema) { CachedQueries.Emplace(Schema); }
+	TArrayView<float> EditFeatureVector();
+	
+	const UPoseSearchFeatureChannel* GetCachedChannelData(uint32 ChannelUniqueIdentifier, const UPoseSearchFeatureChannel* Channel, TConstArrayView<float>& CachedChannelData);
+	bool IsUseCachedChannelData() const { return bUseCachedChannelData; }
+	void SetUseCachedChannelData(bool bInUseCachedChannelData) { bUseCachedChannelData = bInUseCachedChannelData; }
+
+private:
+	FVector GetSamplePositionInternal(float SampleTime, float OriginTime, int8 SchemaSampleBoneIdx, int8 SchemaOriginBoneIdx, const FRole& SampleRole, const FRole& OriginRole, const FVector* SampleBonePositionWorldOverride = nullptr);
+	FQuat GetSampleRotationInternal(float SampleTime, float OriginTime, int8 SchemaSampleBoneIdx, int8 SchemaOriginBoneIdx, const FRole& SampleRole, const FRole& OriginRole, const FQuat* SampleBoneRotationWorldOverride = nullptr);
+	FTransform GetWorldRootBoneTransformAtTime(float SampleTime, const FRole& SampleRole) const;
+	
+	TArray<const UAnimInstance*, TInlineAllocator<PreallocatedRolesNum, TMemStackAllocator<>>> AnimInstances;
+	TArray<const IPoseHistory*, TInlineAllocator<PreallocatedRolesNum, TMemStackAllocator<>>> PoseHistories;
+	FRoleToIndex RoleToIndex;
+	
+	// if AssetsToConsider is not empty, we'll search only for poses from UObject(s) that are in the AssetsToConsider
+	TConstArrayView<const UObject*> AssetsToConsider;
+
+	const float DesiredPermutationTimeOffset = 0.f;
+	const FPoseIndicesHistory* PoseIndicesHistory = nullptr;
+	const FSearchResult& CurrentResult;
+	const FFloatInterval& PoseJumpThresholdTime;
+	bool bUseCachedChannelData = false;
+
+	TConstArrayView<float> CurrentResultPoseVector;
+
+	TStackAlignedArray<float> CurrentResultPoseVectorData;
+
+	// transforms cached in world space
+	TMap<uint32, FTransform, TInlineSetAllocator<64, TMemStackSetAllocator<>>> CachedTransforms;
+
+	TArray<FCachedQuery, TInlineAllocator<PreallocatedCachedQueriesNum, TMemStackAllocator<>>> CachedQueries;
+
+	// mapping channel unique identifier (hash) to FCachedChannel
+	TMap<uint32, FCachedChannel, TInlineSetAllocator<PreallocatedCachedChannelDataNum, TMemStackSetAllocator<>>> CachedChannels;
+
+	float CurrentBestTotalCost = MAX_flt;
+	
+#if WITH_EDITOR
+	bool bAsyncBuildIndexInProgress = false;
+#endif // WITH_EDITOR
+
+#if UE_POSE_SEARCH_TRACE_ENABLED
+
+	struct FPoseCandidateIdCost
+	{
+		int32 PoseIdx = 0;
+		FPoseSearchCost Cost;
+		bool operator<(const FPoseCandidateIdCost& Other) const { return Other.Cost < Cost; } // Reverse compare because BestCandidates is a max heap
+	};
+	
+public:
+
+	struct FPoseCandidate : public FPoseCandidateIdCost
+	{
+		EPoseCandidateFlags PoseCandidateFlags = EPoseCandidateFlags::None;
+	};
+
+	void Track(const UPoseSearchDatabase* Database, int32 PoseIdx = INDEX_NONE, EPoseCandidateFlags PoseCandidateFlags = EPoseCandidateFlags::None, const FPoseSearchCost& Cost = FPoseSearchCost())
+	{
+		check(Database);
+
+		FBestPoseCandidates& BestPoseCandidates = BestPoseCandidatesMap.FindOrAdd(Database);
+		if (PoseIdx != INDEX_NONE)
+		{
+			BestPoseCandidates.Add(PoseIdx, PoseCandidateFlags, Cost);
+		}
+	}
+
+	struct FBestPoseCandidates
+	{
+		FBestPoseCandidates()
+		{
+			// preallocating memory to avoid multiple reallocations / rehashing
+			PoseCandidateHeap.Reserve(MaxNumberOfCollectedPoseCandidatesPerDatabase);
+			PoseIdxToFlags.Empty(MaxNumberOfCollectedPoseCandidatesPerDatabase);
+		}
+
+		void Add(int32 PoseIdx, EPoseCandidateFlags PoseCandidateFlags, const FPoseSearchCost& Cost)
+		{
+			check(PoseIdx >= 0);
+			if (EPoseCandidateFlags* PoseIdxPoseCandidateFlags = PoseIdxToFlags.Find(PoseIdx))
+			{
+				*PoseIdxPoseCandidateFlags |= PoseCandidateFlags;
+			}
+			else if (PoseCandidateHeap.Num() < MaxNumberOfCollectedPoseCandidatesPerDatabase || Cost < PoseCandidateHeap.HeapTop().Cost)
+			{
+				bool bPoppedContinuingPoseCandidate = false;
+				FPoseCandidate ContinuingPoseCandidate;
+				while (PoseCandidateHeap.Num() >= MaxNumberOfCollectedPoseCandidatesPerDatabase)
+				{
+					FPoseCandidate PoppedPoseCandidate;
+					Pop(PoppedPoseCandidate);
+
+					if (EnumHasAnyFlags(PoppedPoseCandidate.PoseCandidateFlags, EPoseCandidateFlags::Valid_ContinuingPose))
+					{
+						// we can only have one continuing pose candidate
+						check(!bPoppedContinuingPoseCandidate);
+						ContinuingPoseCandidate = PoppedPoseCandidate;
+						bPoppedContinuingPoseCandidate = true;
+					}					
+				}
+
+				if (bPoppedContinuingPoseCandidate)
+				{
+					// if we popped the continuing pose candidate, we make some space for it and push it back
+					FPoseCandidate PoppedPoseCandidate;
+					Pop(PoppedPoseCandidate);
+					PoseCandidateHeap.HeapPush(ContinuingPoseCandidate);
+					PoseIdxToFlags.Add(ContinuingPoseCandidate.PoseIdx, ContinuingPoseCandidate.PoseCandidateFlags);
+				}
+
+				FPoseCandidate PoseCandidate;
+				PoseCandidate.PoseIdx = PoseIdx;
+				PoseCandidate.Cost = Cost;
+				PoseCandidateHeap.HeapPush(PoseCandidate);
+				PoseIdxToFlags.Add(PoseIdx, PoseCandidateFlags);
+			}
+		}
+
+		int32 Num() const
+		{
+			return PoseCandidateHeap.Num();
+		}
+
+		FPoseCandidate GetUnsortedCandidate(int32 Index) const
+		{
+			FPoseCandidate PoseCandidate;
+			const FPoseCandidateIdCost& PoseCandidateIdCost = PoseCandidateHeap[Index];
+			PoseCandidate.PoseIdx = PoseCandidateIdCost.PoseIdx;
+			PoseCandidate.Cost = PoseCandidateIdCost.Cost;
+			PoseCandidate.PoseCandidateFlags = PoseIdxToFlags[PoseCandidateIdCost.PoseIdx];
+			return PoseCandidate;
+		}
+
+	private:
+		void Pop(FPoseCandidate& OutItem)
+		{
+			PoseCandidateHeap.HeapPop(OutItem, EAllowShrinking::No);
+			OutItem.PoseCandidateFlags = PoseIdxToFlags.FindAndRemoveChecked(OutItem.PoseIdx);
+		}
+
+		TArray<FPoseCandidateIdCost> PoseCandidateHeap;
+		TMap<int32, EPoseCandidateFlags> PoseIdxToFlags;
+	};
+	
+	const TMap<const UPoseSearchDatabase*, FBestPoseCandidates>& GetBestPoseCandidatesMap() const
+	{
+		return BestPoseCandidatesMap;
+	}
+
+private:
+	TMap<const UPoseSearchDatabase*, FBestPoseCandidates> BestPoseCandidatesMap;
+#endif // UE_POSE_SEARCH_TRACE_ENABLED
+};
+
+} // namespace UE::PoseSearch
diff --git a/Engine/Plugins/Animation/PoseSearch/Source/Runtime/Public/PoseSearch/PoseSearchDefines.h b/Engine/Plugins/Animation/PoseSearch/Source/Runtime/Public/PoseSearch/PoseSearchDefines.h
new file mode 100644
--- /dev/null
+++ b/Engine/Plugins/Animation/PoseSearch/Source/Runtime/Public/PoseSearch/PoseSearchDefines.h
@@ -0,0 +1,108 @@
+// Copyright Epic Games, Inc. All Rights Reserved.
+
+#pragma once
+
+#include "Async/ParallelFor.h"
+#include "Logging/LogMacros.h"
+#include "Animation/AnimTrace.h"
+
+// Enable this if object tracing is enabled, mimics animation tracing
+#define UE_POSE_SEARCH_TRACE_ENABLED ANIM_TRACE_ENABLED
+
+#ifndef UE_POSE_SEARCH_FORCE_SINGLE_THREAD
+#define UE_POSE_SEARCH_FORCE_SINGLE_THREAD 0
+#endif
+
+#if UE_POSE_SEARCH_FORCE_SINGLE_THREAD
+constexpr EParallelForFlags ParallelForFlags = EParallelForFlags::ForceSingleThread;
+#else
+constexpr EParallelForFlags ParallelForFlags = EParallelForFlags::None;
+#endif // UE_POSE_SEARCH_FORCE_SINGLE_THREAD
+
+DECLARE_LOG_CATEGORY_EXTERN(LogPoseSearch, Log, All);
+
+namespace UE::PoseSearch
+{
+static constexpr int8 RootSchemaBoneIdx = 0;
+static constexpr FBoneIndexType RootBoneIndexType = 0;
+static constexpr FBoneIndexType ComponentSpaceIndexType = FBoneIndexType(-1);
+static constexpr FBoneIndexType WorldSpaceIndexType = FBoneIndexType(-2);
+
+static constexpr int32 PreallocatedRolesNum = 4;
+static constexpr int32 PreallocatedCachedQueriesNum = 8;
+static constexpr int32 PreallocatedCachedChannelDataNum = 64;
+static constexpr float FiniteDelta = 1 / 60.0f; // Time delta used for computing pose derivatives
+static constexpr int32 MaxNumberOfCollectedPoseCandidatesPerDatabase = 200;
+
+template<uint32 Alignment = DEFAULT_ALIGNMENT> class TMemStackSetAllocator : public TSetAllocator<TSparseArrayAllocator<TMemStackAllocator<Alignment>, TMemStackAllocator<Alignment>>, TMemStackAllocator<Alignment>> {};
+
+constexpr int32 TMax(int32 A, int32 B) { return (A > B ? A : B); }
+template<typename ElementType> constexpr int32 TAlignOf() { return TMax(alignof(ElementType), 16); }
+template<typename ElementType> using TAlignedArray = TArray<ElementType, TAlignedHeapAllocator<TAlignOf<ElementType>()>>;
+template<typename ElementType> using TStackAlignedArray = TArray<ElementType, TMemStackAllocator<TAlignOf<ElementType>()>>;
+
+#if ENABLE_ANIM_DEBUG
+
+	template<bool bIsRead = 0>
+	struct FDebugThreadSafetyChecker
+	{
+		// if any read is in progress Counter > 0
+		// if any write is in progress Counter < 0
+		explicit FDebugThreadSafetyChecker(FThreadSafeCounter& InCounter)
+			: Counter(InCounter)
+		{
+			if (bIsRead)
+			{
+				const int32 PreviousCounter = Counter.Add(1);
+				if (PreviousCounter < 0)
+				{
+					UE_LOG(LogPoseSearch, Error, TEXT("Non thread safe call! Read detected a Write in progress"));
+				}
+			}
+			else
+			{
+				const int32 PreviousCounter = Counter.Set(-1);
+				if (PreviousCounter > 0)
+				{
+					UE_LOG(LogPoseSearch, Error, TEXT("Non thread safe call! Write detected %d Read(s) in progress"), PreviousCounter);
+				}
+			}
+		}
+
+		~FDebugThreadSafetyChecker()
+		{
+			if (bIsRead)
+			{
+				const int32 PreviousCounter = Counter.Add(-1);
+				if (PreviousCounter <= 0)
+				{
+					UE_LOG(LogPoseSearch, Error, TEXT("Non thread safe call! Read detected a Write in progress"));
+				}
+			}
+			else
+			{
+				const int32 PreviousCounter = Counter.Set(0);
+				if (PreviousCounter >= 0)
+				{
+					UE_LOG(LogPoseSearch, Error, TEXT("Non thread safe call! Write detected %d Read(s) in progress"), PreviousCounter);
+				}
+			}
+		}
+
+		FThreadSafeCounter& Counter;
+	};
+
+#define CheckThreadSafetyRead(ThreadSafeCounter) FDebugThreadSafetyChecker<1> DebugThreadSafetyChecker##ThreadSafeCounter(ThreadSafeCounter)
+#define CheckThreadSafetyWrite(ThreadSafeCounter) FDebugThreadSafetyChecker DebugThreadSafetyChecker##ThreadSafeCounter(ThreadSafeCounter)
+
+#else // ENABLE_ANIM_DEBUG
+
+	#define CheckThreadSafetyRead(ThreadSafeCounter)
+	#define CheckThreadSafetyWrite(ThreadSafeCounter)
+
+#endif // ENABLE_ANIM_DEBUG
+
+}
+
+
+
diff --git a/Engine/Plugins/Animation/PoseSearch/Source/Runtime/Public/PoseSearch/PoseSearchHistory.h b/Engine/Plugins/Animation/PoseSearch/Source/Runtime/Public/PoseSearch/PoseSearchHistory.h
new file mode 100644
--- /dev/null
+++ b/Engine/Plugins/Animation/PoseSearch/Source/Runtime/Public/PoseSearch/PoseSearchHistory.h
@@ -0,0 +1,230 @@
+// Copyright Epic Games, Inc. All Rights Reserved.
+
+#pragma once
+
+#include "Animation/AnimNodeMessages.h"
+#include "BonePose.h"
+#include "Containers/RingBuffer.h"
+#include "DrawDebugHelpers.h"
+#include "PoseSearch/PoseSearchDefines.h"
+#include "PoseSearch/PoseSearchTrajectoryLibrary.h"
+#include "UObject/ObjectKey.h"
+
+struct FAnimInstanceProxy;
+class USkeleton;
+class UWorld;
+
+namespace UE::PoseSearch
+{
+
+#if ENABLE_DRAW_DEBUG && ENABLE_ANIM_DEBUG
+extern POSESEARCH_API TAutoConsoleVariable<bool> CVarAnimPoseHistoryDebugDrawPose;
+extern POSESEARCH_API TAutoConsoleVariable<bool> CVarAnimPoseHistoryDebugDrawTrajectory;
+#endif
+
+struct FSearchResult;
+typedef uint16 FComponentSpaceTransformIndex;
+typedef TPair<FBoneIndexType, FComponentSpaceTransformIndex> FBoneToTransformPair;
+typedef TMap<FBoneIndexType, FComponentSpaceTransformIndex> FBoneToTransformMap;
+
+struct FPoseHistoryEntry
+{
+	// collected bones transforms in component space
+	TArray<FQuat4f> ComponentSpaceRotations;
+	TArray<FVector> ComponentSpacePositions;
+	TArray<FVector3f> ComponentSpaceScales;
+	float AccumulatedSeconds = 0.f;
+
+	void Update(float Time, FCSPose<FCompactPose>& ComponentSpacePose, const FBoneToTransformMap& BoneToTransformMap, bool bStoreScales);
+
+	void SetNum(int32 Num, bool bStoreScales);
+	int32 Num() const;
+
+	void SetComponentSpaceTransform(int32 Index, const FTransform& Transform);
+	FTransform GetComponentSpaceTransform(int32 Index) const;
+};
+FArchive& operator<<(FArchive& Ar, FPoseHistoryEntry& Entry);
+
+struct POSESEARCH_API IPoseHistory
+{
+public:
+	virtual ~IPoseHistory() {}
+	
+	// returns the BoneIndexType transform relative to ReferenceBoneIndexType: 
+	// if ReferenceBoneIndexType is 0 (RootBoneIndexType), OutBoneTransform is in root bone space
+	// if ReferenceBoneIndexType is FBoneIndexType(-1) (ComponentSpaceIndexType), OutBoneTransform is in component space
+	// if ReferenceBoneIndexType is FBoneIndexType(-2) (WorldSpaceIndexType), OutBoneTransform is in world space
+	virtual bool GetTransformAtTime(float Time, FTransform& OutBoneTransform, const USkeleton* BoneIndexSkeleton = nullptr, FBoneIndexType BoneIndexType = RootBoneIndexType, FBoneIndexType ReferenceBoneIndexType = ComponentSpaceIndexType, bool bExtrapolate = true) const = 0;
+	virtual const FPoseSearchQueryTrajectory& GetTrajectory() const = 0;
+	
+	// @todo: deprecate this API. TrajectorySpeedMultiplier should be a global query scaling value passed as input parameter of FSearchContext during config BuildQuery
+	virtual float GetTrajectorySpeedMultiplier() const = 0;
+	virtual bool IsEmpty() const = 0;
+
+	virtual const FBoneToTransformMap& GetBoneToTransformMap() const = 0;
+	virtual int32 GetNumEntries() const = 0;
+	virtual const FPoseHistoryEntry& GetEntry(int32 EntryIndex) const = 0;
+
+#if ENABLE_DRAW_DEBUG && ENABLE_ANIM_DEBUG
+	virtual void DebugDraw(const UWorld* World, FColor Color) const = 0;
+	virtual void DebugDraw(FAnimInstanceProxy& AnimInstanceProxy, FColor Color) const = 0;
+	virtual void DebugDraw(FAnimInstanceProxy& AnimInstanceProxy, FColor Color, float Time, float PointSize = 6.f, bool bExtrapolate = true) const;
+#endif
+};
+
+struct POSESEARCH_API FArchivedPoseHistory : public IPoseHistory
+{
+	void InitFrom(const IPoseHistory* PoseHistory);
+
+	// IPoseHistory interface
+	virtual bool GetTransformAtTime(float Time, FTransform& OutBoneTransform, const USkeleton* BoneIndexSkeleton = nullptr, FBoneIndexType BoneIndexType = RootBoneIndexType, FBoneIndexType ReferenceBoneIndexType = ComponentSpaceIndexType, bool bExtrapolate = true) const override;
+	virtual const FPoseSearchQueryTrajectory& GetTrajectory() const override { return Trajectory; }
+	virtual float GetTrajectorySpeedMultiplier() const override { return 1.f; }
+	virtual bool IsEmpty() const override { return Entries.IsEmpty(); }
+	virtual const FBoneToTransformMap& GetBoneToTransformMap() const override { return BoneToTransformMap; }
+	
+	virtual int32 GetNumEntries() const override { return Entries.Num(); }
+	virtual const FPoseHistoryEntry& GetEntry(int32 EntryIndex) const override { return Entries[EntryIndex]; }
+
+#if ENABLE_DRAW_DEBUG && ENABLE_ANIM_DEBUG
+	virtual void DebugDraw(const UWorld* World, FColor Color) const override;
+	virtual void DebugDraw(FAnimInstanceProxy& AnimInstanceProxy, FColor Color) const override { unimplemented(); }
+#endif
+	// End of IPoseHistory interface
+
+	FBoneToTransformMap BoneToTransformMap;
+	TArray<FPoseHistoryEntry> Entries;
+	FPoseSearchQueryTrajectory Trajectory;
+};
+FArchive& operator<<(FArchive& Ar, FArchivedPoseHistory& Entry);
+
+struct FPoseHistory : public IPoseHistory
+{
+	FPoseHistory() = default;
+	FPoseHistory(const FPoseHistory& Other);
+	FPoseHistory(FPoseHistory&& Other);
+	FPoseHistory& operator=(const FPoseHistory& Other);
+	FPoseHistory& operator=(FPoseHistory&& Other);
+
+	void PreUpdate(const UAnimInstance* AnimInstance, float DeltaTime, bool bGenerateTrajectory, const FPoseSearchTrajectoryData& TrajectoryData, const FPoseSearchTrajectoryData::FSampling& TrajectoryDataSampling, bool bNeedsReset);
+
+	void Initialize_AnyThread(int32 InNumPoses, float InSamplingInterval);
+	void CacheBones_AnyThread(const TArray<FBoneIndexType>& RequiredBones);
+	void EvaluateComponentSpace_AnyThread(float DeltaTime, FCSPose<FCompactPose>& ComponentSpacePose, bool bStoreScales);
+
+	// IPoseHistory interface
+	virtual bool GetTransformAtTime(float Time, FTransform& OutBoneTransform, const USkeleton* BoneIndexSkeleton = nullptr, FBoneIndexType BoneIndexType = RootBoneIndexType, FBoneIndexType ReferenceBoneIndexType = ComponentSpaceIndexType, bool bExtrapolate = true) const override;
+	virtual const FPoseSearchQueryTrajectory& GetTrajectory() const override;
+	virtual float GetTrajectorySpeedMultiplier() const override;
+	virtual bool IsEmpty() const override;
+	virtual const FBoneToTransformMap& GetBoneToTransformMap() const override;
+	
+	void SetTrajectory(const FPoseSearchQueryTrajectory& InTrajectory, float InTrajectorySpeedMultiplier = 1.f);
+	virtual int32 GetNumEntries() const override;
+	virtual const FPoseHistoryEntry& GetEntry(int32 EntryIndex) const override;
+
+#if ENABLE_DRAW_DEBUG && ENABLE_ANIM_DEBUG
+	virtual void DebugDraw(const UWorld* World, FColor Color) const override { unimplemented(); }
+	virtual void DebugDraw(FAnimInstanceProxy& AnimInstanceProxy, FColor Color) const override;
+#endif
+	// End of IPoseHistory interface
+	
+private:
+
+	// caching MaxNumPoses, since FData::Entries.Max() is a padded number
+	int32 MaxNumPoses = 0;
+
+	float SamplingInterval = 0.f;
+
+	FPoseSearchQueryTrajectory Trajectory;
+	FPoseSearchTrajectoryData::FState TrajectoryDataState;
+	// @todo: deprecate this member and expose it via blue print logic or as global query scaling multiplier
+	float TrajectorySpeedMultiplier = 1.f;
+
+	struct FData
+	{
+		// skeleton from the last Update, to keep tracking skeleton changes, and support compatible skeletons
+		TWeakObjectPtr<const USkeleton> LastUpdateSkeleton;
+
+		// map of FBoneIndexType(s) to collect. If Empty all the bones get collected
+		FBoneToTransformMap BoneToTransformMap;
+
+		// ring buffer of collected bones
+		TRingBuffer<FPoseHistoryEntry> Entries;
+	};
+
+	FData ReadData;
+	FData WriteData;
+
+#if ENABLE_ANIM_DEBUG
+	
+	// used to analyze thread safety
+	mutable FThreadSafeCounter ReadDataThreadSafeCounter = 0;
+	mutable FThreadSafeCounter WriteDataThreadSafeCounter = 0;
+
+#endif // ENABLE_ANIM_DEBUG
+};
+
+struct FMemStackPoseHistory : public IPoseHistory
+{
+	void Init(const IPoseHistory* InPoseHistory);
+
+	// IPoseHistory interface
+	virtual bool GetTransformAtTime(float Time, FTransform& OutBoneTransform, const USkeleton* BoneIndexSkeleton = nullptr, FBoneIndexType BoneIndexType = RootBoneIndexType, FBoneIndexType ReferenceBoneIndexType = ComponentSpaceIndexType, bool bExtrapolate = true) const override;
+	virtual const FPoseSearchQueryTrajectory& GetTrajectory() const override { check(PoseHistory); return PoseHistory->GetTrajectory(); }
+	virtual float GetTrajectorySpeedMultiplier() const override { check(PoseHistory); return PoseHistory->GetTrajectorySpeedMultiplier(); }
+	virtual bool IsEmpty() const override { check(PoseHistory); return PoseHistory->IsEmpty() && FutureEntries.IsEmpty(); }
+	virtual const FBoneToTransformMap& GetBoneToTransformMap() const override { check(PoseHistory); return PoseHistory->GetBoneToTransformMap(); }
+	
+	virtual int32 GetNumEntries() const override;
+	virtual const FPoseHistoryEntry& GetEntry(int32 EntryIndex) const override;
+	
+#if ENABLE_DRAW_DEBUG && ENABLE_ANIM_DEBUG
+	virtual void DebugDraw(const UWorld* World, FColor Color) const override { unimplemented(); }
+	virtual void DebugDraw(FAnimInstanceProxy& AnimInstanceProxy, FColor Color) const override;
+#endif
+	// End of IPoseHistory interface
+
+	void AddFutureRootBone(float Time, const FTransform& FutureRootBoneTransform, bool bStoreScales);
+	void AddFuturePose(float Time, FCSPose<FCompactPose>& ComponentSpacePose);
+
+	const IPoseHistory* GetThisOrPoseHistory() const { return FutureEntries.IsEmpty() ? PoseHistory : this; }
+
+private:
+	const IPoseHistory* PoseHistory = nullptr;
+	TArray<FPoseHistoryEntry, TInlineAllocator<4, TMemStackAllocator<>>> FutureEntries;
+};
+
+class IPoseHistoryProvider : public UE::Anim::IGraphMessage
+{
+	DECLARE_ANIMGRAPH_MESSAGE(IPoseHistoryProvider);
+public:
+	virtual const IPoseHistory& GetPoseHistory() const = 0;
+};
+
+struct FHistoricalPoseIndex
+{
+	bool operator==(const FHistoricalPoseIndex& Index) const
+	{
+		return PoseIndex == Index.PoseIndex && DatabaseKey == Index.DatabaseKey;
+	}
+
+	friend FORCEINLINE uint32 GetTypeHash(const FHistoricalPoseIndex& Index)
+	{
+		return HashCombineFast(::GetTypeHash(Index.PoseIndex), GetTypeHash(Index.DatabaseKey));
+	}
+
+	int32 PoseIndex = INDEX_NONE;
+	FObjectKey DatabaseKey;
+};
+
+struct FPoseIndicesHistory
+{
+	void Update(const FSearchResult& SearchResult, float DeltaTime, float MaxTime);
+	void Reset() { IndexToTime.Reset(); }
+	TMap<FHistoricalPoseIndex, float> IndexToTime;
+};
+
+} // namespace UE::PoseSearch
+
+
diff --git a/Engine/Plugins/Animation/PoseSearch/Source/Runtime/Public/PoseSearch/PoseSearchHistoryCollectorAnimNodeLibrary.h b/Engine/Plugins/Animation/PoseSearch/Source/Runtime/Public/PoseSearch/PoseSearchHistoryCollectorAnimNodeLibrary.h
new file mode 100644
--- /dev/null
+++ b/Engine/Plugins/Animation/PoseSearch/Source/Runtime/Public/PoseSearch/PoseSearchHistoryCollectorAnimNodeLibrary.h
@@ -0,0 +1,45 @@
+// Copyright Epic Games, Inc. All Rights Reserved.
+
+#pragma once
+
+#include "Animation/AnimExecutionContext.h"
+#include "Animation/AnimNodeReference.h"
+#include "Kismet/BlueprintFunctionLibrary.h"
+#include "PoseSearchHistoryCollectorAnimNodeLibrary.generated.h"
+
+struct FAnimNode_PoseSearchHistoryCollector;
+
+USTRUCT(BlueprintType)
+struct FPoseSearchHistoryCollectorAnimNodeReference : public FAnimNodeReference
+{
+	GENERATED_BODY()
+
+	typedef FAnimNode_PoseSearchHistoryCollector FInternalNodeType;
+};
+
+// Exposes operations that can be run on a Pose History node via Anim Node Functions such as "On Become Relevant" and "On Update".
+UCLASS()
+class POSESEARCH_API UPoseSearchHistoryCollectorAnimNodeLibrary : public UBlueprintFunctionLibrary
+{
+	GENERATED_BODY()
+
+public:
+	/** Get a Pose History node context from an anim node context */
+	UFUNCTION(BlueprintCallable, Category = "Animation|PoseHistory", meta = (BlueprintThreadSafe, ExpandEnumAsExecs = "Result"))
+	static FPoseSearchHistoryCollectorAnimNodeReference ConvertToPoseHistoryNode(const FAnimNodeReference& Node, EAnimNodeReferenceConversionResult& Result);
+
+	/** Get a Pose History node context from an anim node context (pure) */
+	UFUNCTION(BlueprintPure, Category = "Animation|PoseHistory", meta = (BlueprintThreadSafe, DisplayName = "Convert to Pose History Node"))
+	static void ConvertToPoseHistoryNodePure(const FAnimNodeReference& Node, FPoseSearchHistoryCollectorAnimNodeReference& PoseSearchHistoryCollectorNode, bool& Result)
+	{
+		EAnimNodeReferenceConversionResult ConversionResult;
+		PoseSearchHistoryCollectorNode = ConvertToPoseHistoryNode(Node, ConversionResult);
+		Result = (ConversionResult == EAnimNodeReferenceConversionResult::Succeeded);
+	}
+
+	UFUNCTION(BlueprintPure, Category = "Animation|PoseHistory", meta = (BlueprintThreadSafe, DisplayName = "Get Pose History Node Trajectory"))
+	static void GetPoseHistoryNodeTrajectory(const FPoseSearchHistoryCollectorAnimNodeReference& PoseSearchHistoryCollectorNode, FPoseSearchQueryTrajectory& Trajectory);
+
+	UFUNCTION(BlueprintCallable, Category = "Animation|PoseHistory", meta = (BlueprintThreadSafe, DisplayName = "Set Pose History Node Trajectory"))
+	static void SetPoseHistoryNodeTrajectory(const FPoseSearchHistoryCollectorAnimNodeReference& PoseSearchHistoryCollectorNode, const FPoseSearchQueryTrajectory& Trajectory);
+};
\ No newline at end of file
