Move PoseSearch/BlendStack out of experimental folder

diff --git a/Engine/Plugins/Animation/PoseSearch/Source/Runtime/Public/PoseSearch/PoseSearchMultiSequence.h b/Engine/Plugins/Animation/PoseSearch/Source/Runtime/Public/PoseSearch/PoseSearchMultiSequence.h
new file mode 100644
--- /dev/null
+++ b/Engine/Plugins/Animation/PoseSearch/Source/Runtime/Public/PoseSearch/PoseSearchMultiSequence.h
@@ -0,0 +1,51 @@
+// Copyright Epic Games, Inc. All Rights Reserved.
+
+#pragma once
+
+#include "Engine/DataAsset.h"
+#include "PoseSearch/PoseSearchRole.h"
+#include "PoseSearchMultiSequence.generated.h"
+
+class UAnimSequenceBase;
+
+USTRUCT(Experimental)
+struct POSESEARCH_API FPoseSearchMultiSequenceItem
+{
+	GENERATED_BODY()
+
+	UPROPERTY(EditAnywhere, Category = "Settings")
+	TObjectPtr<UAnimSequenceBase> Sequence;
+
+	UPROPERTY(EditAnywhere, Category = "Settings")
+	FName Role;
+
+	UPROPERTY(EditAnywhere, Category = "Settings")
+	FTransform Origin = FTransform::Identity;
+};
+
+UCLASS(Experimental, BlueprintType, Category = "Animation|Pose Search")
+class UPoseSearchMultiSequence : public UDataAsset
+{
+	GENERATED_BODY()
+
+	UPROPERTY(EditAnywhere, Category = "Settings")
+	TArray<FPoseSearchMultiSequenceItem> Items;
+
+public:
+	bool IsLooping() const;
+	const FString GetName() const;
+	bool HasRootMotion() const;
+	float GetPlayLength() const;
+
+#if WITH_EDITOR
+	int32 GetFrameAtTime(float Time) const;
+#endif // WITH_EDITOR
+
+	int32 GetNumRoles() const { return Items.Num(); }
+	const UE::PoseSearch::FRole& GetRole(int32 RoleIndex) const { return Items[RoleIndex].Role; }
+	
+	UFUNCTION(BlueprintPure, Category = "Animation", meta=(BlueprintThreadSafe))
+	UAnimSequenceBase* GetSequence(const FName /*UE::PoseSearch::FRole*/& Role) const;
+
+	const FTransform& GetOrigin(const UE::PoseSearch::FRole& Role) const;
+};
\ No newline at end of file
