motion matching - support for multiple events for UPoseSearchFeatureChannel_TimeToEvent

diff --git a/Engine/Plugins/Animation/PoseSearch/Source/Runtime/Private/PoseSearchAssetIndexer.cpp b/Engine/Plugins/Animation/PoseSearch/Source/Runtime/Private/PoseSearchAssetIndexer.cpp
--- a/Engine/Plugins/Animation/PoseSearch/Source/Runtime/Private/PoseSearchAssetIndexer.cpp
+++ b/Engine/Plugins/Animation/PoseSearch/Source/Runtime/Private/PoseSearchAssetIndexer.cpp
@@ -560,42 +560,31 @@ bool FAssetIndexer::GetSampleRotation(FQuat& OutSampleRotation, float SampleTime
 
 	if (SamplingAttributeId >= 0)
 	{
-		if (ProcessAllAnimNotifyEvents([&OutSampleRotation, &SampleRole, &OriginRole, SamplingAttributeId, this, OriginTime](const TConstArrayView<FAnimNotifyEvent> AnimNotifyEvents)
+		const FPoseSearchTimedNotifies<UAnimNotifyState_PoseSearchSamplingAttribute> TimedNotifies(SamplingAttributeId, *this);
+		const FPoseSearchTimedNotifies<UAnimNotifyState_PoseSearchSamplingAttribute>::FItem TimedNotifiesItem = TimedNotifies.GetClosestFutureEvent(SampleTime);
+		if (const UAnimNotifyState_PoseSearchSamplingAttribute* SamplingAttribute = TimedNotifiesItem.NotifyState)
+		{
+			bool bUnused;
+			if (SamplingAttribute->Bone.BoneName != NAME_None)
 			{
-				for (const FAnimNotifyEvent& AnimNotifyEvent : AnimNotifyEvents)
+				FBoneReference TempBoneReference = SamplingAttribute->Bone;
+				const int32 SampleRoleIndex = RoleToIndex[SampleRole];
+				TempBoneReference.Initialize(BoneContainers[SampleRoleIndex].GetSkeletonAsset());
+				if (TempBoneReference.HasValidSetup())
 				{
-					if (const UAnimNotifyState_PoseSearchSamplingAttribute* SamplingAttribute = Cast<UAnimNotifyState_PoseSearchSamplingAttribute>(AnimNotifyEvent.NotifyStateClass))
-					{
-						if (SamplingAttribute->SamplingAttributeId == SamplingAttributeId)
-						{
-							bool bUnused;
-							if (SamplingAttribute->Bone.BoneName != NAME_None)
-							{
-								FBoneReference TempBoneReference = SamplingAttribute->Bone;
-								const int32 SampleRoleIndex = RoleToIndex[SampleRole];
-								TempBoneReference.Initialize(BoneContainers[SampleRoleIndex].GetSkeletonAsset());
-								if (TempBoneReference.HasValidSetup())
-								{
-									const float SamplingAttributeTime = AnimNotifyEvent.GetTime();
-									const FTransform RootBoneTransform = GetTransform(OriginTime, OriginRole, bUnused, RootSchemaBoneIdx);
-									const FTransform SamplingAttributeBoneTransform = GetTransform(SamplingAttributeTime, SampleRoleIndex, bUnused, TempBoneReference);
-									OutSampleRotation = RootBoneTransform.InverseTransformRotation(SamplingAttributeBoneTransform.GetRotation());
-									return true;
-								}
-
-								UE_LOG(LogPoseSearch, Error, TEXT("FAssetIndexer::GetSampleRotation: required UAnimNotifyState_PoseSearchSamplingAttribute in '%s' has an invalid Bone"), *AssetSamplers.GetAssetName());
-								return false;
-							}
-
-							const FTransform RootBoneTransform = GetTransform(OriginTime, OriginRole, bUnused, RootSchemaBoneIdx);
-							OutSampleRotation = RootBoneTransform.InverseTransformRotation(SamplingAttribute->Rotation);
-							return true;
-						}
-					}
+					const float SamplingAttributeTime = TimedNotifiesItem.Time;
+					const FTransform RootBoneTransform = GetTransform(OriginTime, OriginRole, bUnused, RootSchemaBoneIdx);
+					const FTransform SamplingAttributeBoneTransform = GetTransform(SamplingAttributeTime, SampleRoleIndex, bUnused, TempBoneReference);
+					OutSampleRotation = RootBoneTransform.InverseTransformRotation(SamplingAttributeBoneTransform.GetRotation());
+					return true;
 				}
+
+				UE_LOG(LogPoseSearch, Error, TEXT("FAssetIndexer::GetSampleRotation: required UAnimNotifyState_PoseSearchSamplingAttribute in '%s' has an invalid Bone"), *AssetSamplers.GetAssetName());
 				return false;
-			}))
-		{
+			}
+
+			const FTransform RootBoneTransform = GetTransform(OriginTime, OriginRole, bUnused, RootSchemaBoneIdx);
+			OutSampleRotation = RootBoneTransform.InverseTransformRotation(SamplingAttribute->Rotation);
 			return true;
 		}
 
@@ -631,53 +620,42 @@ bool FAssetIndexer::GetSamplePositionInternal(FVector& OutSamplePosition, float
 
 	if (SamplingAttributeId >= 0)
 	{
-		if (ProcessAllAnimNotifyEvents([&OutSamplePosition, OriginTime, &bClamped, SchemaOriginBoneIdx, &SampleRole, OriginRole, SamplingAttributeId, this](const TConstArrayView<FAnimNotifyEvent> AnimNotifyEvents)
+		const FPoseSearchTimedNotifies<UAnimNotifyState_PoseSearchSamplingAttribute> TimedNotifies(SamplingAttributeId, *this);
+		const FPoseSearchTimedNotifies<UAnimNotifyState_PoseSearchSamplingAttribute>::FItem TimedNotifiesItem = TimedNotifies.GetClosestFutureEvent(SampleTime);
+		if (const UAnimNotifyState_PoseSearchSamplingAttribute* SamplingAttribute = TimedNotifiesItem.NotifyState)
+		{
+			bool bUnused;
+			if (SamplingAttribute->Bone.BoneName != NAME_None)
 			{
-				for (const FAnimNotifyEvent& AnimNotifyEvent : AnimNotifyEvents)
+				FBoneReference TempBoneReference = SamplingAttribute->Bone;
+				const int32 SampleRoleIndex = RoleToIndex[SampleRole];
+				TempBoneReference.Initialize(BoneContainers[SampleRoleIndex].GetSkeletonAsset());
+				if (TempBoneReference.HasValidSetup())
 				{
-					if (const UAnimNotifyState_PoseSearchSamplingAttribute* SamplingAttribute = Cast<UAnimNotifyState_PoseSearchSamplingAttribute>(AnimNotifyEvent.NotifyStateClass))
+					const float SamplingAttributeTime = TimedNotifiesItem.Time;
+					const FTransform RootBoneTransform = GetTransform(OriginTime, OriginRole, bUnused, RootSchemaBoneIdx);
+					const FTransform SamplingAttributeBoneTransform = GetTransform(SamplingAttributeTime, SampleRoleIndex, bClamped, TempBoneReference);
+					if (SchemaOriginBoneIdx == RootSchemaBoneIdx)
 					{
-						if (SamplingAttribute->SamplingAttributeId == SamplingAttributeId)
-						{
-							bool bUnused;
-							if (SamplingAttribute->Bone.BoneName != NAME_None)
-							{
-								FBoneReference TempBoneReference = SamplingAttribute->Bone;
-								const int32 SampleRoleIndex = RoleToIndex[SampleRole];
-								TempBoneReference.Initialize(BoneContainers[SampleRoleIndex].GetSkeletonAsset());
-								if (TempBoneReference.HasValidSetup())
-								{
-									const float SamplingAttributeTime = AnimNotifyEvent.GetTime();
-									const FTransform RootBoneTransform = GetTransform(OriginTime, OriginRole, bUnused, RootSchemaBoneIdx);
-									const FTransform SamplingAttributeBoneTransform = GetTransform(SamplingAttributeTime, SampleRoleIndex, bClamped, TempBoneReference);
-									if (SchemaOriginBoneIdx == RootSchemaBoneIdx)
-									{
-										OutSamplePosition = RootBoneTransform.InverseTransformPosition(SamplingAttributeBoneTransform.GetTranslation());
-									}
-									else
-									{
-										bool bOriginClamped;
-										const FTransform OriginBoneTransform = GetTransform(OriginTime, OriginRole, bOriginClamped, SchemaOriginBoneIdx);
-										bClamped |= bOriginClamped;
-										const FVector DeltaBoneTranslation = SamplingAttributeBoneTransform.GetTranslation() - OriginBoneTransform.GetTranslation();
-										OutSamplePosition = RootBoneTransform.InverseTransformVector(DeltaBoneTranslation);
-									}
-									return true;
-								}
-
-								UE_LOG(LogPoseSearch, Error, TEXT("FAssetIndexer::GetSamplePositionInternal: required UAnimNotifyState_PoseSearchSamplingAttribute in '%s' has an invalid Bone"), *AssetSamplers.GetAssetName());
-								return false;
-							}
-
-							const FTransform RootBoneTransform = GetTransform(OriginTime, OriginRole, bUnused, RootSchemaBoneIdx);
-							OutSamplePosition = RootBoneTransform.InverseTransformPosition(SamplingAttribute->Position);
-							return true;
-						}
+						OutSamplePosition = RootBoneTransform.InverseTransformPosition(SamplingAttributeBoneTransform.GetTranslation());
+					}
+					else
+					{
+						bool bOriginClamped;
+						const FTransform OriginBoneTransform = GetTransform(OriginTime, OriginRole, bOriginClamped, SchemaOriginBoneIdx);
+						bClamped |= bOriginClamped;
+						const FVector DeltaBoneTranslation = SamplingAttributeBoneTransform.GetTranslation() - OriginBoneTransform.GetTranslation();
+						OutSamplePosition = RootBoneTransform.InverseTransformVector(DeltaBoneTranslation);
 					}
+					return true;
 				}
+
+				UE_LOG(LogPoseSearch, Error, TEXT("FAssetIndexer::GetSamplePositionInternal: required UAnimNotifyState_PoseSearchSamplingAttribute in '%s' has an invalid Bone"), *AssetSamplers.GetAssetName());
 				return false;
-			}))
-		{
+			}
+
+			const FTransform RootBoneTransform = GetTransform(OriginTime, OriginRole, bUnused, RootSchemaBoneIdx);
+			OutSamplePosition = RootBoneTransform.InverseTransformPosition(SamplingAttribute->Position);
 			return true;
 		}
 
@@ -715,58 +693,47 @@ bool FAssetIndexer::GetSampleVelocity(FVector& OutSampleVelocity, float SampleTi
 	
 	if (SamplingAttributeId >= 0)
 	{
-		if (ProcessAllAnimNotifyEvents([&OutSampleVelocity, SchemaSampleBoneIdx, SchemaOriginBoneIdx, bUseCharacterSpaceVelocities, &SampleRole, &OriginRole, SamplingAttributeId, this, OriginTime](const TConstArrayView<FAnimNotifyEvent> AnimNotifyEvents)
+		const FPoseSearchTimedNotifies<UAnimNotifyState_PoseSearchSamplingAttribute> TimedNotifies(SamplingAttributeId, *this);
+		const FPoseSearchTimedNotifies<UAnimNotifyState_PoseSearchSamplingAttribute>::FItem TimedNotifiesItem = TimedNotifies.GetClosestFutureEvent(SampleTime);
+		if (const UAnimNotifyState_PoseSearchSamplingAttribute* SamplingAttribute = TimedNotifiesItem.NotifyState)
+		{
+			bool bUnused, bClampedPast;
+			FVector BonePositionPast, BonePositionPresent;
+			if (SamplingAttribute->Bone.BoneName != NAME_None)
 			{
-				for (const FAnimNotifyEvent& AnimNotifyEvent : AnimNotifyEvents)
+				FBoneReference TempBoneReference = SamplingAttribute->Bone;
+				const int32 SampleRoleIndex = RoleToIndex[SampleRole];
+				TempBoneReference.Initialize(BoneContainers[SampleRoleIndex].GetSkeletonAsset());
+				if (TempBoneReference.HasValidSetup())
 				{
-					if (const UAnimNotifyState_PoseSearchSamplingAttribute* SamplingAttribute = Cast<UAnimNotifyState_PoseSearchSamplingAttribute>(AnimNotifyEvent.NotifyStateClass))
+					const float SamplingAttributeTime = TimedNotifiesItem.Time;
+
+					if (GetSamplePositionInternal(BonePositionPast, SamplingAttributeTime - FiniteDelta, bUseCharacterSpaceVelocities ? OriginTime - FiniteDelta : OriginTime, bClampedPast, SchemaSampleBoneIdx, SchemaOriginBoneIdx, SampleRole, OriginRole, SamplingAttributeId) &&
+						GetSamplePositionInternal(BonePositionPresent, SamplingAttributeTime, OriginTime, bUnused, SchemaSampleBoneIdx, SchemaOriginBoneIdx, SampleRole, OriginRole, SamplingAttributeId))
 					{
-						if (SamplingAttribute->SamplingAttributeId == SamplingAttributeId)
+						if (!bClampedPast)
 						{
-							bool bUnused, bClampedPast;
-							FVector BonePositionPast, BonePositionPresent;
-							if (SamplingAttribute->Bone.BoneName != NAME_None)
-							{
-								FBoneReference TempBoneReference = SamplingAttribute->Bone;
-								const int32 SampleRoleIndex = RoleToIndex[SampleRole];
-								TempBoneReference.Initialize(BoneContainers[SampleRoleIndex].GetSkeletonAsset());
-								if (TempBoneReference.HasValidSetup())
-								{
-									const float SamplingAttributeTime = AnimNotifyEvent.GetTime();
-
-									if (GetSamplePositionInternal(BonePositionPast, SamplingAttributeTime - FiniteDelta, bUseCharacterSpaceVelocities ? OriginTime - FiniteDelta : OriginTime, bClampedPast, SchemaSampleBoneIdx, SchemaOriginBoneIdx, SampleRole, OriginRole, SamplingAttributeId) &&
-										GetSamplePositionInternal(BonePositionPresent, SamplingAttributeTime, OriginTime, bUnused, SchemaSampleBoneIdx, SchemaOriginBoneIdx, SampleRole, OriginRole, SamplingAttributeId))
-									{
-										if (!bClampedPast)
-										{
-											OutSampleVelocity = (BonePositionPresent - BonePositionPast) / FiniteDelta;
-											return true;
-										}
-
-										FVector BonePositionFuture;
-										if (GetSamplePositionInternal(BonePositionFuture, SamplingAttributeTime + FiniteDelta, bUseCharacterSpaceVelocities ? OriginTime + FiniteDelta : OriginTime, bUnused, SchemaSampleBoneIdx, SchemaOriginBoneIdx, SampleRole, OriginRole, SamplingAttributeId))
-										{
-											OutSampleVelocity = (BonePositionFuture - BonePositionPresent) / FiniteDelta;
-											return true;
-										}
-									}
-
-									return false;
-								}
-
-								UE_LOG(LogPoseSearch, Error, TEXT("FAssetIndexer::GetSampleVelocity: required UAnimNotifyState_PoseSearchSamplingAttribute in '%s' has an invalid Bone"), *AssetSamplers.GetAssetName());
-								return false;
-							}
+							OutSampleVelocity = (BonePositionPresent - BonePositionPast) / FiniteDelta;
+							return true;
+						}
 
-							const FTransform RootBoneTransform = GetTransform(OriginTime, OriginRole, bUnused, RootSchemaBoneIdx);
-							OutSampleVelocity = RootBoneTransform.InverseTransformPosition(SamplingAttribute->LinearVelocity);
+						FVector BonePositionFuture;
+						if (GetSamplePositionInternal(BonePositionFuture, SamplingAttributeTime + FiniteDelta, bUseCharacterSpaceVelocities ? OriginTime + FiniteDelta : OriginTime, bUnused, SchemaSampleBoneIdx, SchemaOriginBoneIdx, SampleRole, OriginRole, SamplingAttributeId))
+						{
+							OutSampleVelocity = (BonePositionFuture - BonePositionPresent) / FiniteDelta;
 							return true;
 						}
 					}
+
+					return false;
 				}
+
+				UE_LOG(LogPoseSearch, Error, TEXT("FAssetIndexer::GetSampleVelocity: required UAnimNotifyState_PoseSearchSamplingAttribute in '%s' has an invalid Bone"), *AssetSamplers.GetAssetName());
 				return false;
-			}))
-		{
+			}
+
+			const FTransform RootBoneTransform = GetTransform(OriginTime, OriginRole, bUnused, RootSchemaBoneIdx);
+			OutSampleVelocity = RootBoneTransform.InverseTransformPosition(SamplingAttribute->LinearVelocity);
 			return true;
 		}
 
diff --git a/Engine/Plugins/Animation/PoseSearch/Source/Runtime/Private/PoseSearchFeatureChannel_TimeToEvent.cpp b/Engine/Plugins/Animation/PoseSearch/Source/Runtime/Private/PoseSearchFeatureChannel_TimeToEvent.cpp
--- a/Engine/Plugins/Animation/PoseSearch/Source/Runtime/Private/PoseSearchFeatureChannel_TimeToEvent.cpp
+++ b/Engine/Plugins/Animation/PoseSearch/Source/Runtime/Private/PoseSearchFeatureChannel_TimeToEvent.cpp
@@ -52,31 +52,12 @@ bool UPoseSearchFeatureChannel_TimeToEvent::IndexAsset(UE::PoseSearch::FAssetInd
 {
 	using namespace UE::PoseSearch;
 
-	float EventTime = Indexer.GetPlayLength();
-
-	if (SamplingAttributeId >= 0)
-	{
-		Indexer.ProcessAllAnimNotifyEvents([&Indexer, this, &EventTime](const TConstArrayView<FAnimNotifyEvent> AnimNotifyEvents)
-			{
-				for (const FAnimNotifyEvent& AnimNotifyEvent : AnimNotifyEvents)
-				{
-					if (const UAnimNotifyState_PoseSearchSamplingEvent* SamplingEvent = Cast<UAnimNotifyState_PoseSearchSamplingEvent>(AnimNotifyEvent.NotifyStateClass))
-					{
-						if (SamplingEvent->SamplingAttributeId == SamplingAttributeId)
-						{
-							EventTime = AnimNotifyEvent.GetTime();
-							return true;
-						}
-					}
-				}
-
-				return false;
-			});
-	}
-
+	const FPoseSearchTimedNotifies<UAnimNotifyState_PoseSearchSamplingEvent> TimedNotifies(SamplingAttributeId, Indexer);
 	for (int32 SampleIdx = Indexer.GetBeginSampleIdx(); SampleIdx != Indexer.GetEndSampleIdx(); ++SampleIdx)
 	{
-		FFeatureVectorHelper::EncodeFloat(Indexer.GetPoseVector(SampleIdx), ChannelDataOffset, EventTime - Indexer.CalculateSampleTime(SampleIdx));
+		const float SampleTime = Indexer.CalculateSampleTime(SampleIdx);
+		const float EventTime = TimedNotifies.GetClosestFutureEvent(SampleTime).Time;
+		FFeatureVectorHelper::EncodeFloat(Indexer.GetPoseVector(SampleIdx), ChannelDataOffset, EventTime - SampleTime);
 	}
 	
 	return true;
diff --git a/Engine/Plugins/Animation/PoseSearch/Source/Runtime/Public/PoseSearch/PoseSearchAssetIndexer.h b/Engine/Plugins/Animation/PoseSearch/Source/Runtime/Public/PoseSearch/PoseSearchAssetIndexer.h
--- a/Engine/Plugins/Animation/PoseSearch/Source/Runtime/Public/PoseSearch/PoseSearchAssetIndexer.h
+++ b/Engine/Plugins/Animation/PoseSearch/Source/Runtime/Public/PoseSearch/PoseSearchAssetIndexer.h
@@ -160,6 +160,83 @@ private:
 	bool bProcessFailed = false;
 };
 
+template <class TAnimNotifyState>
+struct FPoseSearchTimedNotifies
+{
+	static_assert(TIsDerivedFrom<TAnimNotifyState, UAnimNotifyState_PoseSearchBase>::IsDerived, "The template class TAnimNotifyState must be a subclass of UAnimNotifyState_PoseSearchBase.");
+
+	struct FItem
+	{
+		float Time = 0.f;
+		const TAnimNotifyState* NotifyState = nullptr;
+	};
+
+	FPoseSearchTimedNotifies()
+	{
+	}
+
+	FPoseSearchTimedNotifies(int32 SamplingAttributeId, FAssetIndexer& Indexer)
+	{
+		Initialize(SamplingAttributeId, Indexer);
+	}
+
+	void Initialize(int32 SamplingAttributeId, FAssetIndexer& Indexer)
+	{
+		Items.Reset();
+
+		if (SamplingAttributeId >= 0)
+		{
+			Indexer.ProcessAllAnimNotifyEvents([SamplingAttributeId, this](const TConstArrayView<FAnimNotifyEvent> AnimNotifyEvents)
+				{
+					for (const FAnimNotifyEvent& AnimNotifyEvent : AnimNotifyEvents)
+					{
+						if (const TAnimNotifyState* SamplingEvent = Cast<TAnimNotifyState>(AnimNotifyEvent.NotifyStateClass))
+						{
+							if (SamplingEvent->SamplingAttributeId == SamplingAttributeId)
+							{
+								Items.Add({ AnimNotifyEvent.GetTime(), SamplingEvent });
+							}
+						}
+					}
+
+					return true;
+				});
+		
+			Items.Sort([](const FItem& A, const FItem& B) { return A.Time < B.Time; });
+		}
+
+		CachedPlayLength = Items.IsEmpty() ? Indexer.GetPlayLength() : 0.f;
+	}
+
+	FItem GetClosestFutureEvent(float SampleTime) const
+	{
+		const int32 NumItems = Items.Num();
+		if (NumItems == 0)
+		{
+			FItem Item;
+			Item.Time = CachedPlayLength;
+			return Item;
+		}
+
+		if (NumItems == 1)
+		{
+			return Items[0];
+		}
+
+		const int32 LowerBoundIdx = Algo::LowerBound(Items, SampleTime, [](const FItem& Item, float Value)
+		{
+			return Value > Item.Time;
+		});
+
+		const int32 ClampedLowerBoundIdx = FMath::Min(LowerBoundIdx, NumItems - 1);
+		return Items[ClampedLowerBoundIdx];
+	}
+
+private:
+	TArray<FItem, TInlineAllocator<128>> Items;
+	float CachedPlayLength = 0.f;
+};
+
 } // namespace UE::PoseSearch
 
 #endif // WITH_EDITOR
