motion matching - updating Experimental flag

diff --git a/Engine/Plugins/Animation/PoseSearch/Source/Runtime/Public/PoseSearch/PoseSearchDatabase.h b/Engine/Plugins/Animation/PoseSearch/Source/Runtime/Public/PoseSearch/PoseSearchDatabase.h
--- a/Engine/Plugins/Animation/PoseSearch/Source/Runtime/Public/PoseSearch/PoseSearchDatabase.h
+++ b/Engine/Plugins/Animation/PoseSearch/Source/Runtime/Public/PoseSearch/PoseSearchDatabase.h
@@ -36,7 +36,7 @@ enum class EPoseSearchMode : int32
 	PCAKDTree,
 
 	// Optimized search mode using a vantage point tree (Experimental)
-	VPTree
+	VPTree UMETA(DisplayName = "VPTree (Experimental)")
 };
 
 UENUM()
diff --git a/Engine/Plugins/Animation/PoseSearch/Source/Runtime/Public/PoseSearch/PoseSearchTrajectoryLibrary.h b/Engine/Plugins/Animation/PoseSearch/Source/Runtime/Public/PoseSearch/PoseSearchTrajectoryLibrary.h
--- a/Engine/Plugins/Animation/PoseSearch/Source/Runtime/Public/PoseSearch/PoseSearchTrajectoryLibrary.h
+++ b/Engine/Plugins/Animation/PoseSearch/Source/Runtime/Public/PoseSearch/PoseSearchTrajectoryLibrary.h
@@ -84,7 +84,7 @@ public:
 /**
  * Set of functions to help populate a FPoseSearchQueryTrajectory for motion matching.
  */
-UCLASS(Experimental)
+UCLASS()
 class POSESEARCH_API UPoseSearchTrajectoryLibrary : public UBlueprintFunctionLibrary
 {
 	GENERATED_BODY()
