motion matching - heading channel normalization fix

diff --git a/Engine/Plugins/Animation/PoseSearch/Source/Runtime/Private/PoseSearchDerivedDataKey.cpp b/Engine/Plugins/Animation/PoseSearch/Source/Runtime/Private/PoseSearchDerivedDataKey.cpp
--- a/Engine/Plugins/Animation/PoseSearch/Source/Runtime/Private/PoseSearchDerivedDataKey.cpp
+++ b/Engine/Plugins/Animation/PoseSearch/Source/Runtime/Private/PoseSearchDerivedDataKey.cpp
@@ -41,7 +41,7 @@ FKeyBuilder::FKeyBuilder(const UObject* Object, bool bUseDataVer, bool bPerformC
 	if (bUseDataVer)
 	{
 		// used to invalidate the key without having to change POSESEARCHDB_DERIVEDDATA_VER all the times
-		int32 POSESEARCHDB_DERIVEDDATA_VER_SMALL = 237;
+		int32 POSESEARCHDB_DERIVEDDATA_VER_SMALL = 239;
 		FGuid VersionGuid = FDevSystemGuids::GetSystemGuid(FDevSystemGuids::Get().POSESEARCHDB_DERIVEDDATA_VER);
 
 		*this << VersionGuid;
diff --git a/Engine/Plugins/Animation/PoseSearch/Source/Runtime/Private/PoseSearchFeatureChannel.cpp b/Engine/Plugins/Animation/PoseSearch/Source/Runtime/Private/PoseSearchFeatureChannel.cpp
--- a/Engine/Plugins/Animation/PoseSearch/Source/Runtime/Private/PoseSearchFeatureChannel.cpp
+++ b/Engine/Plugins/Animation/PoseSearch/Source/Runtime/Private/PoseSearchFeatureChannel.cpp
@@ -26,25 +26,56 @@ int32 FFeatureVectorHelper::GetVectorCardinality(EComponentStrippingVector Compo
 	}
 }
 
-void FFeatureVectorHelper::EncodeVector(TArrayView<float> Values, int32 DataOffset, const FVector& Vector, EComponentStrippingVector ComponentStrippingVector)
+void FFeatureVectorHelper::EncodeVector(TArrayView<float> Values, int32 DataOffset, const FVector& Vector, EComponentStrippingVector ComponentStrippingVector, bool bNormalize)
 {
-	switch (ComponentStrippingVector)
+	if (bNormalize)
 	{
-	case EComponentStrippingVector::None:
-		Values[DataOffset + 0] = Vector.X;
-		Values[DataOffset + 1] = Vector.Y;
-		Values[DataOffset + 2] = Vector.Z;
-		break;
-	case EComponentStrippingVector::StripXY:
-		Values[DataOffset + 0] = Vector.Z;
-		break;
-	case EComponentStrippingVector::StripZ:
-		Values[DataOffset + 0] = Vector.X;
-		Values[DataOffset + 1] = Vector.Y;
-		break;
-	default:
-		checkNoEntry();
-		break;
+		FVector NormalizedVector = Vector;
+		switch (ComponentStrippingVector)
+		{
+		case EComponentStrippingVector::None:
+			NormalizedVector.Normalize();
+			Values[DataOffset + 0] = NormalizedVector.X;
+			Values[DataOffset + 1] = NormalizedVector.Y;
+			Values[DataOffset + 2] = NormalizedVector.Z;
+			break;
+		case EComponentStrippingVector::StripXY:
+			NormalizedVector.X = 0.f;
+			NormalizedVector.Y = 0.f;
+			NormalizedVector.Normalize();
+			Values[DataOffset + 0] = NormalizedVector.Z;
+			break;
+		case EComponentStrippingVector::StripZ:
+			NormalizedVector.Z = 0.f;
+			NormalizedVector.Normalize();
+			Values[DataOffset + 0] = NormalizedVector.X;
+			Values[DataOffset + 1] = NormalizedVector.Y;
+			break;
+		default:
+			checkNoEntry();
+			break;
+		}
+	}
+	else
+	{
+		switch (ComponentStrippingVector)
+		{
+		case EComponentStrippingVector::None:
+			Values[DataOffset + 0] = Vector.X;
+			Values[DataOffset + 1] = Vector.Y;
+			Values[DataOffset + 2] = Vector.Z;
+			break;
+		case EComponentStrippingVector::StripXY:
+			Values[DataOffset + 0] = Vector.Z;
+			break;
+		case EComponentStrippingVector::StripZ:
+			Values[DataOffset + 0] = Vector.X;
+			Values[DataOffset + 1] = Vector.Y;
+			break;
+		default:
+			checkNoEntry();
+			break;
+		}
 	}
 }
 
diff --git a/Engine/Plugins/Animation/PoseSearch/Source/Runtime/Private/PoseSearchFeatureChannel_Heading.cpp b/Engine/Plugins/Animation/PoseSearch/Source/Runtime/Private/PoseSearchFeatureChannel_Heading.cpp
--- a/Engine/Plugins/Animation/PoseSearch/Source/Runtime/Private/PoseSearchFeatureChannel_Heading.cpp
+++ b/Engine/Plugins/Animation/PoseSearch/Source/Runtime/Private/PoseSearchFeatureChannel_Heading.cpp
@@ -108,7 +108,7 @@ void UPoseSearchFeatureChannel_Heading::BuildQuery(UE::PoseSearch::FSearchContex
 	{
 		const FQuat BoneRotationWorld = BP_GetWorldRotation(SearchContext.GetAnimInstance(SampleRole));
 		const FQuat BoneRotation = SearchContext.GetSampleRotation(SampleTimeOffset, OriginTimeOffset, SchemaBoneIdx, SchemaOriginBoneIdx, SampleRole, OriginRole, EPermutationTimeType::UseSampleTime, &BoneRotationWorld);
-		FFeatureVectorHelper::EncodeVector(SearchContext.EditFeatureVector(), ChannelDataOffset, GetAxis(BoneRotation), ComponentStripping);
+		FFeatureVectorHelper::EncodeVector(SearchContext.EditFeatureVector(), ChannelDataOffset, GetAxis(BoneRotation), ComponentStripping, true);
 		return;
 	}
 	
@@ -175,7 +175,7 @@ void UPoseSearchFeatureChannel_Heading::BuildQuery(UE::PoseSearch::FSearchContex
 	
 	// calculating the BoneRotation in component space for the bone indexed by SchemaBoneIdx
 	const FQuat BoneRotation = SearchContext.GetSampleRotation(SampleTimeOffset, OriginTimeOffset, SchemaBoneIdx, SchemaOriginBoneIdx, SampleRole, OriginRole, PermutationTimeType);
-	FFeatureVectorHelper::EncodeVector(SearchContext.EditFeatureVector(), ChannelDataOffset, GetAxis(BoneRotation), ComponentStripping);
+	FFeatureVectorHelper::EncodeVector(SearchContext.EditFeatureVector(), ChannelDataOffset, GetAxis(BoneRotation), ComponentStripping,true);
 }
 
 #if ENABLE_DRAW_DEBUG
@@ -216,7 +216,7 @@ bool UPoseSearchFeatureChannel_Heading::IndexAsset(UE::PoseSearch::FAssetIndexer
 	{
 		if (Indexer.GetSampleRotation(SampleRotation, SampleTimeOffset, OriginTimeOffset, SampleIdx, SchemaBoneIdx, SchemaOriginBoneIdx, SampleRole, OriginRole, PermutationTimeType, SamplingAttributeId))
 		{
-			FFeatureVectorHelper::EncodeVector(Indexer.GetPoseVector(SampleIdx), ChannelDataOffset, GetAxis(SampleRotation), ComponentStripping);
+			FFeatureVectorHelper::EncodeVector(Indexer.GetPoseVector(SampleIdx), ChannelDataOffset, GetAxis(SampleRotation), ComponentStripping, true);
 		}
 		else
 		{
diff --git a/Engine/Plugins/Animation/PoseSearch/Source/Runtime/Private/PoseSearchFeatureChannel_Position.cpp b/Engine/Plugins/Animation/PoseSearch/Source/Runtime/Private/PoseSearchFeatureChannel_Position.cpp
--- a/Engine/Plugins/Animation/PoseSearch/Source/Runtime/Private/PoseSearchFeatureChannel_Position.cpp
+++ b/Engine/Plugins/Animation/PoseSearch/Source/Runtime/Private/PoseSearchFeatureChannel_Position.cpp
@@ -85,7 +85,7 @@ void UPoseSearchFeatureChannel_Position::BuildQuery(UE::PoseSearch::FSearchConte
 	{
 		const FVector BonePositionWorld = BP_GetWorldPosition(SearchContext.GetAnimInstance(SampleRole));
 		const FVector BonePosition = SearchContext.GetSamplePosition(SampleTimeOffset, OriginTimeOffset, SchemaBoneIdx, SchemaOriginBoneIdx, SampleRole, OriginRole, PermutationTimeType, &BonePositionWorld);
-  		FFeatureVectorHelper::EncodeVector(SearchContext.EditFeatureVector(), ChannelDataOffset, BonePosition, ComponentStripping);
+  		FFeatureVectorHelper::EncodeVector(SearchContext.EditFeatureVector(), ChannelDataOffset, BonePosition, ComponentStripping, false);
 		return;
 	}
 
@@ -151,7 +151,7 @@ void UPoseSearchFeatureChannel_Position::BuildQuery(UE::PoseSearch::FSearchConte
 	
 	// calculating the BonePosition in root bone space for the bone indexed by SchemaBoneIdx
 	const FVector BonePosition = SearchContext.GetSamplePosition(SampleTimeOffset, OriginTimeOffset, SchemaBoneIdx, SchemaOriginBoneIdx, SampleRole, OriginRole, PermutationTimeType);
-	FFeatureVectorHelper::EncodeVector(SearchContext.EditFeatureVector(), ChannelDataOffset, BonePosition, ComponentStripping);
+	FFeatureVectorHelper::EncodeVector(SearchContext.EditFeatureVector(), ChannelDataOffset, BonePosition, ComponentStripping, false);
 }
 
 #if ENABLE_DRAW_DEBUG
@@ -201,7 +201,7 @@ bool UPoseSearchFeatureChannel_Position::IndexAsset(UE::PoseSearch::FAssetIndexe
 	{
 		if (Indexer.GetSamplePosition(BonePosition, SampleTimeOffset, OriginTimeOffset, SampleIdx, SchemaBoneIdx, SchemaOriginBoneIdx, SampleRole, OriginRole, PermutationTimeType, SamplingAttributeId))
 		{
-			FFeatureVectorHelper::EncodeVector(Indexer.GetPoseVector(SampleIdx), ChannelDataOffset, BonePosition, ComponentStripping);
+			FFeatureVectorHelper::EncodeVector(Indexer.GetPoseVector(SampleIdx), ChannelDataOffset, BonePosition, ComponentStripping, false);
 		}
 		else
 		{
diff --git a/Engine/Plugins/Animation/PoseSearch/Source/Runtime/Private/PoseSearchFeatureChannel_Velocity.cpp b/Engine/Plugins/Animation/PoseSearch/Source/Runtime/Private/PoseSearchFeatureChannel_Velocity.cpp
--- a/Engine/Plugins/Animation/PoseSearch/Source/Runtime/Private/PoseSearchFeatureChannel_Velocity.cpp
+++ b/Engine/Plugins/Animation/PoseSearch/Source/Runtime/Private/PoseSearchFeatureChannel_Velocity.cpp
@@ -52,7 +52,7 @@ void UPoseSearchFeatureChannel_Velocity::BuildQuery(UE::PoseSearch::FSearchConte
 		{
 			LinearVelocity = LinearVelocity.GetClampedToMaxSize(1.f);
 		}
-		FFeatureVectorHelper::EncodeVector(SearchContext.EditFeatureVector(), ChannelDataOffset, LinearVelocity, ComponentStripping);
+		FFeatureVectorHelper::EncodeVector(SearchContext.EditFeatureVector(), ChannelDataOffset, LinearVelocity, ComponentStripping, false);
 		return;
 	}
 	
@@ -126,7 +126,7 @@ void UPoseSearchFeatureChannel_Velocity::BuildQuery(UE::PoseSearch::FSearchConte
 		LinearVelocity = LinearVelocity.GetClampedToMaxSize(1.f);
 	}
 
-	FFeatureVectorHelper::EncodeVector(SearchContext.EditFeatureVector(), ChannelDataOffset, LinearVelocity, ComponentStripping);
+	FFeatureVectorHelper::EncodeVector(SearchContext.EditFeatureVector(), ChannelDataOffset, LinearVelocity, ComponentStripping, false);
 }
 
 #if ENABLE_DRAW_DEBUG
@@ -173,7 +173,7 @@ bool UPoseSearchFeatureChannel_Velocity::IndexAsset(UE::PoseSearch::FAssetIndexe
 			{
 				LinearVelocity = LinearVelocity.GetClampedToMaxSize(1.f);
 			}
-			FFeatureVectorHelper::EncodeVector(Indexer.GetPoseVector(SampleIdx), ChannelDataOffset, LinearVelocity, ComponentStripping);
+			FFeatureVectorHelper::EncodeVector(Indexer.GetPoseVector(SampleIdx), ChannelDataOffset, LinearVelocity, ComponentStripping, false);
 		}
 		else
 		{
diff --git a/Engine/Plugins/Animation/PoseSearch/Source/Runtime/Public/PoseSearch/PoseSearchFeatureChannel.h b/Engine/Plugins/Animation/PoseSearch/Source/Runtime/Public/PoseSearch/PoseSearchFeatureChannel.h
--- a/Engine/Plugins/Animation/PoseSearch/Source/Runtime/Public/PoseSearch/PoseSearchFeatureChannel.h
+++ b/Engine/Plugins/Animation/PoseSearch/Source/Runtime/Public/PoseSearch/PoseSearchFeatureChannel.h
@@ -83,7 +83,7 @@ class POSESEARCH_API FFeatureVectorHelper
 {
 public:
 	static int32 GetVectorCardinality(EComponentStrippingVector ComponentStrippingVector);
-	static void EncodeVector(TArrayView<float> Values, int32 DataOffset, const FVector& Vector, EComponentStrippingVector ComponentStrippingVector);
+	static void EncodeVector(TArrayView<float> Values, int32 DataOffset, const FVector& Vector, EComponentStrippingVector ComponentStrippingVector, bool bNormalize);
 	static FVector DecodeVector(TConstArrayView<float> Values, int32 DataOffset, EComponentStrippingVector ComponentStrippingVector);
 
 	static void EncodeVector2D(TArrayView<float> Values, int32 DataOffset, const FVector2D& Vector2D);
