Motion Matching - Clear pose search database preview actors after deleting a selection of nodes.

#rnx none

diff --git a/Engine/Plugins/Animation/PoseSearch/Source/Editor/Private/PoseSearchDatabaseAssetTree.cpp b/Engine/Plugins/Animation/PoseSearch/Source/Editor/Private/PoseSearchDatabaseAssetTree.cpp
--- a/Engine/Plugins/Animation/PoseSearch/Source/Editor/Private/PoseSearchDatabaseAssetTree.cpp
+++ b/Engine/Plugins/Animation/PoseSearch/Source/Editor/Private/PoseSearchDatabaseAssetTree.cpp
@@ -726,6 +726,7 @@ namespace UE::PoseSearch
 				}
 			}
 			
+			ViewModel->RemovePreviewActors();
 			FinalizeTreeChanges();
 		}
 	}
