Motion Matching Editor - Clipboard support for database items and other fixes.

- Moved Modify() calls before any calls to modify UPoseSearchDatabase and/or UAnimSequenceBase to properly support reverting back to previous states.
- Made PoseSearchDatabaseViewModel now call Modify() on its respective PoseSearchDatabase before adding any anim assets.
- Added PoseSearchDatabaseEditorClipboard.cpp/h
- Fixed OnAddAnimComposite/OnAddAnimMontage/OnAddMultiSequence not being revertable actions.

#rnx minor

diff --git a/Engine/Plugins/Animation/PoseSearch/Source/Editor/PoseSearchEditor.Build.cs b/Engine/Plugins/Animation/PoseSearch/Source/Editor/PoseSearchEditor.Build.cs
--- a/Engine/Plugins/Animation/PoseSearch/Source/Editor/PoseSearchEditor.Build.cs
+++ b/Engine/Plugins/Animation/PoseSearch/Source/Editor/PoseSearchEditor.Build.cs
@@ -44,6 +44,7 @@ public class PoseSearchEditor : ModuleRules
 		PrivateDependencyModuleNames.AddRange(
 			new string[]
 			{
+				"ApplicationCore",
 				"AnimationEditor",
 				"AssetDefinition",
 				"BlueprintGraph",
diff --git a/Engine/Plugins/Animation/PoseSearch/Source/Editor/Private/PoseSearchDatabaseAssetTree.cpp b/Engine/Plugins/Animation/PoseSearch/Source/Editor/Private/PoseSearchDatabaseAssetTree.cpp
--- a/Engine/Plugins/Animation/PoseSearch/Source/Editor/Private/PoseSearchDatabaseAssetTree.cpp
+++ b/Engine/Plugins/Animation/PoseSearch/Source/Editor/Private/PoseSearchDatabaseAssetTree.cpp
@@ -9,6 +9,7 @@
 #include "AssetSelection.h"
 #include "ClassIconFinder.h"
 #include "DetailColumnSizeData.h"
+#include "PoseSearchDatabaseEditorClipboard.h"
 #include "DragAndDrop/AssetDragDropOp.h"
 #include "Framework/Commands/GenericCommands.h"
 #include "Framework/MultiBox/MultiBoxBuilder.h"
@@ -17,6 +18,7 @@
 #include "PoseSearch/PoseSearchAnimNotifies.h"
 #include "PoseSearch/PoseSearchDatabase.h"
 #include "PoseSearchDatabaseViewModel.h"
+#include "PoseSearchEditor.h"
 #include "SPositiveActionButton.h"
 #include "Styling/AppStyle.h"
 #include "ScopedTransaction.h"
@@ -395,9 +397,9 @@ namespace UE::PoseSearch
 		{
 			GWarn->BeginSlowTask(LOCTEXT("LoadingAssets", "Loading Asset(s)"), true);
 
-			const FScopedTransaction Transaction(LOCTEXT("AddSequencesOrBlendspaces", "Add Sequence(s) and/or Blendspace(s) to Pose Search Database"));
-			
+			const FScopedTransaction Transaction(LOCTEXT("AddAssetsOnDrop", "Add Animation Asset(s) to Pose Search Database"));
 			PoseSearchDatabase->Modify();
+			
 			for (int32 DroppedAssetIdx = 0; DroppedAssetIdx < NumAssets; ++DroppedAssetIdx)
 			{
 				const FAssetData& AssetData = DroppedAssetData[DroppedAssetIdx];
@@ -572,37 +574,26 @@ namespace UE::PoseSearch
 
 	void SDatabaseAssetTree::OnAddSequence(bool bFinalizeChanges)
 	{
-		const TSharedPtr<FDatabaseViewModel> ViewModel = EditorViewModel.Pin();
-		if (UPoseSearchDatabase* PoseSearchDatabase = ViewModel->GetPoseSearchDatabase())
-		{
-			FScopedTransaction Transaction(LOCTEXT("AddSequence", "Add Sequence"));
-
-			PoseSearchDatabase->Modify();
+		FScopedTransaction Transaction(LOCTEXT("AddSequence", "Add Sequence"));
 
-			ViewModel->AddSequenceToDatabase(nullptr);
+		EditorViewModel.Pin()->AddSequenceToDatabase(nullptr);
 
-			if (bFinalizeChanges)
-			{
-				FinalizeTreeChanges();
-			}
+		if (bFinalizeChanges)
+		{
+			FinalizeTreeChanges();
 		}
+		
 	}
 
 	void SDatabaseAssetTree::OnAddBlendSpace(bool bFinalizeChanges)
 	{
-		const TSharedPtr<FDatabaseViewModel> ViewModel = EditorViewModel.Pin();
-		if (UPoseSearchDatabase* PoseSearchDatabase = ViewModel->GetPoseSearchDatabase())
-		{
-			FScopedTransaction Transaction(LOCTEXT("AddBlendSpaceTransaction", "Add Blend Space"));
-
-			PoseSearchDatabase->Modify();
-
-			ViewModel->AddBlendSpaceToDatabase(nullptr);
+		FScopedTransaction Transaction(LOCTEXT("AddBlendSpaceTransaction", "Add Blend Space"));
+		
+		EditorViewModel.Pin()->AddBlendSpaceToDatabase(nullptr);
 
-			if (bFinalizeChanges)
-			{
-				FinalizeTreeChanges();
-			}
+		if (bFinalizeChanges)
+		{
+			FinalizeTreeChanges();
 		}
 	}
 
@@ -683,18 +674,33 @@ namespace UE::PoseSearch
 	void SDatabaseAssetTree::CreateCommandList()
 	{
 		CommandList = MakeShared<FUICommandList>();
-
+		
 		CommandList->MapAction(
 			FGenericCommands::Get().Delete,
 			FUIAction(
 				FExecuteAction::CreateSP(this, &SDatabaseAssetTree::OnDeleteNodes),
 				FCanExecuteAction::CreateSP(this, &SDatabaseAssetTree::CanDeleteNodes)));
+		
+		CommandList->MapAction(
+		FGenericCommands::Get().Copy,
+		FExecuteAction::CreateSP(this, &SDatabaseAssetTree::OnCopySelectedNodesToClipboard),
+		FCanExecuteAction::CreateSP(this, &SDatabaseAssetTree::CanCopyToClipboard));
+
+		CommandList->MapAction(
+			FGenericCommands::Get().Paste,
+			FExecuteAction::CreateSP(this, &SDatabaseAssetTree::OnPasteNodesFromClipboard),
+			FCanExecuteAction::CreateSP(this, &SDatabaseAssetTree::CanPasteFromClipboard));
+	
+		CommandList->MapAction(
+			FGenericCommands::Get().Cut,
+			FExecuteAction::CreateSP(this, &SDatabaseAssetTree::OnCutSelectedNodesToClipboard),
+			FCanExecuteAction::CreateSP(this, &SDatabaseAssetTree::CanCutToClipboard));
 	}
 
 	bool SDatabaseAssetTree::CanDeleteNodes() const
 	{
 		TArray<TSharedPtr<FDatabaseAssetTreeNode>> SelectedNodes = TreeView->GetSelectedItems();
-		for (TSharedPtr<FDatabaseAssetTreeNode> SelectedNode : SelectedNodes)
+		for (const TSharedPtr<FDatabaseAssetTreeNode>& SelectedNode : SelectedNodes)
 		{
 			if (SelectedNode->SourceAssetIdx != INDEX_NONE)
 			{
@@ -710,7 +716,7 @@ namespace UE::PoseSearch
 		TArray<TSharedPtr<FDatabaseAssetTreeNode>> SelectedNodes = TreeView->GetSelectedItems();
 		if (!SelectedNodes.IsEmpty())
 		{
-			const FScopedTransaction Transaction(LOCTEXT("DeletePoseSearchDatabaseNodes", "Delete selected items from Pose Search Database"));
+			const FScopedTransaction Transaction(LOCTEXT("DeletePoseSearchDatabaseNodes", "Delete selected item(s) from Pose Search Database"));
 			const TSharedPtr<FDatabaseViewModel> ViewModel = EditorViewModel.Pin();
 
 			SelectedNodes.Sort([](const TSharedPtr<FDatabaseAssetTreeNode>& A, const TSharedPtr<FDatabaseAssetTreeNode>& B)
@@ -718,19 +724,136 @@ namespace UE::PoseSearch
 					return B->SourceAssetIdx < A->SourceAssetIdx;
 				});
 
-			for (TSharedPtr<FDatabaseAssetTreeNode> SelectedNode : SelectedNodes)
+			for (const TSharedPtr<FDatabaseAssetTreeNode>& SelectedNode : SelectedNodes)
 			{
 				if (SelectedNode->SourceAssetIdx != INDEX_NONE)
 				{
 					OnDeleteAsset(SelectedNode, false);
 				}
 			}
+
+			ViewModel->RemovePreviewActors();
+			FinalizeTreeChanges();
+		}
+	}
+
+	void SDatabaseAssetTree::OnCopySelectedNodesToClipboard() const
+	{
+		TArray<TSharedPtr<FDatabaseAssetTreeNode>> SelectedNodes = TreeView->GetSelectedItems();
+		
+		if (!SelectedNodes.IsEmpty())
+		{
+			if (UPoseSearchDatabaseEditorClipboardContent* ClipboardContent = UPoseSearchDatabaseEditorClipboardContent::Create())
+			{
+				const FScopedTransaction Transaction(LOCTEXT("CopyPoseSearchDatabaseNodes", "Copy selected item(s) from Pose Search Database"));
+				const TSharedPtr<FDatabaseViewModel> ViewModel = EditorViewModel.Pin();
+				
+				for (const TSharedPtr<FDatabaseAssetTreeNode>& SelectedNode : SelectedNodes)
+				{
+					if (SelectedNode->SourceAssetIdx != INDEX_NONE)
+					{
+						if (UPoseSearchDatabase* Database = ViewModel->GetPoseSearchDatabase())
+						{
+							if (const FPoseSearchDatabaseAnimationAssetBase* DatabaseAnimationAsset = Database->GetMutableAnimationAssetBase(SelectedNode->SourceAssetIdx))
+							{
+								// @todo: Support copying assets added via BranchIn notifies. 
+								if (!DatabaseAnimationAsset->bSynchronizeWithExternalDependency)
+								{
+									ClipboardContent->CopyDatabaseItem(DatabaseAnimationAsset);
+								}
+								else
+								{
+									UE_LOG(LogPoseSearchEditor, Log, TEXT("Failed to copy %s. Asset(s) with BranchIn notifies do not have clipboard support."), *DatabaseAnimationAsset->GetName())
+								}
+							}
+						}
+					}
+				}
 			
+				ClipboardContent->CopyToClipboard();
+			}
+			else
+			{
+				UE_LOG(LogPoseSearchEditor, Warning, TEXT("Failed create clipboard object while attempting to copy data"));
+			}
+		}
+	}
+
+	bool SDatabaseAssetTree::CanCopyToClipboard() const
+	{
+		return !TreeView->GetSelectedItems().IsEmpty();
+	}
+
+	void SDatabaseAssetTree::OnPasteNodesFromClipboard()
+	{
+		const TSharedPtr<FDatabaseViewModel> ViewModel = EditorViewModel.Pin();
+		
+		if (const UPoseSearchDatabaseEditorClipboardContent* ClipboardContent = UPoseSearchDatabaseEditorClipboardContent::CreateFromClipboard())
+		{
+			if (UPoseSearchDatabase* Database = ViewModel->GetPoseSearchDatabase())
+			{
+				const FScopedTransaction Transaction(LOCTEXT("PastePoseSearchDatabaseNodes", "Paste item(s) to Pose Search Database"));
+			
+				ClipboardContent->PasteToDatabase(Database);
+
+				FinalizeTreeChanges();
+			}
+		}
+		else
+		{
+			UE_LOG(LogPoseSearchEditor, Warning, TEXT("Failed to get valid clipboard data while attempting to paste data"));
+		}
+	}
+
+	bool SDatabaseAssetTree::CanPasteFromClipboard()
+	{
+		const UPoseSearchDatabaseEditorClipboardContent* ClipboardContent = UPoseSearchDatabaseEditorClipboardContent::CreateFromClipboard();
+		return ClipboardContent && !ClipboardContent->DatabaseItems.IsEmpty();
+	}
+
+	void SDatabaseAssetTree::OnCutSelectedNodesToClipboard()
+	{
+		const FScopedTransaction Transaction(LOCTEXT("CutPoseSearchDatabaseNodes", "Cut selected item(s) from Pose Search Database"));
+		
+		OnCopySelectedNodesToClipboard();
+
+		// @todo: Following code can be replaced with OnDeleteNodes() call once assets with external dependencies support copying/pasting.
+		TArray<TSharedPtr<FDatabaseAssetTreeNode>> SelectedNodes = TreeView->GetSelectedItems();
+		if (!SelectedNodes.IsEmpty())
+		{
+			const TSharedPtr<FDatabaseViewModel> ViewModel = EditorViewModel.Pin();
+
+			SelectedNodes.Sort([](const TSharedPtr<FDatabaseAssetTreeNode>& A, const TSharedPtr<FDatabaseAssetTreeNode>& B)
+				{
+					return B->SourceAssetIdx < A->SourceAssetIdx;
+				});
+
+			for (const TSharedPtr<FDatabaseAssetTreeNode>& SelectedNode : SelectedNodes)
+			{
+				if (SelectedNode->SourceAssetIdx != INDEX_NONE)
+				{
+					if (UPoseSearchDatabase* Database = ViewModel->GetPoseSearchDatabase())
+					{
+						const FPoseSearchDatabaseAnimationAssetBase* DatabaseAnimationAsset = Database->GetAnimationAssetBase(SelectedNode->SourceAssetIdx);
+						
+						if (DatabaseAnimationAsset && !DatabaseAnimationAsset->bSynchronizeWithExternalDependency)
+						{
+							OnDeleteAsset(SelectedNode, false);
+						}
+					}
+				}
+			}
+
 			ViewModel->RemovePreviewActors();
 			FinalizeTreeChanges();
 		}
 	}
 
+	bool SDatabaseAssetTree::CanCutToClipboard() const
+	{
+		return CanCopyToClipboard() && CanDeleteNodes();
+	}
+
 	void SDatabaseAssetTree::EnableSelectedNodes(bool bIsEnabled)
 	{
 		const TSharedPtr<FDatabaseViewModel> ViewModel = EditorViewModel.Pin();
@@ -744,7 +867,7 @@ namespace UE::PoseSearch
 
 				PoseSearchDatabase->Modify();
 
-				for (TSharedPtr<FDatabaseAssetTreeNode> SelectedNode : SelectedNodes)
+				for (const TSharedPtr<FDatabaseAssetTreeNode>& SelectedNode : SelectedNodes)
 				{
 					ViewModel->SetIsEnabled(SelectedNode->SourceAssetIdx, bIsEnabled);
 				}
@@ -764,16 +887,24 @@ namespace UE::PoseSearch
 			{
 				const FScopedTransaction Transaction(LOCTEXT("ConvertToBranchInTransaction", "Create PoseSearchBranchIn notify state for assets in Pose Search Database"));
 
-				for (TSharedPtr<FDatabaseAssetTreeNode> SelectedNode : SelectedNodes)
+				bool bModified = false;
+				
+				for (const TSharedPtr<FDatabaseAssetTreeNode>& SelectedNode : SelectedNodes)
 				{
 					if (FPoseSearchDatabaseAnimationAssetBase* DatabaseAnimationAssetBase = PoseSearchDatabase->GetMutableAnimationAssetBase(SelectedNode->SourceAssetIdx))
 					{
 						if (UAnimSequenceBase* AnimSequenceBase = Cast<UAnimSequenceBase>(DatabaseAnimationAssetBase->GetAnimationAsset()))
 						{
+							if (!bModified)
+							{
+								AnimSequenceBase->Modify();
+								bModified = true;
+							}
+							
 							const FFloatInterval SamplingRange = FPoseSearchDatabaseAnimationAssetBase::GetEffectiveSamplingRange(AnimSequenceBase, DatabaseAnimationAssetBase->GetSamplingRange());
 							const float StartTime = SamplingRange.Min;
 							const float Duration = SamplingRange.Max - SamplingRange.Min;
-							FName TrackName = "PoseSearch";
+							const FName TrackName = "PoseSearch";
 							
 							if (!UAnimationBlueprintLibrary::IsValidAnimNotifyTrackName(AnimSequenceBase, TrackName))
 							{
@@ -783,7 +914,6 @@ namespace UE::PoseSearch
 							UAnimNotifyState_PoseSearchBranchIn* PoseSearchBranchIn = CastChecked<UAnimNotifyState_PoseSearchBranchIn>(UAnimationBlueprintLibrary::AddAnimationNotifyStateEvent(AnimSequenceBase, TrackName, StartTime, Duration, UAnimNotifyState_PoseSearchBranchIn::StaticClass()));
 							PoseSearchBranchIn->Database = PoseSearchDatabase;
 							DatabaseAnimationAssetBase->bSynchronizeWithExternalDependency = true;
-							AnimSequenceBase->Modify();
 						}
 					}
 				}
diff --git a/Engine/Plugins/Animation/PoseSearch/Source/Editor/Private/PoseSearchDatabaseEditorClipboard.cpp b/Engine/Plugins/Animation/PoseSearch/Source/Editor/Private/PoseSearchDatabaseEditorClipboard.cpp
new file mode 100644
--- /dev/null
+++ b/Engine/Plugins/Animation/PoseSearch/Source/Editor/Private/PoseSearchDatabaseEditorClipboard.cpp
@@ -0,0 +1,166 @@
+﻿// Copyright Epic Games, Inc. All Rights Reserved.
+
+#include "PoseSearchDatabaseEditorClipboard.h"
+
+#include "Factories.h"
+#include "InstancedStruct.h"
+#include "Animation/AnimComposite.h"
+#include "Animation/AnimSequence.h"
+#include "Animation/AnimMontage.h"
+#include "Animation/BlendSpace.h"
+#include "PoseSearch/PoseSearchMultiSequence.h"
+#include "PoseSearch/PoseSearchDatabase.h"
+
+#include "UnrealExporter.h"
+#include "Exporters/Exporter.h"
+#include "HAL/PlatformApplicationMisc.h"
+
+void UPoseSearchDatabaseEditorClipboardContent::CopyDatabaseItem(const FPoseSearchDatabaseAnimationAssetBase* InItem)
+{
+	check(InItem != nullptr)
+	
+	UPoseSearchDatabaseItemCopyObject* CopyObj = NewObject<UPoseSearchDatabaseItemCopyObject>((UObject*)GetTransientPackage(), UPoseSearchDatabaseItemCopyObject::StaticClass(), NAME_None, RF_Transient);
+	
+	CopyObj->ClassName = InItem->GetAnimationAssetStaticClass()->GetFName();
+	
+	if (UAnimSequence::StaticClass() == InItem->GetAnimationAssetStaticClass())
+	{
+		FPoseSearchDatabaseSequence::StaticStruct()->ExportText(CopyObj->Content, InItem, nullptr, nullptr, PPF_None, nullptr);
+	}
+	else if (UAnimComposite::StaticClass() == InItem->GetAnimationAssetStaticClass())
+	{
+		FPoseSearchDatabaseAnimComposite::StaticStruct()->ExportText(CopyObj->Content, InItem, nullptr, nullptr, PPF_None, nullptr);
+	}
+	else if (UAnimMontage::StaticClass() == InItem->GetAnimationAssetStaticClass())
+	{
+		FPoseSearchDatabaseAnimMontage::StaticStruct()->ExportText(CopyObj->Content, InItem, nullptr, nullptr, PPF_None, nullptr);
+	}
+	else if (UBlendSpace::StaticClass() == InItem->GetAnimationAssetStaticClass())
+	{
+		FPoseSearchDatabaseBlendSpace::StaticStruct()->ExportText(CopyObj->Content, InItem, nullptr, nullptr, PPF_None, nullptr);
+	}
+	else if (UPoseSearchMultiSequence::StaticClass() == InItem->GetAnimationAssetStaticClass())
+	{
+		FPoseSearchDatabaseMultiSequence::StaticStruct()->ExportText(CopyObj->Content, InItem, nullptr, nullptr, PPF_None, nullptr);
+	}
+
+	DatabaseItems.Add(CopyObj);
+}
+
+void UPoseSearchDatabaseEditorClipboardContent::CopyToClipboard()
+{
+	// Clear the mark state for saving.
+	UnMarkAllObjects(EObjectMark(OBJECTMARK_TagExp | OBJECTMARK_TagImp));
+	
+	// Export the clipboard to text.
+	FStringOutputDevice Archive;
+	const FExportObjectInnerContext Context;
+	UExporter::ExportToOutputDevice(&Context, this, nullptr, Archive, TEXT("copy"), 0, PPF_ExportsNotFullyQualified | PPF_Copy | PPF_Delimited, false, this->GetOuter());
+	FPlatformApplicationMisc::ClipboardCopy(*Archive);
+}
+
+void UPoseSearchDatabaseEditorClipboardContent::PasteToDatabase(UPoseSearchDatabase* InTargetDatabase) const
+{
+	check(InTargetDatabase != nullptr)
+	
+	InTargetDatabase->Modify();
+	
+	for (TObjectPtr<UPoseSearchDatabaseItemCopyObject> Item : DatabaseItems)
+	{
+		if (Item->ClassName == UAnimSequence::StaticClass()->GetFName())
+		{
+			FPoseSearchDatabaseSequence NewAsset;
+			FPoseSearchDatabaseSequence::StaticStruct()->ImportText(*Item->Content, &NewAsset, nullptr, PPF_None, GLog, FPoseSearchDatabaseSequence::StaticStruct()->GetName());
+
+			InTargetDatabase->AddAnimationAsset(FInstancedStruct::Make(NewAsset));
+		}
+		else if (Item->ClassName == UAnimComposite::StaticClass()->GetFName())
+		{
+			FPoseSearchDatabaseAnimComposite NewAsset;
+			FPoseSearchDatabaseAnimComposite::StaticStruct()->ImportText(*Item->Content, &NewAsset, nullptr, PPF_None, GLog, FPoseSearchDatabaseAnimComposite::StaticStruct()->GetName());
+
+			InTargetDatabase->AddAnimationAsset(FInstancedStruct::Make(NewAsset));
+		}
+		else if (Item->ClassName == UBlendSpace::StaticClass()->GetFName())
+		{
+			FPoseSearchDatabaseBlendSpace NewAsset;
+			FPoseSearchDatabaseBlendSpace::StaticStruct()->ImportText(*Item->Content, &NewAsset, nullptr, PPF_None, GLog, FPoseSearchDatabaseBlendSpace::StaticStruct()->GetName());
+
+			InTargetDatabase->AddAnimationAsset(FInstancedStruct::Make(NewAsset));
+		}
+		else if (Item->ClassName == UAnimMontage::StaticClass()->GetFName())
+		{
+			FPoseSearchDatabaseAnimMontage NewAsset;
+			FPoseSearchDatabaseAnimMontage::StaticStruct()->ImportText(*Item->Content, &NewAsset, nullptr, PPF_None, GLog, FPoseSearchDatabaseAnimMontage::StaticStruct()->GetName());
+
+			InTargetDatabase->AddAnimationAsset(FInstancedStruct::Make(NewAsset));
+		}
+		else if (Item->ClassName == UPoseSearchMultiSequence::StaticClass()->GetFName())
+		{
+			FPoseSearchDatabaseMultiSequence NewAsset;
+			FPoseSearchDatabaseMultiSequence::StaticStruct()->ImportText(*Item->Content, &NewAsset, nullptr, PPF_None, GLog, FPoseSearchDatabaseMultiSequence::StaticStruct()->GetName());
+				
+			InTargetDatabase->AddAnimationAsset(FInstancedStruct::Make(NewAsset));
+		}
+	}
+}
+
+class FPoseSearchDatabaseEditorClipboardContentTextFactory : public FCustomizableTextObjectFactory
+{
+public:
+	
+	FPoseSearchDatabaseEditorClipboardContentTextFactory()
+		: FCustomizableTextObjectFactory(GWarn)
+		, ClipboardContent(nullptr) 
+	{
+	}
+
+	UPoseSearchDatabaseEditorClipboardContent* ClipboardContent;
+
+protected:
+
+	virtual bool CanCreateClass(UClass* InObjectClass, bool& bOmitSubObjs) const override
+	{
+		if (InObjectClass->IsChildOf(UPoseSearchDatabaseEditorClipboardContent::StaticClass()))
+		{
+			return true;
+		}
+		return false;
+	}
+	
+	virtual void ProcessConstructedObject(UObject* CreatedObject) override
+	{
+		if (CreatedObject->IsA<UPoseSearchDatabaseEditorClipboardContent>())
+		{
+			ClipboardContent = CastChecked<UPoseSearchDatabaseEditorClipboardContent>(CreatedObject);
+		}
+	}
+};
+
+bool UPoseSearchDatabaseEditorClipboardContent::CanPasteContentFromClipboard(const FString& InTextToImport)
+{
+	const FPoseSearchDatabaseEditorClipboardContentTextFactory ClipboardContentFactory;
+	return ClipboardContentFactory.CanCreateObjectsFromText(InTextToImport);
+}
+
+UPoseSearchDatabaseEditorClipboardContent* UPoseSearchDatabaseEditorClipboardContent::CreateFromClipboard()
+{
+	// Get the text from the clipboard.
+	FString ClipboardText;
+	FPlatformApplicationMisc::ClipboardPaste(ClipboardText);
+
+	// Try to create clipboard content from text.
+	if (CanPasteContentFromClipboard(ClipboardText))
+	{
+		FPoseSearchDatabaseEditorClipboardContentTextFactory ClipboardContentFactory;
+		ClipboardContentFactory.ProcessBuffer((UObject*)GetTransientPackage(), RF_Transactional, ClipboardText);
+		return ClipboardContentFactory.ClipboardContent;
+	}
+
+	return nullptr;
+}
+
+UPoseSearchDatabaseEditorClipboardContent* UPoseSearchDatabaseEditorClipboardContent::Create()
+{
+	return NewObject<UPoseSearchDatabaseEditorClipboardContent>((UObject*)GetTransientPackage());
+}
diff --git a/Engine/Plugins/Animation/PoseSearch/Source/Editor/Private/PoseSearchDatabaseEditorClipboard.h b/Engine/Plugins/Animation/PoseSearch/Source/Editor/Private/PoseSearchDatabaseEditorClipboard.h
new file mode 100644
--- /dev/null
+++ b/Engine/Plugins/Animation/PoseSearch/Source/Editor/Private/PoseSearchDatabaseEditorClipboard.h
@@ -0,0 +1,54 @@
+﻿// Copyright Epic Games, Inc. All Rights Reserved.
+
+#pragma once
+
+#include "CoreMinimal.h"
+#include "PoseSearch/PoseSearchDatabase.h"
+#include "PoseSearchDatabaseEditorClipboard.generated.h"
+
+/** Object used to export FPoseSearchDatabaseAnimationAssetBase and derived classes to clipboard */
+UCLASS(Transient)
+class UPoseSearchDatabaseItemCopyObject : public UObject
+{
+	GENERATED_BODY()
+
+public:
+	
+	UPROPERTY()
+	FName ClassName = NAME_None;
+
+	/** Store's exported database asset as a string */
+	UPROPERTY()
+	FString Content;
+};
+
+UCLASS(Transient)
+class UPoseSearchDatabaseEditorClipboardContent : public UObject
+{
+	GENERATED_BODY()
+
+public:
+	
+	UPROPERTY(Export)
+	TArray<TObjectPtr<UPoseSearchDatabaseItemCopyObject>> DatabaseItems;
+
+	// Utility functions
+
+	/** Store copy of database item. */
+	void CopyDatabaseItem(const FPoseSearchDatabaseAnimationAssetBase* InItem);
+
+	/** Copy contents to system's clipboard */
+	void CopyToClipboard();
+
+	/** Paste all items from clipboard to a given pose search database. Note that this will mark the target database as dirty. */
+	void PasteToDatabase(UPoseSearchDatabase* InTargetDatabase) const;
+
+	/** Determine if a pose search database clipboard content can be constructed from system's clipboard. */
+	static bool CanPasteContentFromClipboard(const FString & InTextToImport);
+
+	/** Attempt to construct a pose search database clipboard content from system's clipboard. Will return nullptr if it fails to do so. */
+	static UPoseSearchDatabaseEditorClipboardContent* CreateFromClipboard();
+
+	/** Create new transient clipboard content object  */
+	static UPoseSearchDatabaseEditorClipboardContent* Create();
+};
\ No newline at end of file
diff --git a/Engine/Plugins/Animation/PoseSearch/Source/Editor/Private/PoseSearchDatabaseViewModel.cpp b/Engine/Plugins/Animation/PoseSearch/Source/Editor/Private/PoseSearchDatabaseViewModel.cpp
--- a/Engine/Plugins/Animation/PoseSearch/Source/Editor/Private/PoseSearchDatabaseViewModel.cpp
+++ b/Engine/Plugins/Animation/PoseSearch/Source/Editor/Private/PoseSearchDatabaseViewModel.cpp
@@ -550,6 +550,8 @@ void FDatabaseViewModel::AddSequenceToDatabase(UAnimSequence* AnimSequence)
 {
 	if (UPoseSearchDatabase* Database = GetPoseSearchDatabase())
 	{
+		Database->Modify();
+		
 		FPoseSearchDatabaseSequence NewAsset;
 		NewAsset.Sequence = AnimSequence;
 		Database->AddAnimationAsset(FInstancedStruct::Make(NewAsset));
@@ -560,6 +562,8 @@ void FDatabaseViewModel::AddBlendSpaceToDatabase(UBlendSpace* BlendSpace)
 {
 	if (UPoseSearchDatabase* Database = GetPoseSearchDatabase())
 	{
+		Database->Modify();
+		
 		FPoseSearchDatabaseBlendSpace NewAsset;
 		NewAsset.BlendSpace = BlendSpace;
 		Database->AddAnimationAsset(FInstancedStruct::Make(NewAsset));
@@ -580,6 +584,8 @@ void FDatabaseViewModel::AddAnimMontageToDatabase(UAnimMontage* AnimMontage)
 {
 	if (UPoseSearchDatabase* Database = GetPoseSearchDatabase())
 	{
+		Database->Modify();
+		
 		FPoseSearchDatabaseAnimMontage NewAsset;
 		NewAsset.AnimMontage = AnimMontage;
 		Database->AddAnimationAsset(FInstancedStruct::Make(NewAsset));
@@ -590,6 +596,8 @@ void FDatabaseViewModel::AddMultiSequenceToDatabase()
 {
 	if (UPoseSearchDatabase* Database = GetPoseSearchDatabase())
 	{
+		Database->Modify();
+		
 		FPoseSearchDatabaseMultiSequence NewAsset;
 		Database->AddAnimationAsset(FInstancedStruct::Make(NewAsset));
 	}
@@ -601,6 +609,8 @@ bool FDatabaseViewModel::DeleteFromDatabase(int32 AnimationAssetIndex)
 	{
 		if (const FPoseSearchDatabaseAnimationAssetBase* DatabaseAnimationAssetBase = Database->GetAnimationAssetBase(AnimationAssetIndex))
 		{
+			Database->Modify();
+			
 			if (UAnimSequenceBase* AnimSequenceBase = Cast<UAnimSequenceBase>(DatabaseAnimationAssetBase->GetAnimationAsset()))
 			{
 				bool bModified = false;
@@ -611,8 +621,13 @@ bool FDatabaseViewModel::DeleteFromDatabase(int32 AnimationAssetIndex)
 					{
 						if (PoseSearchBranchIn->Database == Database)
 						{
+							if (!bModified)
+							{
+								AnimSequenceBase->Modify();
+								bModified = true;
+							}
+
 							AnimSequenceBase->Notifies.RemoveAt(NotifyIndex);
-							bModified = true;
 						}
 					}
 				}
@@ -620,12 +635,10 @@ bool FDatabaseViewModel::DeleteFromDatabase(int32 AnimationAssetIndex)
 				if (bModified)
 				{
 					AnimSequenceBase->RefreshCacheData();	
-					AnimSequenceBase->Modify();
 				}
 			}
 
 			Database->RemoveAnimationAssetAt(AnimationAssetIndex);
-			Database->Modify();
 
 			return true;
 		}
@@ -640,6 +653,8 @@ void FDatabaseViewModel::SetDisableReselection(int32 AnimationAssetIndex, bool b
 	{
 		if (FPoseSearchDatabaseAnimationAssetBase* DatabaseAnimationAsset = Database->GetMutableAnimationAssetBase(AnimationAssetIndex))
 		{
+			Database->Modify();
+			
 			DatabaseAnimationAsset->SetDisableReselection(bEnabled);
 		}
 	}
diff --git a/Engine/Plugins/Animation/PoseSearch/Source/Editor/Public/PoseSearchDatabaseAssetTree.h b/Engine/Plugins/Animation/PoseSearch/Source/Editor/Public/PoseSearchDatabaseAssetTree.h
--- a/Engine/Plugins/Animation/PoseSearch/Source/Editor/Public/PoseSearchDatabaseAssetTree.h
+++ b/Engine/Plugins/Animation/PoseSearch/Source/Editor/Public/PoseSearchDatabaseAssetTree.h
@@ -97,7 +97,19 @@ namespace UE::PoseSearch
 		/** Removes existing selected component nodes from the tree*/
 		bool CanDeleteNodes() const;
 		void OnDeleteNodes();
+		
+		/** Copy selected nodes to clipboard */
+		void OnCopySelectedNodesToClipboard() const;
+		bool CanCopyToClipboard() const;
+
+		/** Paste nodes from clipboard. Adds or overwrites curves (if identifiers collide) */
+		void OnPasteNodesFromClipboard();
+		bool CanPasteFromClipboard();
 
+		/** Cut selected nodes to clipboard */
+		void OnCutSelectedNodesToClipboard();
+		bool CanCutToClipboard() const;
+		
 		void EnableSelectedNodes(bool bIsEnabled);
 		void OnEnableNodes() { EnableSelectedNodes(true); }
 		void OnDisableNodes() { EnableSelectedNodes(false); }
diff --git a/Engine/Plugins/Animation/PoseSearch/Source/Runtime/Public/PoseSearch/PoseSearchMultiSequence.h b/Engine/Plugins/Animation/PoseSearch/Source/Runtime/Public/PoseSearch/PoseSearchMultiSequence.h
--- a/Engine/Plugins/Animation/PoseSearch/Source/Runtime/Public/PoseSearch/PoseSearchMultiSequence.h
+++ b/Engine/Plugins/Animation/PoseSearch/Source/Runtime/Public/PoseSearch/PoseSearchMultiSequence.h
@@ -24,7 +24,7 @@ struct POSESEARCH_API FPoseSearchMultiSequenceItem
 };
 
 UCLASS(Experimental, BlueprintType, Category = "Animation|Pose Search")
-class UPoseSearchMultiSequence : public UDataAsset
+class POSESEARCH_API UPoseSearchMultiSequence : public UDataAsset
 {
 	GENERATED_BODY()
 
