Motion Matching Editor - Sampling slider does not trigger a index rebuilding until value is commited. Additionally, it fixes an ensure() being hit right after changing the value using slider for the first time.

#rnx none

diff --git a/Engine/Plugins/Animation/PoseSearch/Source/Editor/Private/PoseSearchDatabaseAssetTree.cpp b/Engine/Plugins/Animation/PoseSearch/Source/Editor/Private/PoseSearchDatabaseAssetTree.cpp
--- a/Engine/Plugins/Animation/PoseSearch/Source/Editor/Private/PoseSearchDatabaseAssetTree.cpp
+++ b/Engine/Plugins/Animation/PoseSearch/Source/Editor/Private/PoseSearchDatabaseAssetTree.cpp
@@ -925,9 +925,13 @@ namespace UE::PoseSearch
 		}
 	}
 
-	void SDatabaseAssetTree::FinalizeTreeChanges(bool bRecoverSelection)
+	void SDatabaseAssetTree::FinalizeTreeChanges(bool bRecoverSelection, bool bRefreshView)
 	{
-		RefreshTreeView(false, bRecoverSelection);
+		if (bRefreshView)
+		{
+			RefreshTreeView(false, bRecoverSelection);
+		}
+
 		EditorViewModel.Pin()->BuildSearchIndex();
 	}
 
diff --git a/Engine/Plugins/Animation/PoseSearch/Source/Editor/Private/PoseSearchDatabaseEditorReflection.cpp b/Engine/Plugins/Animation/PoseSearch/Source/Editor/Private/PoseSearchDatabaseEditorReflection.cpp
--- a/Engine/Plugins/Animation/PoseSearch/Source/Editor/Private/PoseSearchDatabaseEditorReflection.cpp
+++ b/Engine/Plugins/Animation/PoseSearch/Source/Editor/Private/PoseSearchDatabaseEditorReflection.cpp
@@ -5,6 +5,7 @@
 #include "PoseSearchDatabaseAssetTree.h"
 #include "PoseSearchDatabaseAssetTreeNode.h"
 #include "PoseSearchDatabaseViewModel.h"
+#include "PoseSearch/PoseSearchDerivedData.h"
 #include "PoseSearch/PoseSearchSchema.h"
 
 #define LOCTEXT_NAMESPACE "UPoseSearchDatabaseReflection"
@@ -19,10 +20,26 @@ void UPoseSearchDatabaseReflectionBase::SetSourceLink(
 	AssetTreeWidget = InAssetTreeWidget;
 }
 
-void UPoseSearchDatabaseSequenceReflection::PostEditChangeProperty(struct FPropertyChangedEvent& PropertyChangedEvent)
+void UPoseSearchDatabaseReflectionBase::PostEditChangeProperty(FPropertyChangedEvent& PropertyChangedEvent)
 {
 	Super::PostEditChangeProperty(PropertyChangedEvent);
 
+	// Skip changes during EPropertyChangeType::Interactive since they are always followed by a PostEditChangeProperty() call
+	// with EPropertyChangeType::ValueSet holding the final values.
+	if (PropertyChangedEvent.ChangeType == EPropertyChangeType::Interactive)
+	{
+		return;
+	}
+	
+	if (ApplyChanges())
+	{
+		const bool bShouldRefreshView = (PropertyChangedEvent.Property == nullptr) || !PropertyChangedEvent.Property->IsA(FFloatProperty::StaticClass());
+		AssetTreeWidget->FinalizeTreeChanges(true, bShouldRefreshView);
+	}
+}
+
+bool UPoseSearchDatabaseSequenceReflection::ApplyChanges() const 
+{
 	if (const TSharedPtr<UE::PoseSearch::FDatabaseViewModel> ViewModel = WeakAssetTreeNode.Pin()->EditorViewModel.Pin())
 	{
 		UPoseSearchDatabase* Database = ViewModel->GetPoseSearchDatabase();
@@ -33,80 +50,89 @@ void UPoseSearchDatabaseSequenceReflection::PostEditChangeProperty(struct FPrope
 			{
 				*DatabaseSequence = Sequence;
 				Database->MarkPackageDirty();
-		
-				AssetTreeWidget->FinalizeTreeChanges(true);
+
+				return true;
 			}
 		}
 	}
+
+	return false;
 }
 
-void UPoseSearchDatabaseBlendSpaceReflection::PostEditChangeProperty(struct FPropertyChangedEvent& PropertyChangedEvent)
+bool UPoseSearchDatabaseBlendSpaceReflection::ApplyChanges() const
 {
-	Super::PostEditChangeProperty(PropertyChangedEvent);
-
 	if (const TSharedPtr<UE::PoseSearch::FDatabaseViewModel> ViewModel = WeakAssetTreeNode.Pin()->EditorViewModel.Pin())
 	{
 		UPoseSearchDatabase* Database = ViewModel->GetPoseSearchDatabase();
 		if (IsValid(Database))
 		{
-			FInstancedStruct& DatabaseAsset = Database->GetMutableAnimationAssetStruct(WeakAssetTreeNode.Pin()->SourceAssetIdx);
-			if (FPoseSearchDatabaseBlendSpace* DatabaseBlendSpace = DatabaseAsset.GetMutablePtr<FPoseSearchDatabaseBlendSpace>())
+			if (UE::PoseSearch::EAsyncBuildIndexResult::Success == UE::PoseSearch::FAsyncPoseSearchDatabasesManagement::RequestAsyncBuildIndex(Database, UE::PoseSearch::ERequestAsyncBuildFlag::ContinueRequest))
 			{
-				*DatabaseBlendSpace = BlendSpace;
-				Database->MarkPackageDirty();
+				FInstancedStruct& DatabaseAsset = Database->GetMutableAnimationAssetStruct(WeakAssetTreeNode.Pin()->SourceAssetIdx);
+				if (FPoseSearchDatabaseBlendSpace* DatabaseBlendSpace = DatabaseAsset.GetMutablePtr<FPoseSearchDatabaseBlendSpace>())
+				{
+					*DatabaseBlendSpace = BlendSpace;
+					Database->MarkPackageDirty();
 
-				AssetTreeWidget->FinalizeTreeChanges(true);
+					return true;
+				}
 			}
 		}
 	}
+
+	return false;
 }
 
-void UPoseSearchDatabaseAnimCompositeReflection::PostEditChangeProperty(struct FPropertyChangedEvent& PropertyChangedEvent)
+bool UPoseSearchDatabaseAnimCompositeReflection::ApplyChanges() const
 {
-	Super::PostEditChangeProperty(PropertyChangedEvent);
-
 	if (const TSharedPtr<UE::PoseSearch::FDatabaseViewModel> ViewModel = WeakAssetTreeNode.Pin()->EditorViewModel.Pin())
 	{
 		UPoseSearchDatabase* Database = ViewModel->GetPoseSearchDatabase();
 		if (IsValid(Database))
 		{
-			FInstancedStruct& DatabaseAsset = Database->GetMutableAnimationAssetStruct(WeakAssetTreeNode.Pin()->SourceAssetIdx);
-			if (FPoseSearchDatabaseAnimComposite* DatabaseAnimComposite = DatabaseAsset.GetMutablePtr<FPoseSearchDatabaseAnimComposite>())
+			if (UE::PoseSearch::EAsyncBuildIndexResult::Success == UE::PoseSearch::FAsyncPoseSearchDatabasesManagement::RequestAsyncBuildIndex(Database, UE::PoseSearch::ERequestAsyncBuildFlag::ContinueRequest))
 			{
-				*DatabaseAnimComposite = AnimComposite;
-				Database->MarkPackageDirty();
+				FInstancedStruct& DatabaseAsset = Database->GetMutableAnimationAssetStruct(WeakAssetTreeNode.Pin()->SourceAssetIdx);
+				if (FPoseSearchDatabaseAnimComposite* DatabaseAnimComposite = DatabaseAsset.GetMutablePtr<FPoseSearchDatabaseAnimComposite>())
+				{
+					*DatabaseAnimComposite = AnimComposite;
+					Database->MarkPackageDirty();
 
-				AssetTreeWidget->FinalizeTreeChanges(true);
+					return true;
+				}
 			}
 		}
 	}
+
+	return false;
 }
 
-void UPoseSearchDatabaseAnimMontageReflection::PostEditChangeProperty(struct FPropertyChangedEvent& PropertyChangedEvent)
+bool UPoseSearchDatabaseAnimMontageReflection::ApplyChanges() const
 {
-	Super::PostEditChangeProperty(PropertyChangedEvent);
-
 	if (const TSharedPtr<UE::PoseSearch::FDatabaseViewModel> ViewModel = WeakAssetTreeNode.Pin()->EditorViewModel.Pin())
 	{
 		UPoseSearchDatabase* Database = ViewModel->GetPoseSearchDatabase();
 		if (IsValid(Database))
 		{
-			FInstancedStruct& DatabaseAsset = Database->GetMutableAnimationAssetStruct(WeakAssetTreeNode.Pin()->SourceAssetIdx);
-			if (FPoseSearchDatabaseAnimMontage* DatabaseAnimMontage = DatabaseAsset.GetMutablePtr<FPoseSearchDatabaseAnimMontage>())
+			if (UE::PoseSearch::EAsyncBuildIndexResult::Success == UE::PoseSearch::FAsyncPoseSearchDatabasesManagement::RequestAsyncBuildIndex(Database, UE::PoseSearch::ERequestAsyncBuildFlag::ContinueRequest))
 			{
-				*DatabaseAnimMontage = AnimMontage;
-				Database->MarkPackageDirty();
+				FInstancedStruct& DatabaseAsset = Database->GetMutableAnimationAssetStruct(WeakAssetTreeNode.Pin()->SourceAssetIdx);
+				if (FPoseSearchDatabaseAnimMontage* DatabaseAnimMontage = DatabaseAsset.GetMutablePtr<FPoseSearchDatabaseAnimMontage>())
+				{
+					*DatabaseAnimMontage = AnimMontage;
+					Database->MarkPackageDirty();
 
-				AssetTreeWidget->FinalizeTreeChanges(true);
+					return true;
+				}
 			}
 		}
 	}
+
+	return false;
 }
 
-void UPoseSearchDatabaseMultiSequenceReflection::PostEditChangeProperty(struct FPropertyChangedEvent& PropertyChangedEvent)
+bool UPoseSearchDatabaseMultiSequenceReflection::ApplyChanges() const
 {
-	Super::PostEditChangeProperty(PropertyChangedEvent);
-
 	if (const TSharedPtr<UE::PoseSearch::FDatabaseViewModel> ViewModel = WeakAssetTreeNode.Pin()->EditorViewModel.Pin())
 	{
 		UPoseSearchDatabase* Database = ViewModel->GetPoseSearchDatabase();
@@ -118,10 +144,12 @@ void UPoseSearchDatabaseMultiSequenceReflection::PostEditChangeProperty(struct F
 				*DatabaseMultiSequence = MultiSequence;
 				Database->MarkPackageDirty();
 
-				AssetTreeWidget->FinalizeTreeChanges(true);
+				return true;
 			}
 		}
 	}
+
+	return false;
 }
 
 #endif // WITH_EDITOR
diff --git a/Engine/Plugins/Animation/PoseSearch/Source/Editor/Public/PoseSearchDatabaseAssetTree.h b/Engine/Plugins/Animation/PoseSearch/Source/Editor/Public/PoseSearchDatabaseAssetTree.h
--- a/Engine/Plugins/Animation/PoseSearch/Source/Editor/Public/PoseSearchDatabaseAssetTree.h
+++ b/Engine/Plugins/Animation/PoseSearch/Source/Editor/Public/PoseSearchDatabaseAssetTree.h
@@ -37,7 +37,7 @@ namespace UE::PoseSearch
 		// End FEditorUndoClient interface
 		
 		void RefreshTreeView(bool bIsInitialSetup = false, bool bRecoverSelection = false);
-		void FinalizeTreeChanges(bool bRecoverSelection = false);
+		void FinalizeTreeChanges(bool bRecoverSelection = false, bool bRefreshView = true);
 
 		void SetSelectedItem(int32 SourceAssetIdx, bool bClearSelection);
 		FDetailColumnSizeData& GetColumnSizeData() { return ColumnSizeData; }
diff --git a/Engine/Plugins/Animation/PoseSearch/Source/Editor/Public/PoseSearchDatabaseEditorReflection.h b/Engine/Plugins/Animation/PoseSearch/Source/Editor/Public/PoseSearchDatabaseEditorReflection.h
--- a/Engine/Plugins/Animation/PoseSearch/Source/Editor/Public/PoseSearchDatabaseEditorReflection.h
+++ b/Engine/Plugins/Animation/PoseSearch/Source/Editor/Public/PoseSearchDatabaseEditorReflection.h
@@ -23,7 +23,13 @@ public:
 	void SetSourceLink(
 		const TWeakPtr<UE::PoseSearch::FDatabaseAssetTreeNode>& InWeakAssetTreeNode,
 		const TSharedPtr<UE::PoseSearch::SDatabaseAssetTree>& InAssetTreeWidget);
+	
+	virtual bool ApplyChanges() const { return false; };
 
+#if WITH_EDITOR
+	virtual void PostEditChangeProperty(FPropertyChangedEvent& PropertyChangedEvent) override;
+#endif
+	
 protected:
 	TWeakPtr<UE::PoseSearch::FDatabaseAssetTreeNode> WeakAssetTreeNode;
 	TSharedPtr<UE::PoseSearch::SDatabaseAssetTree> AssetTreeWidget;
@@ -52,10 +58,8 @@ class UPoseSearchDatabaseSequenceReflection : public UPoseSearchDatabaseReflecti
 public:
 	UPROPERTY(EditAnywhere, Category = "Selected Sequence")
 	FPoseSearchDatabaseSequenceEx Sequence;
-
-#if WITH_EDITOR
-	virtual void PostEditChangeProperty(struct FPropertyChangedEvent& PropertyChangedEvent) override;
-#endif
+	
+	virtual bool ApplyChanges() const override;
 };
 
 USTRUCT()
@@ -81,10 +85,8 @@ class UPoseSearchDatabaseBlendSpaceReflection : public UPoseSearchDatabaseReflec
 public:
 	UPROPERTY(EditAnywhere, Category = "Selected Blend Space")
 	FPoseSearchDatabaseBlendSpaceEx BlendSpace;
-
-#if WITH_EDITOR
-	virtual void PostEditChangeProperty(struct FPropertyChangedEvent& PropertyChangedEvent) override;
-#endif
+	
+	virtual bool ApplyChanges() const override;
 };
 
 USTRUCT()
@@ -110,10 +112,8 @@ class UPoseSearchDatabaseAnimCompositeReflection : public UPoseSearchDatabaseRef
 public:
 	UPROPERTY(EditAnywhere, Category = "Selected Anim Composite")
 	FPoseSearchDatabaseAnimCompositeEx AnimComposite;
-
-#if WITH_EDITOR
-	virtual void PostEditChangeProperty(struct FPropertyChangedEvent& PropertyChangedEvent) override;
-#endif
+	
+	virtual bool ApplyChanges() const override;
 };
 
 USTRUCT()
@@ -139,10 +139,8 @@ class UPoseSearchDatabaseAnimMontageReflection : public UPoseSearchDatabaseRefle
 public:
 	UPROPERTY(EditAnywhere, Category = "Selected Anim Montage")
 	FPoseSearchDatabaseAnimMontageEx AnimMontage;
-
-#if WITH_EDITOR
-	virtual void PostEditChangeProperty(struct FPropertyChangedEvent& PropertyChangedEvent) override;
-#endif
+	
+	virtual bool ApplyChanges() const override;
 };
 
 USTRUCT()
@@ -168,10 +166,8 @@ class UPoseSearchDatabaseMultiSequenceReflection : public UPoseSearchDatabaseRef
 public:
 	UPROPERTY(EditAnywhere, Category = "Selected Sequence")
 	FPoseSearchDatabaseMultiSequenceEx MultiSequence;
-
-#if WITH_EDITOR
-	virtual void PostEditChangeProperty(struct FPropertyChangedEvent& PropertyChangedEvent) override;
-#endif
+	
+	virtual bool ApplyChanges() const override;
 };
 
 USTRUCT()
