motion matching - fix for link error using UE::PoseSearch::FAssetIndexer on a custom PoseSearchFeatureChannel

#changelist validated
#virtualized

diff --git a/Engine/Plugins/Animation/PoseSearch/Source/Runtime/Public/PoseSearch/PoseSearchAssetIndexer.h b/Engine/Plugins/Animation/PoseSearch/Source/Runtime/Public/PoseSearch/PoseSearchAssetIndexer.h
--- a/Engine/Plugins/Animation/PoseSearch/Source/Runtime/Public/PoseSearch/PoseSearchAssetIndexer.h
+++ b/Engine/Plugins/Animation/PoseSearch/Source/Runtime/Public/PoseSearch/PoseSearchAssetIndexer.h
@@ -49,7 +49,7 @@ struct FAnimationAssetSamplers
 	TArray<const FMirrorDataCache*> MirrorDataCaches;
 };
 
-class FAssetIndexer
+class POSESEARCH_API FAssetIndexer
 {
 public:
 	struct FStats
