motion matching - compatible skeleton pose history fix

diff --git a/Engine/Plugins/Animation/PoseSearch/Source/Runtime/Private/PoseSearchHistory.cpp b/Engine/Plugins/Animation/PoseSearch/Source/Runtime/Private/PoseSearchHistory.cpp
--- a/Engine/Plugins/Animation/PoseSearch/Source/Runtime/Private/PoseSearchHistory.cpp
+++ b/Engine/Plugins/Animation/PoseSearch/Source/Runtime/Private/PoseSearchHistory.cpp
@@ -89,9 +89,9 @@ static FComponentSpaceTransformIndex GetRemappedComponentSpaceTransformIndex(con
 	check(BoneIndexType != WorldSpaceIndexType);
 
 	FComponentSpaceTransformIndex BoneTransformIndex = FComponentSpaceTransformIndex(BoneIndexType);
-	if (BoneIndexType != ComponentSpaceIndexType)
+	if (BoneTransformIndex != ComponentSpaceIndexType)
 	{
-		BoneIndexType = GetRemappedBoneIndexType(BoneIndexType, BoneIndexSkeleton, LastUpdateSkeleton);
+		BoneTransformIndex = GetRemappedBoneIndexType(BoneTransformIndex, BoneIndexSkeleton, LastUpdateSkeleton);
 
 		if (!BoneToTransformMap.IsEmpty())
 		{
