motion matching - jasper forward facing model support

diff --git a/Engine/Plugins/Animation/PoseSearch/Source/Runtime/Private/AnimNode_PoseSearchHistoryCollector.cpp b/Engine/Plugins/Animation/PoseSearch/Source/Runtime/Private/AnimNode_PoseSearchHistoryCollector.cpp
--- a/Engine/Plugins/Animation/PoseSearch/Source/Runtime/Private/AnimNode_PoseSearchHistoryCollector.cpp
+++ b/Engine/Plugins/Animation/PoseSearch/Source/Runtime/Private/AnimNode_PoseSearchHistoryCollector.cpp
@@ -83,7 +83,7 @@ void FAnimNode_PoseSearchHistoryCollector_Base::CacheBones_AnyThread(const FAnim
 		Pose.SetBoneContainer(&Context.AnimInstanceProxy->GetRequiredBones());
 		FCSPose<FCompactPose> ComponentSpacePose;
 		ComponentSpacePose.InitPose(Pose);
-		PoseHistory.EvaluateComponentSpace_AnyThread(0.f, ComponentSpacePose, bStoreScales);
+		PoseHistory.EvaluateComponentSpace_AnyThread(0.f, ComponentSpacePose, bStoreScales, RootBoneRecoveryTime);
 	}
 }
 
@@ -141,7 +141,7 @@ void FAnimNode_PoseSearchHistoryCollector::Evaluate_AnyThread(FPoseContext& Outp
 
 	FCSPose<FCompactPose> ComponentSpacePose;
 	ComponentSpacePose.InitPose(Output.Pose);
-	PoseHistory.EvaluateComponentSpace_AnyThread(Output.AnimInstanceProxy->GetDeltaSeconds(), ComponentSpacePose, bStoreScales);
+	PoseHistory.EvaluateComponentSpace_AnyThread(Output.AnimInstanceProxy->GetDeltaSeconds(), ComponentSpacePose, bStoreScales, RootBoneRecoveryTime);
 
 #if ENABLE_DRAW_DEBUG && ENABLE_ANIM_DEBUG
 	FColor Color;
@@ -195,7 +195,7 @@ void FAnimNode_PoseSearchComponentSpaceHistoryCollector::EvaluateComponentSpace_
 	Super::EvaluateComponentSpace_AnyThread(Output);
 	Source.EvaluateComponentSpace(Output);
 
-	PoseHistory.EvaluateComponentSpace_AnyThread(Output.AnimInstanceProxy->GetDeltaSeconds(), Output.Pose, bStoreScales);
+	PoseHistory.EvaluateComponentSpace_AnyThread(Output.AnimInstanceProxy->GetDeltaSeconds(), Output.Pose, bStoreScales, RootBoneRecoveryTime);
 
 #if ENABLE_DRAW_DEBUG && ENABLE_ANIM_DEBUG
 	FColor Color;
diff --git a/Engine/Plugins/Animation/PoseSearch/Source/Runtime/Private/PoseSearchContext.cpp b/Engine/Plugins/Animation/PoseSearch/Source/Runtime/Private/PoseSearchContext.cpp
--- a/Engine/Plugins/Animation/PoseSearch/Source/Runtime/Private/PoseSearchContext.cpp
+++ b/Engine/Plugins/Animation/PoseSearch/Source/Runtime/Private/PoseSearchContext.cpp
@@ -55,15 +55,6 @@ const UPoseSearchSchema* FDebugDrawParams::GetSchema() const
 	return Database ? Database->Schema : nullptr;
 }
 
-FVector FDebugDrawParams::ExtractPosition(TConstArrayView<float> PoseVector, const UPoseSearchFeatureChannel_Position* Position) const
-{
-	check(Position);
-	check(Position->SampleRole == Position->OriginRole);
-	const FVector BonePosition = FFeatureVectorHelper::DecodeVector(PoseVector, Position->GetChannelDataOffset(), Position->ComponentStripping);
-	const FVector WorldBonePosition = GetRootTransform(Position->SampleRole).TransformPosition(BonePosition);
-	return WorldBonePosition;
-}
-
 FVector FDebugDrawParams::ExtractPosition(TConstArrayView<float> PoseVector, float SampleTimeOffset, int8 SchemaBoneIdx, const FRole& Role, EPermutationTimeType PermutationTimeType, int32 SamplingAttributeId) const
 {
 	// we don't wanna ask for a SchemaOriginBoneIdx in the future or past
@@ -92,7 +83,9 @@ FVector FDebugDrawParams::ExtractPosition(TConstArrayView<float> PoseVector, flo
 					return nullptr;
 				})))
 		{
-			return ExtractPosition(PoseVector, FoundPosition);
+			const FVector BonePosition = FFeatureVectorHelper::DecodeVector(PoseVector, FoundPosition->GetChannelDataOffset(), FoundPosition->ComponentStripping);
+			const FVector WorldBonePosition = GetRootBoneTransform(FoundPosition->SampleRole).TransformPosition(BonePosition);
+			return WorldBonePosition;
 		}
 
 		if (const int32* RoleIndex = RoleToIndex.Find(Role))
@@ -117,7 +110,7 @@ FVector FDebugDrawParams::ExtractPosition(TConstArrayView<float> PoseVector, flo
 			}
 		}
 	}
-	return GetRootTransform(Role, SampleTimeOffset).GetTranslation();
+	return GetRootBoneTransform(Role, SampleTimeOffset).GetTranslation();
 }
 
 FQuat FDebugDrawParams::ExtractRotation(TConstArrayView<float> PoseVector, float SampleTimeOffset, int8 SchemaBoneIdx, const FRole& Role, EPermutationTimeType PermutationTimeType, int32 SamplingAttributeId) const
@@ -215,7 +208,7 @@ FQuat FDebugDrawParams::ExtractRotation(TConstArrayView<float> PoseVector, float
 				// world rotation associated to the time zero, we can calcualte the world rotation at time SampleTimeOffset
 				const FMatrix RotMatrix(DecodedHeading[int32(EHeadingAxis::X)], DecodedHeading[int32(EHeadingAxis::Y)], DecodedHeading[int32(EHeadingAxis::Z)], FVector::ZeroVector);
 				const FQuat RotQuat(RotMatrix);
-				const FQuat RotQuatWorld = RotQuat * GetRootTransform(Role).GetRotation();
+				const FQuat RotQuatWorld = RotQuat * GetRootBoneTransform(Role).GetRotation();
 				return RotQuatWorld;
 			}
 		}
@@ -243,16 +236,17 @@ FQuat FDebugDrawParams::ExtractRotation(TConstArrayView<float> PoseVector, float
 		}
 	}
 
-	return GetRootTransform(Role, SampleTimeOffset).GetRotation();
+	return GetRootBoneTransform(Role, SampleTimeOffset).GetRotation();
 }
 
-FTransform FDebugDrawParams::GetRootTransform(const FRole& Role, float SampleTimeOffset) const
+FTransform FDebugDrawParams::GetRootBoneTransform(const FRole& Role, float SampleTimeOffset) const
 {
+	FTransform RootBoneTransform = FTransform::Identity;
 	if (const int32* RoleIndex = RoleToIndex.Find(Role))
 	{
-		return PoseHistories[*RoleIndex]->GetTrajectory().GetSampleAtTime(SampleTimeOffset).GetTransform();
+		PoseHistories[*RoleIndex]->GetTransformAtTime(SampleTimeOffset, RootBoneTransform, nullptr, RootBoneIndexType, WorldSpaceIndexType);
 	}
-	return FTransform::Identity;
+	return RootBoneTransform;
 }
 
 void FDebugDrawParams::DrawLine(const FVector& LineStart, const FVector& LineEnd, const FColor& Color, float Thickness) const
diff --git a/Engine/Plugins/Animation/PoseSearch/Source/Runtime/Private/PoseSearchFeatureChannel_Phase.cpp b/Engine/Plugins/Animation/PoseSearch/Source/Runtime/Private/PoseSearchFeatureChannel_Phase.cpp
--- a/Engine/Plugins/Animation/PoseSearch/Source/Runtime/Private/PoseSearchFeatureChannel_Phase.cpp
+++ b/Engine/Plugins/Animation/PoseSearch/Source/Runtime/Private/PoseSearchFeatureChannel_Phase.cpp
@@ -371,9 +371,10 @@ void UPoseSearchFeatureChannel_Phase::DebugDraw(const UE::PoseSearch::FDebugDraw
 	const FVector2D Phase = FFeatureVectorHelper::DecodeVector2D(PoseVector, ChannelDataOffset);
 	const FVector BonePos = DrawParams.ExtractPosition(PoseVector, 0.f, SchemaBoneIdx, SampleRole);
 
-	const FVector TransformXAxisVector = DrawParams.GetRootTransform(SampleRole).TransformVector(FVector::XAxisVector);
-	const FVector TransformYAxisVector = DrawParams.GetRootTransform(SampleRole).TransformVector(FVector::YAxisVector);
-	const FVector TransformZAxisVector = DrawParams.GetRootTransform(SampleRole).TransformVector(FVector::ZAxisVector);
+	const FTransform RootBoneTransform = DrawParams.GetRootBoneTransform(SampleRole);
+	const FVector TransformXAxisVector = RootBoneTransform.TransformVector(FVector::XAxisVector);
+	const FVector TransformYAxisVector = RootBoneTransform.TransformVector(FVector::YAxisVector);
+	const FVector TransformZAxisVector = RootBoneTransform.TransformVector(FVector::ZAxisVector);
 
 	const FVector PhaseVector = (TransformZAxisVector * Phase.X + TransformYAxisVector * Phase.Y) * ScaleFactor;
 	DrawParams.DrawLine(BonePos, BonePos + PhaseVector, Color);
diff --git a/Engine/Plugins/Animation/PoseSearch/Source/Runtime/Private/PoseSearchFeatureChannel_Position.cpp b/Engine/Plugins/Animation/PoseSearch/Source/Runtime/Private/PoseSearchFeatureChannel_Position.cpp
--- a/Engine/Plugins/Animation/PoseSearch/Source/Runtime/Private/PoseSearchFeatureChannel_Position.cpp
+++ b/Engine/Plugins/Animation/PoseSearch/Source/Runtime/Private/PoseSearchFeatureChannel_Position.cpp
@@ -169,7 +169,7 @@ void UPoseSearchFeatureChannel_Position::DebugDraw(const UE::PoseSearch::FDebugD
 	const FVector FeaturesVector = FFeatureVectorHelper::DecodeVector(PoseVector, ChannelDataOffset, ComponentStripping);
 	const EPermutationTimeType TimeType = PermutationTimeType == EPermutationTimeType::UsePermutationTime ? EPermutationTimeType::UseSampleToPermutationTime : EPermutationTimeType::UseSampleTime;
 	const FVector OriginBonePos = DrawParams.ExtractPosition(PoseVector, OriginTimeOffset, SchemaOriginBoneIdx, OriginRole, TimeType);
-	const FVector DeltaPos = DrawParams.GetRootTransform(OriginRole).TransformVector(FeaturesVector);
+	const FVector DeltaPos = DrawParams.GetRootBoneTransform(OriginRole).TransformVector(FeaturesVector);
 	const FVector BonePos = OriginBonePos + DeltaPos;
 	DrawParams.DrawPoint(BonePos, Color);
 
diff --git a/Engine/Plugins/Animation/PoseSearch/Source/Runtime/Private/PoseSearchFeatureChannel_Velocity.cpp b/Engine/Plugins/Animation/PoseSearch/Source/Runtime/Private/PoseSearchFeatureChannel_Velocity.cpp
--- a/Engine/Plugins/Animation/PoseSearch/Source/Runtime/Private/PoseSearchFeatureChannel_Velocity.cpp
+++ b/Engine/Plugins/Animation/PoseSearch/Source/Runtime/Private/PoseSearchFeatureChannel_Velocity.cpp
@@ -143,7 +143,7 @@ void UPoseSearchFeatureChannel_Velocity::DebugDraw(const UE::PoseSearch::FDebugD
 
 	const float LinearVelocityScale = bNormalize ? 15.f : 0.08f;
 
-	const FVector LinearVelocity = DrawParams.GetRootTransform(SampleRole).TransformVector(FFeatureVectorHelper::DecodeVector(PoseVector, ChannelDataOffset, ComponentStripping));
+	const FVector LinearVelocity = DrawParams.GetRootBoneTransform(SampleRole).TransformVector(FFeatureVectorHelper::DecodeVector(PoseVector, ChannelDataOffset, ComponentStripping));
 	const FVector BoneVelDirection = LinearVelocity.GetSafeNormal();
 	const FVector BonePos = DrawParams.ExtractPosition(PoseVector, SampleTimeOffset, SchemaBoneIdx, SampleRole, PermutationTimeType, SamplingAttributeId);
 
diff --git a/Engine/Plugins/Animation/PoseSearch/Source/Runtime/Private/PoseSearchHistory.cpp b/Engine/Plugins/Animation/PoseSearch/Source/Runtime/Private/PoseSearchHistory.cpp
--- a/Engine/Plugins/Animation/PoseSearch/Source/Runtime/Private/PoseSearchHistory.cpp
+++ b/Engine/Plugins/Animation/PoseSearch/Source/Runtime/Private/PoseSearchHistory.cpp
@@ -129,17 +129,39 @@ static bool LerpEntries(float Time, bool bExtrapolate, const FPoseHistoryEntry&
 	{
 		if (ReferenceBoneTransformIndex == ComponentSpaceIndexType)
 		{
-			OutBoneTransform.Blend(
-				PrevEntry.GetComponentSpaceTransform(BoneTransformIndex),
-				NextEntry.GetComponentSpaceTransform(BoneTransformIndex),
-				LerpValue);
+			if (FMath::IsNearlyZero(LerpValue, ZERO_ANIMWEIGHT_THRESH))
+			{
+				OutBoneTransform = PrevEntry.GetComponentSpaceTransform(BoneTransformIndex);
+			}
+			else if (FMath::IsNearlyZero(LerpValue - 1.f, ZERO_ANIMWEIGHT_THRESH))
+			{
+				OutBoneTransform = NextEntry.GetComponentSpaceTransform(BoneTransformIndex);
+			}
+			else
+			{
+				OutBoneTransform.Blend(
+					PrevEntry.GetComponentSpaceTransform(BoneTransformIndex),
+					NextEntry.GetComponentSpaceTransform(BoneTransformIndex),
+					LerpValue);
+			}
 		}
 		else
 		{
-			OutBoneTransform.Blend(
-				PrevEntry.GetComponentSpaceTransform(BoneTransformIndex) * PrevEntry.GetComponentSpaceTransform(ReferenceBoneTransformIndex).Inverse(),
-				NextEntry.GetComponentSpaceTransform(BoneTransformIndex) * NextEntry.GetComponentSpaceTransform(ReferenceBoneTransformIndex).Inverse(),
-				LerpValue);
+			if (FMath::IsNearlyZero(LerpValue, ZERO_ANIMWEIGHT_THRESH))
+			{
+				OutBoneTransform = PrevEntry.GetComponentSpaceTransform(BoneTransformIndex) * PrevEntry.GetComponentSpaceTransform(ReferenceBoneTransformIndex).Inverse();
+			}
+			else if (FMath::IsNearlyZero(LerpValue - 1.f, ZERO_ANIMWEIGHT_THRESH))
+			{
+				OutBoneTransform = NextEntry.GetComponentSpaceTransform(BoneTransformIndex) * NextEntry.GetComponentSpaceTransform(ReferenceBoneTransformIndex).Inverse();
+			}
+			else
+			{
+				OutBoneTransform.Blend(
+					PrevEntry.GetComponentSpaceTransform(BoneTransformIndex) * PrevEntry.GetComponentSpaceTransform(ReferenceBoneTransformIndex).Inverse(),
+					NextEntry.GetComponentSpaceTransform(BoneTransformIndex) * NextEntry.GetComponentSpaceTransform(ReferenceBoneTransformIndex).Inverse(),
+					LerpValue);
+			}
 		}
 	}
 	else
@@ -361,6 +383,8 @@ void FArchivedPoseHistory::DebugDraw(const UWorld* World, FColor Color) const
 
 			PrevGlobalTransforms.SetNum(Max, EAllowShrinking::No);
 
+			const bool bIsCurrentTimeEntry = FMath::IsNearlyZero(Entry.AccumulatedSeconds);
+
 			for (int32 i = 0; i < Entry.Num(); ++i)
 			{
 				const FTransform RootTransform = Trajectory.GetSampleAtTime(Entry.AccumulatedSeconds).GetTransform();
@@ -371,6 +395,17 @@ void FArchivedPoseHistory::DebugDraw(const UWorld* World, FColor Color) const
 					DrawDebugLine(World, PrevGlobalTransforms[i].GetTranslation(), GlobalTransforms.GetTranslation(), Color, false, -1.f, SDPG_Foreground);
 				}
 
+				if (bIsCurrentTimeEntry)
+				{
+					DrawDebugPoint(World, GlobalTransforms.GetTranslation(), 6.f, Color, false, -1.f, SDPG_Foreground);
+
+					if (i == 0)
+					{
+						DrawDebugLine(World, GlobalTransforms.GetTranslation(), GlobalTransforms.GetTranslation() + RootTransform.GetUnitAxis(EAxis::Type::X) * 25.f, FColor::Black, false, -1.f, SDPG_Foreground);
+						DrawDebugLine(World, GlobalTransforms.GetTranslation(), GlobalTransforms.GetTranslation() + GlobalTransforms.GetUnitAxis(EAxis::Type::X) * 20.f, FColor::White, false, -1.f, SDPG_Foreground);
+					}
+				}
+
 				PrevGlobalTransforms[i] = GlobalTransforms;
 			}
 		}
@@ -423,7 +458,7 @@ FPoseHistory& FPoseHistory::operator=(FPoseHistory&& Other)
 void FPoseHistory::Initialize_AnyThread(int32 InNumPoses, float InSamplingInterval)
 {
 	CheckThreadSafetyWrite(WriteDataThreadSafeCounter);
-	check(InNumPoses >= 2 && InSamplingInterval > UE_KINDA_SMALL_NUMBER);
+	check(InNumPoses >= 2);
 
 	MaxNumPoses = InNumPoses;
 	SamplingInterval = InSamplingInterval;
@@ -433,7 +468,7 @@ void FPoseHistory::CacheBones_AnyThread(const TArray<FBoneIndexType>& RequiredBo
 {
 	CheckThreadSafetyWrite(WriteDataThreadSafeCounter);
 
-	check(MaxNumPoses >= 2 && SamplingInterval > UE_KINDA_SMALL_NUMBER);
+	check(MaxNumPoses >= 2);
 
 	WriteData.BoneToTransformMap.Reset();
 	if (!RequiredBones.IsEmpty())
@@ -556,9 +591,9 @@ void FPoseHistory::PreUpdate(const UAnimInstance* AnimInstance, float DeltaTime,
 		// @todo: Synchronize the FPoseSearchQueryTrajectorySample::AccumulatedSeconds of the generated trajectory with the FPoseHistoryEntry::AccumulatedSeconds of the captured poses
 		FPoseSearchTrajectoryData::FDerived TrajectoryDataDerived;
 		TrajectoryData.UpdateData(DeltaTime, AnimInstance, TrajectoryDataDerived, TrajectoryDataState);
-		UPoseSearchTrajectoryLibrary::InitTrajectorySamples(Trajectory, TrajectoryData, TrajectoryDataDerived, TrajectoryDataSampling);
+		UPoseSearchTrajectoryLibrary::InitTrajectorySamples(Trajectory, TrajectoryData, TrajectoryDataDerived, TrajectoryDataSampling, DeltaTime);
 		UPoseSearchTrajectoryLibrary::UpdateHistory_TransformHistory(Trajectory, TrajectoryData, TrajectoryDataDerived, TrajectoryDataSampling, DeltaTime);
-		UPoseSearchTrajectoryLibrary::UpdatePrediction_SimulateCharacterMovement(Trajectory, TrajectoryData, TrajectoryDataDerived, TrajectoryDataSampling);
+		UPoseSearchTrajectoryLibrary::UpdatePrediction_SimulateCharacterMovement(Trajectory, TrajectoryData, TrajectoryDataDerived, TrajectoryDataSampling, DeltaTime);
 
 		// @todo: support TrajectorySpeedMultiplier
 		//TrajectorySpeedMultiplier = 1.f;
@@ -589,13 +624,15 @@ void FPoseHistory::SetTrajectory(const FPoseSearchQueryTrajectory& InTrajectory,
 	}
 }
 
-void FPoseHistory::EvaluateComponentSpace_AnyThread(float DeltaTime, FCSPose<FCompactPose>& ComponentSpacePose, bool bStoreScales)
+void FPoseHistory::EvaluateComponentSpace_AnyThread(float DeltaTime, FCSPose<FCompactPose>& ComponentSpacePose, bool bStoreScales, float RootBoneRecoveryTime)
 {
 	CheckThreadSafetyWrite(WriteDataThreadSafeCounter);
 
-	check(MaxNumPoses >= 2 && SamplingInterval > UE_KINDA_SMALL_NUMBER);
+	check(MaxNumPoses >= 2);
 
 	const USkeleton* Skeleton = ComponentSpacePose.GetPose().GetBoneContainer().GetSkeletonAsset();
+	check(Skeleton);
+
 	if (WriteData.LastUpdateSkeleton != Skeleton)
 	{
 		// @todo: support a different USkeleton per FPoseHistoryEntry if required
@@ -603,6 +640,14 @@ void FPoseHistory::EvaluateComponentSpace_AnyThread(float DeltaTime, FCSPose<FCo
 		WriteData.LastUpdateSkeleton = Skeleton;
 	}
 
+	FPoseHistoryEntry FutureEntryTemp;
+	if (!WriteData.Entries.IsEmpty() && WriteData.Entries.Last().AccumulatedSeconds > 0.f)
+	{
+		// removing the "future" root bone Entry
+		FutureEntryTemp = MoveTemp(WriteData.Entries.Last());
+		WriteData.Entries.Pop();
+	}
+
 	// Age our elapsed times
 	for (FPoseHistoryEntry& Entry : WriteData.Entries)
 	{
@@ -614,7 +659,7 @@ void FPoseHistory::EvaluateComponentSpace_AnyThread(float DeltaTime, FCSPose<FCo
 		// Consume every pose until the queue is full
 		WriteData.Entries.Emplace();
 	}
-	else if (WriteData.Entries[WriteData.Entries.Num() - 2].AccumulatedSeconds <= -SamplingInterval)
+	else if (SamplingInterval <= 0.f || WriteData.Entries[WriteData.Entries.Num() - 2].AccumulatedSeconds <= -SamplingInterval)
 	{
 		FPoseHistoryEntry EntryTemp = MoveTemp(WriteData.Entries.First());
 		WriteData.Entries.PopFront();
@@ -623,6 +668,19 @@ void FPoseHistory::EvaluateComponentSpace_AnyThread(float DeltaTime, FCSPose<FCo
 
 	// Regardless of the retention policy, we always update the most recent Entry
 	WriteData.Entries.Last().Update(0.f, ComponentSpacePose, WriteData.BoneToTransformMap, bStoreScales);
+
+	if (RootBoneRecoveryTime > 0.f)
+	{
+		// adding the updated "future" root bone Entry
+		const FReferenceSkeleton& RefSkeleton = Skeleton->GetReferenceSkeleton();
+		const TArray<FTransform>& RefBonePose = RefSkeleton.GetRefBonePose();
+		const FTransform& RefRootBone = RefBonePose[RootBoneIndexType];
+
+		FutureEntryTemp.SetNum(1, bStoreScales);
+		FutureEntryTemp.SetComponentSpaceTransform(RootBoneIndexType, RefRootBone);
+		FutureEntryTemp.AccumulatedSeconds = RootBoneRecoveryTime;
+		WriteData.Entries.Emplace(MoveTemp(FutureEntryTemp));
+	}
 }
 
 #if ENABLE_DRAW_DEBUG && ENABLE_ANIM_DEBUG
diff --git a/Engine/Plugins/Animation/PoseSearch/Source/Runtime/Private/PoseSearchTrajectoryLibrary.cpp b/Engine/Plugins/Animation/PoseSearch/Source/Runtime/Private/PoseSearchTrajectoryLibrary.cpp
--- a/Engine/Plugins/Animation/PoseSearch/Source/Runtime/Private/PoseSearchTrajectoryLibrary.cpp
+++ b/Engine/Plugins/Animation/PoseSearch/Source/Runtime/Private/PoseSearchTrajectoryLibrary.cpp
@@ -155,7 +155,8 @@ void UPoseSearchTrajectoryLibrary::InitTrajectorySamples(
 	FPoseSearchQueryTrajectory& Trajectory,
 	const FPoseSearchTrajectoryData& TrajectoryData,
 	const FPoseSearchTrajectoryData::FDerived& TrajectoryDataDerived,
-	const FPoseSearchTrajectoryData::FSampling& TrajectoryDataSampling)
+	const FPoseSearchTrajectoryData::FSampling& TrajectoryDataSampling,
+	float DeltaTime)
 {
 	const int32 NumHistorySamples = TrajectoryDataSampling.NumHistorySamples;
 	const int32 NumPredictionSamples = TrajectoryDataSampling.NumPredictionSamples;
@@ -168,21 +169,21 @@ void UPoseSearchTrajectoryLibrary::InitTrajectorySamples(
 		Trajectory.Samples.SetNumUninitialized(TotalNumSamples);
 
 		// Initialize history samples
-		const float SecondsPerHistorySample = TrajectoryDataSampling.SecondsPerHistorySample;
+		const float SecondsPerHistorySample = FMath::Max(TrajectoryDataSampling.SecondsPerHistorySample, 0.f);
 		for (int32 i = 0; i < NumHistorySamples; ++i)
 		{
 			Trajectory.Samples[i].Position = TrajectoryDataDerived.Position;
 			Trajectory.Samples[i].Facing = TrajectoryDataDerived.Facing;
-			Trajectory.Samples[i].AccumulatedSeconds = SecondsPerHistorySample * (i - NumHistorySamples);
+			Trajectory.Samples[i].AccumulatedSeconds = SecondsPerHistorySample * (i - NumHistorySamples - 1);
 		}
 
 		// Initialize current sample and prediction
-		const float SecondsPerPredictionSample = TrajectoryDataSampling.SecondsPerPredictionSample;
+		const float SecondsPerPredictionSample = FMath::Max(TrajectoryDataSampling.SecondsPerPredictionSample, 0.f);
 		for (int32 i = NumHistorySamples; i < Trajectory.Samples.Num(); ++i)
 		{
 			Trajectory.Samples[i].Position = TrajectoryDataDerived.Position;
 			Trajectory.Samples[i].Facing = TrajectoryDataDerived.Facing;
-			Trajectory.Samples[i].AccumulatedSeconds = SecondsPerPredictionSample * (i - NumHistorySamples);
+			Trajectory.Samples[i].AccumulatedSeconds = SecondsPerPredictionSample * (i - NumHistorySamples) + DeltaTime;
 		}
 	}
 }
@@ -195,46 +196,49 @@ void UPoseSearchTrajectoryLibrary::UpdateHistory_TransformHistory(
 	float DeltaTime)
 {
 	const int32 NumHistorySamples = TrajectoryDataSampling.NumHistorySamples;
-	const float SecondsPerHistorySample = TrajectoryDataSampling.SecondsPerHistorySample;
+	if (NumHistorySamples > 0)
+	{
+		const float SecondsPerHistorySample = TrajectoryDataSampling.SecondsPerHistorySample;
 
-	check(NumHistorySamples <= Trajectory.Samples.Num());
+		check(NumHistorySamples <= Trajectory.Samples.Num());
 
-	// converting all the history samples relative to the previous character position (Trajectory.Samples[NumHistorySamples].Position)
-	for (int32 Index = 0; Index < NumHistorySamples; ++Index)
-	{
-		Trajectory.Samples[Index].Position = Trajectory.Samples[NumHistorySamples].Position - Trajectory.Samples[Index].Position;
-	}
+		// converting all the history samples relative to the previous character position (Trajectory.Samples[NumHistorySamples].Position)
+		for (int32 Index = 0; Index < NumHistorySamples; ++Index)
+		{
+			Trajectory.Samples[Index].Position = Trajectory.Samples[NumHistorySamples].Position - Trajectory.Samples[Index].Position;
+		}
 
-	FVector CurrentTranslation = TrajectoryDataDerived.Velocity * DeltaTime;
+		FVector CurrentTranslation = TrajectoryDataDerived.Velocity * DeltaTime;
 
-	// Shift history Samples when it's time to record a new one.
-	if (NumHistorySamples > 0 && FMath::Abs(Trajectory.Samples[NumHistorySamples - 1].AccumulatedSeconds) >= SecondsPerHistorySample)
-	{
-		for (int32 Index = 0; Index < NumHistorySamples - 1; ++Index)
+		// Shift history Samples when it's time to record a new one.
+		if (SecondsPerHistorySample <= 0.f || FMath::Abs(Trajectory.Samples[NumHistorySamples - 1].AccumulatedSeconds) >= SecondsPerHistorySample)
+		{
+			for (int32 Index = 0; Index < NumHistorySamples - 1; ++Index)
+			{
+				Trajectory.Samples[Index].AccumulatedSeconds = Trajectory.Samples[Index + 1].AccumulatedSeconds - DeltaTime;
+				Trajectory.Samples[Index].Position = Trajectory.Samples[Index + 1].Position + CurrentTranslation;
+				Trajectory.Samples[Index].Facing = Trajectory.Samples[Index + 1].Facing;
+			}
+
+			Trajectory.Samples[NumHistorySamples - 1].AccumulatedSeconds = 0.f;
+			Trajectory.Samples[NumHistorySamples - 1].Position = CurrentTranslation;
+			Trajectory.Samples[NumHistorySamples - 1].Facing = Trajectory.Samples[NumHistorySamples].Facing;
+		}
+		else
 		{
-			Trajectory.Samples[Index].AccumulatedSeconds = Trajectory.Samples[Index + 1].AccumulatedSeconds;
-			Trajectory.Samples[Index].Position = Trajectory.Samples[Index + 1].Position + CurrentTranslation;
-			Trajectory.Samples[Index].Facing = Trajectory.Samples[Index + 1].Facing;
+			for (int32 Index = 0; Index < NumHistorySamples; ++Index)
+			{
+				Trajectory.Samples[Index].AccumulatedSeconds -= DeltaTime;
+				Trajectory.Samples[Index].Position += CurrentTranslation;
+			}
 		}
 
-		Trajectory.Samples[NumHistorySamples - 1].AccumulatedSeconds = 0.f;
-		Trajectory.Samples[NumHistorySamples - 1].Position = CurrentTranslation;
-		Trajectory.Samples[NumHistorySamples - 1].Facing = TrajectoryDataDerived.Facing;
-	}
-	else
-	{
+		// converting the history sample positions in world space by applying the current world position.
 		for (int32 Index = 0; Index < NumHistorySamples; ++Index)
 		{
-			Trajectory.Samples[Index].Position += CurrentTranslation;
+			Trajectory.Samples[Index].Position = TrajectoryDataDerived.Position - Trajectory.Samples[Index].Position;
 		}
 	}
-
-	// converting the history sample positions in world space by applying the current world position.
-	for (int32 Index = 0; Index < NumHistorySamples; ++Index)
-	{
-		Trajectory.Samples[Index].AccumulatedSeconds -= DeltaTime;
-		Trajectory.Samples[Index].Position = TrajectoryDataDerived.Position - Trajectory.Samples[Index].Position;
-	}
 }
 
 FVector UPoseSearchTrajectoryLibrary::RemapVectorMagnitudeWithCurve(
@@ -260,6 +264,7 @@ void UPoseSearchTrajectoryLibrary::UpdatePrediction_SimulateCharacterMovement(
 	const FPoseSearchTrajectoryData& TrajectoryData,
 	const FPoseSearchTrajectoryData::FDerived& TrajectoryDataDerived,
 	const FPoseSearchTrajectoryData::FSampling& TrajectoryDataSampling,
+	float DeltaTime,
 	bool bAlwaysApplyGravity)
 {
 	FVector CurrentPositionWS = TrajectoryDataDerived.Position;
@@ -297,7 +302,7 @@ void UPoseSearchTrajectoryLibrary::UpdatePrediction_SimulateCharacterMovement(
 	const float SecondsPerPredictionSample = TrajectoryDataSampling.SecondsPerPredictionSample;
 	const FQuat ControllerRotationPerStep = FQuat::MakeFromEuler(FVector(0.f, 0.f, TrajectoryDataDerived.ControllerYawRate * SecondsPerPredictionSample));
 
-	float AccumulatedSeconds = 0.f;
+	float AccumulatedSeconds = DeltaTime;
 
 	const int32 LastIndex = Trajectory.Samples.Num() - 1;
 	if (NumHistorySamples <= LastIndex)
@@ -369,9 +374,9 @@ void UPoseSearchTrajectoryLibrary::PoseSearchGenerateTrajectory(
 
 	FPoseSearchTrajectoryData::FDerived TrajectoryDataDerived;
 	InTrajectoryData.UpdateData(InDeltaTime, InAnimInstance, TrajectoryDataDerived, TrajectoryDataState);
-	InitTrajectorySamples(InOutTrajectory, InTrajectoryData, TrajectoryDataDerived, TrajectoryDataSampling);
+	InitTrajectorySamples(InOutTrajectory, InTrajectoryData, TrajectoryDataDerived, TrajectoryDataSampling, InDeltaTime);
 	UpdateHistory_TransformHistory(InOutTrajectory, InTrajectoryData, TrajectoryDataDerived, TrajectoryDataSampling, InDeltaTime);
-	UpdatePrediction_SimulateCharacterMovement(InOutTrajectory, InTrajectoryData, TrajectoryDataDerived, TrajectoryDataSampling, bAlwaysApplyGravity);
+	UpdatePrediction_SimulateCharacterMovement(InOutTrajectory, InTrajectoryData, TrajectoryDataDerived, TrajectoryDataSampling, InDeltaTime, bAlwaysApplyGravity);
 
 	InOutDesiredControllerYawLastUpdate = TrajectoryDataState.DesiredControllerYawLastUpdate;
 
diff --git a/Engine/Plugins/Animation/PoseSearch/Source/Runtime/Private/PoseSearchTrajectoryTypes.cpp b/Engine/Plugins/Animation/PoseSearch/Source/Runtime/Private/PoseSearchTrajectoryTypes.cpp
--- a/Engine/Plugins/Animation/PoseSearch/Source/Runtime/Private/PoseSearchTrajectoryTypes.cpp
+++ b/Engine/Plugins/Animation/PoseSearch/Source/Runtime/Private/PoseSearchTrajectoryTypes.cpp
@@ -6,6 +6,7 @@
 #include "Components/SceneComponent.h"
 #include "DrawDebugHelpers.h"
 #include "PoseSearch/PoseSearchDefines.h"
+#include "Math/Axis.h"
 
 FPoseSearchQueryTrajectorySample FPoseSearchQueryTrajectorySample::Lerp(const FPoseSearchQueryTrajectorySample& Other, float Alpha) const
 {
@@ -69,15 +70,28 @@ void FPoseSearchQueryTrajectory::DebugDrawTrajectory(const UWorld* World, float
 	{
 		for (int32 Index = 0; ; ++Index)
 		{
-			DrawDebugSphere(World, Samples[Index].Position + OffsetVector, 2.f /*Radius*/, 4 /*Segments*/, FColor::Black, false /*bPersistentLines*/, -1.f /*LifeTime*/, SDPG_Foreground);
-			DrawDebugCoordinateSystem(World, Samples[Index].Position + OffsetVector, FRotator(Samples[Index].Facing), 12.f /*Scale*/, false /*bPersistentLines*/, -1.f /*LifeTime*/, SDPG_Foreground);
+			const FVector Pos = Samples[Index].Position + OffsetVector;
+
+			DrawDebugSphere(World, Pos, 1.f, 4, FColor::Black, false, -1.f, SDPG_Foreground);
+
+			const FRotationMatrix R(FRotator(Samples[Index].Facing));
+			const FVector X = R.GetScaledAxis( EAxis::X );
+			const FVector Y = R.GetScaledAxis( EAxis::Y );
+
+			const float Scale = 12.f;
+
+			const bool IsPast = Samples[Index].AccumulatedSeconds <= 0.f;
+
+			DrawDebugLine(World, Pos, Pos + X * Scale, IsPast ? FColor::Red : FColor::Blue, false, -1.f, SDPG_Foreground);
+			DrawDebugLine(World, Pos, Pos + Y * Scale, IsPast ? FColor::Orange : FColor::Turquoise, false, -1.f, SDPG_Foreground);
 
 			if (Index == LastIndex)
 			{
 				break;
 			}
 			
-			DrawDebugLine(World, Samples[Index].Position + OffsetVector, Samples[Index + 1].Position + OffsetVector, FColor::Black, false /*bPersistentLines*/, -1.f /*LifeTime*/, SDPG_Foreground);
+			const FVector NextPos = Samples[Index + 1].Position + OffsetVector;
+			DrawDebugLine(World, Pos, NextPos, FColor::Black, false, -1.f, SDPG_Foreground);
 		}
 	}
 }
@@ -91,15 +105,28 @@ void FPoseSearchQueryTrajectory::DebugDrawTrajectory(FAnimInstanceProxy& AnimIns
 	{
 		for (int32 Index = 0; ; ++Index)
 		{
-			AnimInstanceProxy.AnimDrawDebugSphere(Samples[Index].Position + OffsetVector, 2.f /*Radius*/, 4 /*Segments*/, FColor::Black, false /*bPersistentLines*/, -1.f /*LifeTime*/, 0.f /*Thickness*/, SDPG_Foreground);
-			AnimInstanceProxy.AnimDrawDebugCoordinateSystem(Samples[Index].Position + OffsetVector, FRotator(Samples[Index].Facing), 12.f /*Scale*/, false /*bPersistentLines*/, -1.f /*LifeTime*/, 0.f /*Thickness*/, SDPG_Foreground);
+			const FVector Pos = Samples[Index].Position + OffsetVector;
+
+			AnimInstanceProxy.AnimDrawDebugSphere(Samples[Index].Position + OffsetVector, 1.f, 4, FColor::Black, false, -1.f, 0.f, SDPG_Foreground);
+
+			const FRotationMatrix R(FRotator(Samples[Index].Facing));
+			const FVector X = R.GetScaledAxis( EAxis::X );
+			const FVector Y = R.GetScaledAxis( EAxis::Y );
+
+			const float Scale = 12.f;
+
+			const bool IsPast = Samples[Index].AccumulatedSeconds <= 0.f;
+
+			AnimInstanceProxy.AnimDrawDebugLine(Pos, Pos + X * Scale, IsPast ? FColor::Red : FColor::Blue, false, -1.f, 0.f, SDPG_Foreground);
+			AnimInstanceProxy.AnimDrawDebugLine(Pos, Pos + Y * Scale, IsPast ? FColor::Orange : FColor::Turquoise, false, -1.f, 0.f, SDPG_Foreground);
 
 			if (Index == LastIndex)
 			{
 				break;
 			}
 			
-			AnimInstanceProxy.AnimDrawDebugLine(Samples[Index].Position + OffsetVector, Samples[Index + 1].Position + OffsetVector, FColor::Black, false /*bPersistentLines*/, -1.f /*LifeTime*/, 0.f /*Thickness*/, SDPG_Foreground);
+			const FVector NextPos = Samples[Index + 1].Position + OffsetVector;
+			AnimInstanceProxy.AnimDrawDebugLine(Pos, NextPos, FColor::Black, false, -1.f, 0.f, SDPG_Foreground);
 		}
 	}
 }
diff --git a/Engine/Plugins/Animation/PoseSearch/Source/Runtime/Public/PoseSearch/AnimNode_PoseSearchHistoryCollector.h b/Engine/Plugins/Animation/PoseSearch/Source/Runtime/Public/PoseSearch/AnimNode_PoseSearchHistoryCollector.h
--- a/Engine/Plugins/Animation/PoseSearch/Source/Runtime/Public/PoseSearch/AnimNode_PoseSearchHistoryCollector.h
+++ b/Engine/Plugins/Animation/PoseSearch/Source/Runtime/Public/PoseSearch/AnimNode_PoseSearchHistoryCollector.h
@@ -18,8 +18,8 @@ public:
 	UPROPERTY(EditAnywhere, BlueprintReadWrite, Category = Settings, meta = (ClampMin="2"))
 	int32 PoseCount = 2;
 	
-	// how often in seconds poses are collected
-	UPROPERTY(EditAnywhere, BlueprintReadWrite, Category = Settings, meta = (ClampMin="0.001"))
+	// how often in seconds poses are collected (if 0, it will collect every update)
+	UPROPERTY(EditAnywhere, BlueprintReadWrite, Category = Settings, meta = (ClampMin="0"))
 	float SamplingInterval = 0.04f;
 
 	UPROPERTY(EditAnywhere, Category = Settings)
@@ -37,6 +37,10 @@ public:
 	UPROPERTY(EditAnywhere, Category = Settings)
 	bool bStoreScales = false;
 
+	// time in seconds to recover to the reference skeleton root bone from any eventual root bone modification. if zero the behaviour will be disabled (Experimental)
+	UPROPERTY(EditAnywhere, BlueprintReadWrite, Category = Experimental, meta = (ClampMin="0"))
+	float RootBoneRecoveryTime = 0.f;
+
 	// Update Counter for detecting being relevant
 	FGraphTraversalCounter UpdateCounter;
 
@@ -46,7 +50,7 @@ public:
 #endif // WITH_EDITORONLY_DATA
 
 	// if true Trajectory the pose history node will generate the trajectory using the TrajectoryData parameters instead of relying on the input Trajectory (Experimental)
-	UPROPERTY(EditAnywhere, Category = Settings)
+	UPROPERTY(EditAnywhere, Category = Experimental)
 	bool bGenerateTrajectory = false;
 
 	// input Trajectory samples for pose search queries in Motion Matching. These are expected to be in the world space of the SkeletalMeshComponent.
@@ -55,23 +59,23 @@ public:
 	FPoseSearchQueryTrajectory Trajectory;
 
 	// Input Trajectory velocity will be multiplied by TrajectorySpeedMultiplier: values below 1 will result in selecting animation slower than requested from the original Trajectory
-	UPROPERTY(EditAnywhere, Category = Settings, meta = (PinHiddenByDefault, ClampMin="0", EditCondition="!bGenerateTrajectory", EditConditionHides))
+	UPROPERTY(EditAnywhere, Category = Experimental, meta = (PinHiddenByDefault, ClampMin="0", EditCondition="!bGenerateTrajectory", EditConditionHides))
 	float TrajectorySpeedMultiplier = 1.f;
 
 	// if bGenerateTrajectory is true, this is the number of trajectory past (collected) samples
-	UPROPERTY(EditAnywhere, BlueprintReadWrite, Category = Settings, meta = (ClampMin = "2", EditCondition = "bGenerateTrajectory", EditConditionHides))
+	UPROPERTY(EditAnywhere, BlueprintReadWrite, Category = Experimental, meta = (ClampMin = "2", EditCondition = "bGenerateTrajectory", EditConditionHides))
 	int32 TrajectoryHistoryCount = 10;
 
 	// if bGenerateTrajectory is true, this is the number of trajectory future (prediction) samples
-	UPROPERTY(EditAnywhere, BlueprintReadWrite, Category = Settings, meta = (ClampMin="2", EditCondition="bGenerateTrajectory", EditConditionHides))
+	UPROPERTY(EditAnywhere, BlueprintReadWrite, Category = Experimental, meta = (ClampMin="2", EditCondition="bGenerateTrajectory", EditConditionHides))
 	int32 TrajectoryPredictionCount = 8;
 
 	// if bGenerateTrajectory is true, this is the sampling interval between trajectory future (prediction) samples
-	UPROPERTY(EditAnywhere, BlueprintReadWrite, Category = Settings, meta = (ClampMin="0.001", EditCondition="bGenerateTrajectory", EditConditionHides))
+	UPROPERTY(EditAnywhere, BlueprintReadWrite, Category = Experimental, meta = (ClampMin="0.001", EditCondition="bGenerateTrajectory", EditConditionHides))
 	float PredictionSamplingInterval = 0.4f;
 
 	// if bGenerateTrajectory is true, TrajectoryData contains the tuning parameters to generate the trajectory
-	UPROPERTY(EditAnywhere, Category = Settings, meta=(EditCondition="bGenerateTrajectory", EditConditionHides))
+	UPROPERTY(EditAnywhere, Category = Experimental, meta=(EditCondition="bGenerateTrajectory", EditConditionHides))
 	FPoseSearchTrajectoryData TrajectoryData;
 
 	// FAnimNode_Base interface
diff --git a/Engine/Plugins/Animation/PoseSearch/Source/Runtime/Public/PoseSearch/PoseSearchContext.h b/Engine/Plugins/Animation/PoseSearch/Source/Runtime/Public/PoseSearch/PoseSearchContext.h
--- a/Engine/Plugins/Animation/PoseSearch/Source/Runtime/Public/PoseSearch/PoseSearchContext.h
+++ b/Engine/Plugins/Animation/PoseSearch/Source/Runtime/Public/PoseSearch/PoseSearchContext.h
@@ -58,12 +58,11 @@ struct POSESEARCH_API FDebugDrawParams
 	const FSearchIndex* GetSearchIndex() const;
 	const UPoseSearchSchema* GetSchema() const;
 
-	FVector ExtractPosition(TConstArrayView<float> PoseVector, const UPoseSearchFeatureChannel_Position* Position) const;
 	FVector ExtractPosition(TConstArrayView<float> PoseVector, float SampleTimeOffset, int8 SchemaBoneIdx, const FRole& Role, EPermutationTimeType PermutationTimeType = EPermutationTimeType::UseSampleTime, int32 SamplingAttributeId = INDEX_NONE) const;
 
 	FQuat ExtractRotation(TConstArrayView<float> PoseVector, float SampleTimeOffset, int8 SchemaBoneIdx, const FRole& Role, EPermutationTimeType PermutationTimeType = EPermutationTimeType::UseSampleTime, int32 SamplingAttributeId = INDEX_NONE) const;
 
-	FTransform GetRootTransform(const FRole& Role, float SampleTimeOffset = 0.f) const;
+	FTransform GetRootBoneTransform(const FRole& Role, float SampleTimeOffset = 0.f) const;
 
 	void DrawLine(const FVector& LineStart, const FVector& LineEnd, const FColor& Color, float Thickness = 0.f) const;
 	void DrawPoint(const FVector& Position, const FColor& Color, float Thickness = 6.f) const;
diff --git a/Engine/Plugins/Animation/PoseSearch/Source/Runtime/Public/PoseSearch/PoseSearchHistory.h b/Engine/Plugins/Animation/PoseSearch/Source/Runtime/Public/PoseSearch/PoseSearchHistory.h
--- a/Engine/Plugins/Animation/PoseSearch/Source/Runtime/Public/PoseSearch/PoseSearchHistory.h
+++ b/Engine/Plugins/Animation/PoseSearch/Source/Runtime/Public/PoseSearch/PoseSearchHistory.h
@@ -54,7 +54,7 @@ public:
 	// if ReferenceBoneIndexType is 0 (RootBoneIndexType), OutBoneTransform is in root bone space
 	// if ReferenceBoneIndexType is FBoneIndexType(-1) (ComponentSpaceIndexType), OutBoneTransform is in component space
 	// if ReferenceBoneIndexType is FBoneIndexType(-2) (WorldSpaceIndexType), OutBoneTransform is in world space
-	virtual bool GetTransformAtTime(float Time, FTransform& OutBoneTransform, const USkeleton* BoneIndexSkeleton = nullptr, FBoneIndexType BoneIndexType = RootBoneIndexType, FBoneIndexType ReferenceBoneIndexType = ComponentSpaceIndexType, bool bExtrapolate = true) const = 0;
+	virtual bool GetTransformAtTime(float Time, FTransform& OutBoneTransform, const USkeleton* BoneIndexSkeleton = nullptr, FBoneIndexType BoneIndexType = RootBoneIndexType, FBoneIndexType ReferenceBoneIndexType = ComponentSpaceIndexType, bool bExtrapolate = false) const = 0;
 	virtual const FPoseSearchQueryTrajectory& GetTrajectory() const = 0;
 	
 	// @todo: deprecate this API. TrajectorySpeedMultiplier should be a global query scaling value passed as input parameter of FSearchContext during config BuildQuery
@@ -68,7 +68,7 @@ public:
 #if ENABLE_DRAW_DEBUG && ENABLE_ANIM_DEBUG
 	virtual void DebugDraw(const UWorld* World, FColor Color) const = 0;
 	virtual void DebugDraw(FAnimInstanceProxy& AnimInstanceProxy, FColor Color) const = 0;
-	virtual void DebugDraw(FAnimInstanceProxy& AnimInstanceProxy, FColor Color, float Time, float PointSize = 6.f, bool bExtrapolate = true) const;
+	virtual void DebugDraw(FAnimInstanceProxy& AnimInstanceProxy, FColor Color, float Time, float PointSize = 6.f, bool bExtrapolate = false) const;
 #endif
 };
 
@@ -77,7 +77,7 @@ struct POSESEARCH_API FArchivedPoseHistory : public IPoseHistory
 	void InitFrom(const IPoseHistory* PoseHistory);
 
 	// IPoseHistory interface
-	virtual bool GetTransformAtTime(float Time, FTransform& OutBoneTransform, const USkeleton* BoneIndexSkeleton = nullptr, FBoneIndexType BoneIndexType = RootBoneIndexType, FBoneIndexType ReferenceBoneIndexType = ComponentSpaceIndexType, bool bExtrapolate = true) const override;
+	virtual bool GetTransformAtTime(float Time, FTransform& OutBoneTransform, const USkeleton* BoneIndexSkeleton = nullptr, FBoneIndexType BoneIndexType = RootBoneIndexType, FBoneIndexType ReferenceBoneIndexType = ComponentSpaceIndexType, bool bExtrapolate = false) const override;
 	virtual const FPoseSearchQueryTrajectory& GetTrajectory() const override { return Trajectory; }
 	virtual float GetTrajectorySpeedMultiplier() const override { return 1.f; }
 	virtual bool IsEmpty() const override { return Entries.IsEmpty(); }
@@ -110,10 +110,10 @@ struct FPoseHistory : public IPoseHistory
 
 	void Initialize_AnyThread(int32 InNumPoses, float InSamplingInterval);
 	void CacheBones_AnyThread(const TArray<FBoneIndexType>& RequiredBones);
-	void EvaluateComponentSpace_AnyThread(float DeltaTime, FCSPose<FCompactPose>& ComponentSpacePose, bool bStoreScales);
+	void EvaluateComponentSpace_AnyThread(float DeltaTime, FCSPose<FCompactPose>& ComponentSpacePose, bool bStoreScales, float RootBoneRecoveryTime);
 
 	// IPoseHistory interface
-	virtual bool GetTransformAtTime(float Time, FTransform& OutBoneTransform, const USkeleton* BoneIndexSkeleton = nullptr, FBoneIndexType BoneIndexType = RootBoneIndexType, FBoneIndexType ReferenceBoneIndexType = ComponentSpaceIndexType, bool bExtrapolate = true) const override;
+	virtual bool GetTransformAtTime(float Time, FTransform& OutBoneTransform, const USkeleton* BoneIndexSkeleton = nullptr, FBoneIndexType BoneIndexType = RootBoneIndexType, FBoneIndexType ReferenceBoneIndexType = ComponentSpaceIndexType, bool bExtrapolate = false) const override;
 	virtual const FPoseSearchQueryTrajectory& GetTrajectory() const override;
 	virtual float GetTrajectorySpeedMultiplier() const override;
 	virtual bool IsEmpty() const override;
@@ -170,7 +170,7 @@ struct FMemStackPoseHistory : public IPoseHistory
 	void Init(const IPoseHistory* InPoseHistory);
 
 	// IPoseHistory interface
-	virtual bool GetTransformAtTime(float Time, FTransform& OutBoneTransform, const USkeleton* BoneIndexSkeleton = nullptr, FBoneIndexType BoneIndexType = RootBoneIndexType, FBoneIndexType ReferenceBoneIndexType = ComponentSpaceIndexType, bool bExtrapolate = true) const override;
+	virtual bool GetTransformAtTime(float Time, FTransform& OutBoneTransform, const USkeleton* BoneIndexSkeleton = nullptr, FBoneIndexType BoneIndexType = RootBoneIndexType, FBoneIndexType ReferenceBoneIndexType = ComponentSpaceIndexType, bool bExtrapolate = false) const override;
 	virtual const FPoseSearchQueryTrajectory& GetTrajectory() const override { check(PoseHistory); return PoseHistory->GetTrajectory(); }
 	virtual float GetTrajectorySpeedMultiplier() const override { check(PoseHistory); return PoseHistory->GetTrajectorySpeedMultiplier(); }
 	virtual bool IsEmpty() const override { check(PoseHistory); return PoseHistory->IsEmpty() && FutureEntries.IsEmpty(); }
diff --git a/Engine/Plugins/Animation/PoseSearch/Source/Runtime/Public/PoseSearch/PoseSearchTrajectoryLibrary.h b/Engine/Plugins/Animation/PoseSearch/Source/Runtime/Public/PoseSearch/PoseSearchTrajectoryLibrary.h
--- a/Engine/Plugins/Animation/PoseSearch/Source/Runtime/Public/PoseSearch/PoseSearchTrajectoryLibrary.h
+++ b/Engine/Plugins/Animation/PoseSearch/Source/Runtime/Public/PoseSearch/PoseSearchTrajectoryLibrary.h
@@ -42,6 +42,7 @@ public:
 	struct FSampling
 	{
 		int32 NumHistorySamples = 0;
+		// if SecondsPerHistorySample <= 0, then we collect every update
 		float SecondsPerHistorySample = 0.f;
 
 		int32 NumPredictionSamples = 0;
@@ -90,7 +91,7 @@ class POSESEARCH_API UPoseSearchTrajectoryLibrary : public UBlueprintFunctionLib
 	GENERATED_BODY()
 
 public:
-	static void InitTrajectorySamples(FPoseSearchQueryTrajectory& Trajectory, const FPoseSearchTrajectoryData& TrajectoryData, const FPoseSearchTrajectoryData::FDerived& TrajectoryDataDerived, const FPoseSearchTrajectoryData::FSampling& TrajectoryDataSampling);
+	static void InitTrajectorySamples(FPoseSearchQueryTrajectory& Trajectory, const FPoseSearchTrajectoryData& TrajectoryData, const FPoseSearchTrajectoryData::FDerived& TrajectoryDataDerived, const FPoseSearchTrajectoryData::FSampling& TrajectoryDataSampling, float DeltaTime);
 
 	// Update history by tracking offsets that result from character intent (e.g. movement component velocity) and applying
 	// that to the current world transform. This works well on moving platforms as it only stores a history of movement
@@ -98,7 +99,7 @@ public:
 	static void UpdateHistory_TransformHistory(FPoseSearchQueryTrajectory& Trajectory, const FPoseSearchTrajectoryData& TrajectoryData, const FPoseSearchTrajectoryData::FDerived& TrajectoryDataDerived, const FPoseSearchTrajectoryData::FSampling& TrajectoryDataSampling, float DeltaTime);
 
 	// Update prediction by simulating the movement math for ground locomotion from UCharacterMovementComponent.
-	static void UpdatePrediction_SimulateCharacterMovement(FPoseSearchQueryTrajectory& Trajectory, const FPoseSearchTrajectoryData& TrajectoryData, const FPoseSearchTrajectoryData::FDerived& TrajectoryDataDerived, const FPoseSearchTrajectoryData::FSampling& TrajectoryDataSampling, bool bAlwaysApplyGravity = false);
+	static void UpdatePrediction_SimulateCharacterMovement(FPoseSearchQueryTrajectory& Trajectory, const FPoseSearchTrajectoryData& TrajectoryData, const FPoseSearchTrajectoryData::FDerived& TrajectoryDataDerived, const FPoseSearchTrajectoryData::FSampling& TrajectoryDataSampling, float DeltaTime, bool bAlwaysApplyGravity = false);
 
 	/** Get a Pose History node context from an anim node context (pure) */
 	UFUNCTION(BlueprintCallable, Category = "Animation|PoseSearch", meta = (BlueprintThreadSafe, DisplayName = "Pose Search Generate Trajectory"))
@@ -112,7 +113,7 @@ public:
 		ETraceTypeQuery TraceChannel, bool bTraceComplex, const TArray<AActor*>& ActorsToIgnore, EDrawDebugTrace::Type DrawDebugType, bool bIgnoreSelf, FLinearColor TraceColor = FLinearColor::Red, FLinearColor TraceHitColor = FLinearColor::Green, float DrawTime = 5.0f);
 
 	UFUNCTION(BlueprintPure, meta = (BlueprintThreadSafe), Category="Animation|PoseSearch")
-	static void GetTrajectorySampleAtTime(UPARAM(ref) const FPoseSearchQueryTrajectory& InTrajectory, float Time, FPoseSearchQueryTrajectorySample& OutTrajectorySample, bool bExtrapolate = true);
+	static void GetTrajectorySampleAtTime(UPARAM(ref) const FPoseSearchQueryTrajectory& InTrajectory, float Time, FPoseSearchQueryTrajectorySample& OutTrajectorySample, bool bExtrapolate = false);
 
 private:
 	static FVector RemapVectorMagnitudeWithCurve(const FVector& Vector, bool bUseCurve, const FRuntimeFloatCurve& Curve);
diff --git a/Engine/Plugins/Animation/PoseSearch/Source/Runtime/Public/PoseSearch/PoseSearchTrajectoryTypes.h b/Engine/Plugins/Animation/PoseSearch/Source/Runtime/Public/PoseSearch/PoseSearchTrajectoryTypes.h
--- a/Engine/Plugins/Animation/PoseSearch/Source/Runtime/Public/PoseSearch/PoseSearchTrajectoryTypes.h
+++ b/Engine/Plugins/Animation/PoseSearch/Source/Runtime/Public/PoseSearch/PoseSearchTrajectoryTypes.h
@@ -38,7 +38,7 @@ struct POSESEARCH_API FPoseSearchQueryTrajectory
 	UPROPERTY(EditAnywhere, BlueprintReadWrite, Category = "Pose Search Query Trajectory")
 	TArray<FPoseSearchQueryTrajectorySample> Samples;
 
-	FPoseSearchQueryTrajectorySample GetSampleAtTime(float Time, bool bExtrapolate = true) const;
+	FPoseSearchQueryTrajectorySample GetSampleAtTime(float Time, bool bExtrapolate = false) const;
 	
 #if ENABLE_ANIM_DEBUG
 	void DebugDrawTrajectory(const UWorld* World, float HeightOffset = 0.f) const;
