motion matching - editor improvement

- Added "Read only" to tooltips for inmmutable variables.
- Update context menu options when right click on a selection of database items.
- Added ability to open context menu on empty space. Allows to add new assets or paste.
- Added option left click for mirror option to cycle through options.
- Update icons for database items.
- small layout fixes

#rnx none

diff --git a/Engine/Plugins/Animation/PoseSearch/Source/Editor/Private/PoseSearchDatabaseAssetListItem.cpp b/Engine/Plugins/Animation/PoseSearch/Source/Editor/Private/PoseSearchDatabaseAssetListItem.cpp
--- a/Engine/Plugins/Animation/PoseSearch/Source/Editor/Private/PoseSearchDatabaseAssetListItem.cpp
+++ b/Engine/Plugins/Animation/PoseSearch/Source/Editor/Private/PoseSearchDatabaseAssetListItem.cpp
@@ -282,20 +282,21 @@ namespace UE::PoseSearch
 
 					// Root Motion
 					+ SHorizontalBox::Slot()
-					.Padding(1.0f, 1.0f)
+					.Padding(1.0f, 2.0f)
 					.AutoWidth()
 					.HAlign(HAlign_Center)
 					.VAlign(VAlign_Center)
 					[
 						SNew(SImage)
 						.Image(FAppStyle::Get().GetBrush("AnimGraph.Attribute.RootMotionDelta.Icon"))
+						.DesiredSizeOverride(FVector2D{16.f, 16.f})
 						.ColorAndOpacity(this, &SDatabaseAssetListItem::GetRootMotionColorAndOpacity)
 						.ToolTipText(this, &SDatabaseAssetListItem::GetRootMotionOptionToolTip)
 					]
 					
 					// Mirror Type
 					+ SHorizontalBox::Slot()
-					.Padding(1.0f, 1.0f)
+					.Padding(2.0f, 3.0f)
 					.AutoWidth()
 					.HAlign(HAlign_Center)
 					.VAlign(VAlign_Center)
@@ -303,11 +304,12 @@ namespace UE::PoseSearch
 						SNew(SImage)
 						.Image(this, &SDatabaseAssetListItem::GetMirrorOptionSlateBrush)
 						.ToolTipText(this, &SDatabaseAssetListItem::GetMirrorOptionToolTip)
+						.OnMouseButtonDown(this, &SDatabaseAssetListItem::MirrorOptionOnMouseButtonDown)
 					]
 
 					// Disable Reselection
 					+ SHorizontalBox::Slot()
-					.Padding(1.0f, 1.0f)
+					.Padding(4.0f, 1.0f)
 					.AutoWidth()
 					.HAlign(HAlign_Center)
 					.VAlign(VAlign_Center)
@@ -317,12 +319,12 @@ namespace UE::PoseSearch
 						.OnCheckStateChanged(const_cast<SDatabaseAssetListItem*>(this), &SDatabaseAssetListItem::OnDisableReselectionChanged)
 						.ToolTipText(this, &SDatabaseAssetListItem::GetDisableReselectionToolTip)
 						// @todo: customize icon!
-						//.CheckedImage(FAppStyle::Get().GetBrush("FractureFlush"))
-						//.CheckedHoveredImage(FAppStyle::Get().GetBrush("FractureFlush"))
-						//.CheckedPressedImage(FAppStyle::Get().GetBrush("FractureFlush"))
-						//.UncheckedImage(FAppStyle::Get().GetBrush("FractureFlush"))
-						//.UncheckedHoveredImage(FAppStyle::Get().GetBrush("FractureFlush"))
-						//.UncheckedPressedImage(FAppStyle::Get().GetBrush("FractureFlush"))
+						.CheckedImage(FAppStyle::Get().GetBrush("Graph.PosePin.Connected"))
+						.CheckedHoveredImage(FAppStyle::Get().GetBrush("Graph.PosePin.ConnectedHovered"))
+						.CheckedPressedImage(FAppStyle::Get().GetBrush("Graph.PosePin.Connected"))
+						.UncheckedImage(FAppStyle::Get().GetBrush("Graph.PosePin.Disconnected"))
+						.UncheckedHoveredImage(FAppStyle::Get().GetBrush("Graph.PosePin.DisconnectedHovered"))
+						.UncheckedPressedImage(FAppStyle::Get().GetBrush("Graph.PosePin.Disconnected"))
 					]
 				]
 			]
@@ -466,7 +468,7 @@ namespace UE::PoseSearch
 		const TSharedPtr<FDatabaseViewModel> ViewModelPtr = EditorViewModel.Pin();
 		const TSharedPtr<FDatabaseAssetTreeNode> TreeNodePtr = WeakAssetTreeNode.Pin();
 
-		ViewModelPtr->SetIsEnabled(TreeNodePtr->SourceAssetIdx, NewCheckboxState == ECheckBoxState::Checked ? true : false);
+		ViewModelPtr->SetIsEnabled(TreeNodePtr->SourceAssetIdx, NewCheckboxState == ECheckBoxState::Checked);
 
 		SkeletonView.Pin()->RefreshTreeView(false, true);
 		ViewModelPtr->BuildSearchIndex();
@@ -503,7 +505,7 @@ namespace UE::PoseSearch
 	FText SDatabaseAssetListItem::GetLoopingToolTip() const
 	{
 		const TSharedPtr<FDatabaseAssetTreeNode> Node = WeakAssetTreeNode.Pin();
-		return Node->IsLooping() ? LOCTEXT("NodeLoopEnabledToolTip", "Looping") : LOCTEXT("NodeLoopDisabledToolTip", "Not looping");
+		return Node->IsLooping() ? LOCTEXT("NodeLoopEnabledToolTip", "Looping (Read only)") : LOCTEXT("NodeLoopDisabledToolTip", "Not looping (Read only)");
 	}
 
 	FSlateColor SDatabaseAssetListItem::GetRootMotionColorAndOpacity() const
@@ -515,7 +517,7 @@ namespace UE::PoseSearch
 	FText SDatabaseAssetListItem::GetRootMotionOptionToolTip() const
 	{
 		const TSharedPtr<FDatabaseAssetTreeNode> Node = WeakAssetTreeNode.Pin();
-		return Node->IsRootMotionEnabled() ? LOCTEXT("NodeRootMotionEnabledToolTip", "Root motion enabled") : LOCTEXT("NodeRootMotionDisabledToolTip", "No root motion enabled");
+		return Node->IsRootMotionEnabled() ? LOCTEXT("NodeRootMotionEnabledToolTip", "Root motion enabled (Read only)") : LOCTEXT("NodeRootMotionDisabledToolTip", "No root motion enabled (Read only)");
 
 	}
 	const FSlateBrush* SDatabaseAssetListItem::GetMirrorOptionSlateBrush() const
@@ -526,13 +528,13 @@ namespace UE::PoseSearch
 		switch (Node->GetMirrorOption())
 		{
 			case EPoseSearchMirrorOption::UnmirroredOnly: 
-				return FAppStyle::Get().GetBrush("Icons.Minus");
+				return FAppStyle::Get().GetBrush("GraphEditor.AlignNodesRight");
 			
 			case EPoseSearchMirrorOption::MirroredOnly: 
-				return FAppStyle::Get().GetBrush("Icons.Plus");
+				return FAppStyle::Get().GetBrush("GraphEditor.AlignNodesLeft");
 			
 			case EPoseSearchMirrorOption::UnmirroredAndMirrored:
-				return FAppStyle::Get().GetBrush("Icons.X");
+				return FAppStyle::Get().GetBrush("GraphEditor.AlignNodesCenter");
 			
 			default:
 				return nullptr;
@@ -544,7 +546,39 @@ namespace UE::PoseSearch
 		const TSharedPtr<FDatabaseAssetTreeNode> Node = WeakAssetTreeNode.Pin();
 		return FText::FromString(LOCTEXT("ToolTipMirrorOption", "Mirror Option: ").ToString() + (Node ? UEnum::GetDisplayValueAsText(Node->GetMirrorOption()).ToString() : LOCTEXT("ToolTipMirrorOption_Invalid", "Invalid").ToString()));
 	}
-	
+
+	FReply SDatabaseAssetListItem::MirrorOptionOnMouseButtonDown(const FGeometry& InMyGeometry, const FPointerEvent& InMouseEvent)
+	{
+		const TSharedPtr<FDatabaseAssetTreeNode> Node = WeakAssetTreeNode.Pin();
+		const TSharedPtr<FDatabaseViewModel> ViewModel = EditorViewModel.Pin();
+
+		if(InMouseEvent.GetEffectingButton() == EKeys::LeftMouseButton)
+		{
+			if (UPoseSearchDatabase* Database = ViewModel->GetPoseSearchDatabase())
+			{
+				if (FPoseSearchDatabaseAnimationAssetBase* DatabaseAnimationAsset = Database->GetMutableAnimationAssetBase(Node->SourceAssetIdx))
+				{
+					const FScopedTransaction Transaction(LOCTEXT("OnClickEditMirrorOptionPoseSearchDatabase", "Edit Mirror Option"));
+				
+					// Get next mirror option
+					static const TArray<EPoseSearchMirrorOption> OptionArray = { EPoseSearchMirrorOption::UnmirroredOnly, EPoseSearchMirrorOption::MirroredOnly, EPoseSearchMirrorOption::UnmirroredAndMirrored };
+					const int32 NextOption = (static_cast<int32>(DatabaseAnimationAsset->MirrorOption) + 1) % OptionArray.Num();
+
+					// Modify asset (@todo: this should be done through the viewmodel).
+					Database->Modify();
+					DatabaseAnimationAsset->MirrorOption = OptionArray[NextOption];
+				
+					SkeletonView.Pin()->RefreshTreeView(false, true);
+					ViewModel->BuildSearchIndex();
+				
+					return FReply::Handled();
+				}
+			}
+		}
+		
+		return FReply::Unhandled();
+	}
+
 	FText SDatabaseAssetListItem::GetAssetEnabledToolTip() const
 	{
 		if (GetAssetEnabledChecked() == ECheckBoxState::Checked)
diff --git a/Engine/Plugins/Animation/PoseSearch/Source/Editor/Private/PoseSearchDatabaseAssetTree.cpp b/Engine/Plugins/Animation/PoseSearch/Source/Editor/Private/PoseSearchDatabaseAssetTree.cpp
--- a/Engine/Plugins/Animation/PoseSearch/Source/Editor/Private/PoseSearchDatabaseAssetTree.cpp
+++ b/Engine/Plugins/Animation/PoseSearch/Source/Editor/Private/PoseSearchDatabaseAssetTree.cpp
@@ -510,37 +510,61 @@ namespace UE::PoseSearch
 		const TArray<TSharedPtr<FDatabaseAssetTreeNode>> SelectedNodes = TreeView->GetSelectedItems();
 		if (!SelectedNodes.IsEmpty())
 		{
-			MenuBuilder.AddMenuEntry(
-				LOCTEXT("Remove", "Remove"),
-				LOCTEXT("RemoveTooltip", "Removes assets from database"),
-				FSlateIcon(),
-				FUIAction(FExecuteAction::CreateSP(this, &SDatabaseAssetTree::OnDeleteNodes)),
-				NAME_None,
-				EUserInterfaceActionType::Button);
-
-			MenuBuilder.AddMenuEntry(
-				LOCTEXT("Enable", "Enable"),
-				LOCTEXT("EnableTooltip", "Sets Assets Enabled."),
-				FSlateIcon(),
-				FUIAction(FExecuteAction::CreateSP(this, &SDatabaseAssetTree::OnEnableNodes)),
-				NAME_None,
-				EUserInterfaceActionType::Button);
-
-			MenuBuilder.AddMenuEntry(
-				LOCTEXT("Disable", "Disable"),
-				LOCTEXT("DisableToolTip", "Sets Assets Disabled."),
-				FSlateIcon(),
-				FUIAction(FExecuteAction::CreateSP(this, &SDatabaseAssetTree::OnDisableNodes)),
-				NAME_None,
-				EUserInterfaceActionType::Button);
+			MenuBuilder.BeginSection("SelectedAssetsEdit", LOCTEXT("SelectedAssetEdit", "Asset Actions"));
+			{
+				MenuBuilder.AddMenuEntry(
+					LOCTEXT("Enable", "Enable selected assets"),
+					LOCTEXT("EnableTooltip", "Sets Assets Enabled."),
+					FSlateIcon(FAppStyle::GetAppStyleSetName(), "Icons.Visible"),
+					FUIAction(FExecuteAction::CreateSP(this, &SDatabaseAssetTree::OnEnableNodes)),
+					NAME_None,
+					EUserInterfaceActionType::Button);
+
+				MenuBuilder.AddMenuEntry(
+					LOCTEXT("Disable", "Disable selected assets"),
+					LOCTEXT("DisableToolTip", "Sets Assets Disabled."),
+					FSlateIcon(FAppStyle::GetAppStyleSetName(), "Icons.Hidden"),
+					FUIAction(FExecuteAction::CreateSP(this, &SDatabaseAssetTree::OnDisableNodes)),
+					NAME_None,
+					EUserInterfaceActionType::Button);
+
+				MenuBuilder.AddMenuEntry(
+					LOCTEXT("ConvertToBranchIn", "Convert selected assets to sample via BranchIn notify"),
+					LOCTEXT("ConvertToBranchInToolTip", "Creates PoseSearchBranchIn notify state for the asset sampling range"),
+					FSlateIcon(),
+					FUIAction(FExecuteAction::CreateSP(this, &SDatabaseAssetTree::OnConvertToBranchIn)),
+					NAME_None,
+					EUserInterfaceActionType::Button);	
+			}
+			MenuBuilder.EndSection();
 
-			MenuBuilder.AddMenuEntry(
-				LOCTEXT("ConvertToBranchIn", "ConvertToBranchIn"),
-				LOCTEXT("ConvertToBranchInToolTip", "Creates PoseSearchBranchIn notify state for the asset sampling range"),
-				FSlateIcon(),
-				FUIAction(FExecuteAction::CreateSP(this, &SDatabaseAssetTree::OnConvertToBranchIn)),
-				NAME_None,
-				EUserInterfaceActionType::Button);
+			MenuBuilder.BeginSection("SelectionAssetsClipboardEdit", LOCTEXT("SelectionAssetsClipboardEdit", "Edit"));
+			{
+				MenuBuilder.AddMenuEntry(FGenericCommands::Get().Cut);
+				MenuBuilder.AddMenuEntry(FGenericCommands::Get().Copy);
+				MenuBuilder.AddMenuEntry(FGenericCommands::Get().Paste);
+				MenuBuilder.AddMenuEntry(FGenericCommands::Get().Delete);
+			}
+			MenuBuilder.EndSection();
+		}
+		else
+		{
+			// Asset actions
+			MenuBuilder.BeginSection("Edit", LOCTEXT("EditSection", "Edit"));
+			MenuBuilder.AddWrapperSubMenu(
+				LOCTEXT("AddNewAnimAssetNoNodes", "Add"),
+				LOCTEXT("AddNewAnimAssetNoNodesToolTip", "Add a new Sequence, Blend Space, Anim Composite, or Anim Montage"),
+				FOnGetContent::CreateSP(this, &SDatabaseAssetTree::CreateAddNewMenuWidget),
+				FSlateIcon(FAppStyle::GetAppStyleSetName(), "Icons.Plus"));
+			MenuBuilder.EndSection();
+			
+			// Edit / Clipboard actions
+			MenuBuilder.BeginSection("Clipboard", LOCTEXT("ClipboardSection", "Clipboard"));
+			{
+				MenuBuilder.AddMenuEntry(FGenericCommands::Get().Paste);
+			}
+			
+			MenuBuilder.EndSection();
 		}
 
 		return MenuBuilder.MakeWidget();
diff --git a/Engine/Plugins/Animation/PoseSearch/Source/Editor/Public/PoseSearchDatabaseAssetListItem.h b/Engine/Plugins/Animation/PoseSearch/Source/Editor/Public/PoseSearchDatabaseAssetListItem.h
--- a/Engine/Plugins/Animation/PoseSearch/Source/Editor/Public/PoseSearchDatabaseAssetListItem.h
+++ b/Engine/Plugins/Animation/PoseSearch/Source/Editor/Public/PoseSearchDatabaseAssetListItem.h
@@ -38,15 +38,20 @@ namespace UE::PoseSearch
 		void OnAddAnimMontage();
 		virtual FReply OnMouseButtonDoubleClick(const FGeometry& InMyGeometry, const FPointerEvent& InMouseEvent) override;
 
+		
 		EVisibility GetSelectedActorIconVisbility() const;
 
 		FSlateColor GetNameTextColorAndOpacity() const;
+		
 		FSlateColor GetLoopingColorAndOpacity() const;
 		FText GetLoopingToolTip() const;
+		
 		FSlateColor GetRootMotionColorAndOpacity() const;
 		FText GetRootMotionOptionToolTip() const;
+		
 		const FSlateBrush* GetMirrorOptionSlateBrush() const;
 		FText GetMirrorOptionToolTip() const;
+		FReply MirrorOptionOnMouseButtonDown(const FGeometry& InMyGeometry, const FPointerEvent& InMouseEvent);
 		
 		FText GetDisableReselectionToolTip() const;
 		ECheckBoxState GetDisableReselectionChecked() const;
