MM Editor - Removed old group view and renamed asset tree tab to asset list tab

#rnx minor

diff --git a/Engine/Plugins/Animation/PoseSearch/Source/Editor/Private/PoseSearchDatabaseAssetTree.cpp b/Engine/Plugins/Animation/PoseSearch/Source/Editor/Private/PoseSearchDatabaseAssetTree.cpp
--- a/Engine/Plugins/Animation/PoseSearch/Source/Editor/Private/PoseSearchDatabaseAssetTree.cpp
+++ b/Engine/Plugins/Animation/PoseSearch/Source/Editor/Private/PoseSearchDatabaseAssetTree.cpp
@@ -237,15 +237,6 @@ namespace UE::PoseSearch
 
 		// Rebuild node hierarchy
 		{
-			// Setup default group node
-			{
-				const TSharedPtr<FDatabaseAssetTreeNode> DefaultGroupNode = MakeShared<FDatabaseAssetTreeNode>(INDEX_NONE, ViewModelRef);
-				AllNodes.Add(DefaultGroupNode);
-				RootNodes.Add(DefaultGroupNode);
-			}
-			
-			const int32 DefaultGroupIdx = RootNodes.Num() - 1;
-						
 			// Build an index based off of alphabetical order than iterate the index instead
 			TArray<uint32> IndexArray;
 			IndexArray.SetNumUninitialized(Database->GetAnimationAssets().Num());
@@ -287,13 +278,9 @@ namespace UE::PoseSearch
 					{
 						// Create sequence node
 						const TSharedPtr<FDatabaseAssetTreeNode> SequenceGroupNode = MakeShared<FDatabaseAssetTreeNode>(MappedId, ViewModelRef);
-						const TSharedPtr<FDatabaseAssetTreeNode>& ParentGroupNode = RootNodes[DefaultGroupIdx];
-
-						// Setup hierarchy
-						SequenceGroupNode->Parent = ParentGroupNode;
-						ParentGroupNode->Children.Add(SequenceGroupNode);
 
 						// Keep track of node
+						RootNodes.Add(SequenceGroupNode);
 						AllNodes.Add(SequenceGroupNode);
 					}
 				}
diff --git a/Engine/Plugins/Animation/PoseSearch/Source/Editor/Private/PoseSearchDatabaseEditor.cpp b/Engine/Plugins/Animation/PoseSearch/Source/Editor/Private/PoseSearchDatabaseEditor.cpp
--- a/Engine/Plugins/Animation/PoseSearch/Source/Editor/Private/PoseSearchDatabaseEditor.cpp
+++ b/Engine/Plugins/Animation/PoseSearch/Source/Editor/Private/PoseSearchDatabaseEditor.cpp
@@ -449,7 +449,7 @@ namespace UE::PoseSearch
 		check(Args.GetTabId() == FDatabaseEditorTabs::AssetTreeViewID);
 
 		return SNew(SDockTab)
-			.Label(LOCTEXT("AssetTreeView_Title", "Asset Tree"))
+			.Label(LOCTEXT("AssetTreeView_Title", "Asset List"))
 			[
 				AssetTreeWidget.ToSharedRef()
 			];
