MM Editor - Fixed bad tooltip names for pose reselection and added new icon

#rnx none

diff --git a/Engine/Plugins/Animation/PoseSearch/Source/Editor/Private/PoseSearchDatabaseAssetListItem.cpp b/Engine/Plugins/Animation/PoseSearch/Source/Editor/Private/PoseSearchDatabaseAssetListItem.cpp
--- a/Engine/Plugins/Animation/PoseSearch/Source/Editor/Private/PoseSearchDatabaseAssetListItem.cpp
+++ b/Engine/Plugins/Animation/PoseSearch/Source/Editor/Private/PoseSearchDatabaseAssetListItem.cpp
@@ -318,13 +318,13 @@ namespace UE::PoseSearch
 						.IsChecked(this, &SDatabaseAssetListItem::GetDisableReselectionChecked)
 						.OnCheckStateChanged(const_cast<SDatabaseAssetListItem*>(this), &SDatabaseAssetListItem::OnDisableReselectionChanged)
 						.ToolTipText(this, &SDatabaseAssetListItem::GetDisableReselectionToolTip)
-						// @todo: customize icon!
-						.CheckedImage(FAppStyle::Get().GetBrush("Graph.PosePin.Connected"))
-						.CheckedHoveredImage(FAppStyle::Get().GetBrush("Graph.PosePin.ConnectedHovered"))
-						.CheckedPressedImage(FAppStyle::Get().GetBrush("Graph.PosePin.Connected"))
-						.UncheckedImage(FAppStyle::Get().GetBrush("Graph.PosePin.Disconnected"))
-						.UncheckedHoveredImage(FAppStyle::Get().GetBrush("Graph.PosePin.DisconnectedHovered"))
-						.UncheckedPressedImage(FAppStyle::Get().GetBrush("Graph.PosePin.Disconnected"))
+						.Padding(FMargin(0.0f, 2.0f, 0.0f, 0.0f))
+						.CheckedImage(FAppStyle::Get().GetBrush("MotionMatchingEditor.DisablePoseReselection"))
+						.CheckedHoveredImage(FAppStyle::Get().GetBrush("MotionMatchingEditor.DisablePoseReselection"))
+						.CheckedPressedImage(FAppStyle::Get().GetBrush("MotionMatchingEditor.DisablePoseReselection"))
+						.UncheckedImage(FAppStyle::Get().GetBrush("MotionMatchingEditor.EnablePoseReselection"))
+						.UncheckedHoveredImage(FAppStyle::Get().GetBrush("MotionMatchingEditor.EnablePoseReselection"))
+						.UncheckedPressedImage(FAppStyle::Get().GetBrush("MotionMatchingEditor.EnablePoseReselection"))
 					]
 				]
 			]
@@ -401,10 +401,10 @@ namespace UE::PoseSearch
 	{
 		if (GetDisableReselectionChecked() == ECheckBoxState::Checked)
 		{
-			return LOCTEXT("EnableReselectionToolTip", "Enable reselection of poses from the same asset.");
+			return LOCTEXT("EnableReselectionToolTip", "Reselection of poses from the same asset is disabled.");
 		}
 		
-		return LOCTEXT("DisableReselectionToolTip", "Disable reselection of poses from the same asset.");
+		return LOCTEXT("DisableReselectionToolTip", "Reselection of poses from the same asset is enabled.");
 	}
 
 	ECheckBoxState SDatabaseAssetListItem::GetDisableReselectionChecked() const
diff --git a/Engine/Source/Editor/EditorStyle/Private/StarshipStyle.cpp b/Engine/Source/Editor/EditorStyle/Private/StarshipStyle.cpp
--- a/Engine/Source/Editor/EditorStyle/Private/StarshipStyle.cpp
+++ b/Engine/Source/Editor/EditorStyle/Private/StarshipStyle.cpp
@@ -5025,6 +5025,12 @@ void FStarshipEditorStyle::FStyle::SetupPersonaStyle()
 		Set("SkeletalMeshEditor.MeshSectionSelection", new IMAGE_BRUSH_SVG("Starship/Persona/SectionSelection", Icon20x20));
 	}
 
+	// Motion Matching editor
+	{
+		Set("MotionMatchingEditor.EnablePoseReselection", new IMAGE_BRUSH_SVG("Starship/Animation/PoseReselection", Icon20x20));
+		Set("MotionMatchingEditor.DisablePoseReselection", new IMAGE_BRUSH_SVG("Starship/Animation/PoseReselection", Icon20x20, FLinearColor(1.0f, 1.0f, 1.0f, 0.25f)));
+	}
+	
 	// New anim blueprint dialog
 	{
 		Set("NewAnimBlueprintDialog.AreaBorder", new FSlateRoundedBoxBrush(FStyleColors::Panel, 4.0f));
