Added cvar for debug thickness of pose search trajectory
#rnx

diff --git a/Engine/Plugins/Animation/PoseSearch/Source/Runtime/Private/PoseSearchHistory.cpp b/Engine/Plugins/Animation/PoseSearch/Source/Runtime/Private/PoseSearchHistory.cpp
--- a/Engine/Plugins/Animation/PoseSearch/Source/Runtime/Private/PoseSearchHistory.cpp
+++ b/Engine/Plugins/Animation/PoseSearch/Source/Runtime/Private/PoseSearchHistory.cpp
@@ -21,6 +21,7 @@ namespace UE::PoseSearch
 #if ENABLE_DRAW_DEBUG && ENABLE_ANIM_DEBUG
 TAutoConsoleVariable<bool> CVarAnimPoseHistoryDebugDrawPose(TEXT("a.AnimNode.PoseHistory.DebugDrawPose"), false, TEXT("Enable / Disable Pose History Pose DebugDraw"));
 TAutoConsoleVariable<bool> CVarAnimPoseHistoryDebugDrawTrajectory(TEXT("a.AnimNode.PoseHistory.DebugDrawTrajectory"), false, TEXT("Enable / Disable Pose History Trajectory DebugDraw"));
+TAutoConsoleVariable<float> CVarAnimPoseHistoryDebugDrawTrajectoryThickness(TEXT("a.AnimNode.PoseHistory.DebugDrawTrajectoryThickness"), 0.0f, TEXT("Thickness of the trajectory debug draw (Default 0.0f)"));
 #endif
 
 /**
@@ -690,7 +691,8 @@ void FPoseHistory::DebugDraw(FAnimInstanceProxy& AnimInstanceProxy, FColor Color
 
 	if (CVarAnimPoseHistoryDebugDrawTrajectory.GetValueOnAnyThread())
 	{
-		Trajectory.DebugDrawTrajectory(AnimInstanceProxy);
+		const float DebugThickness = CVarAnimPoseHistoryDebugDrawTrajectoryThickness.GetValueOnAnyThread();
+		Trajectory.DebugDrawTrajectory(AnimInstanceProxy, DebugThickness);
 	}
 
 	if (Color.A > 0 && CVarAnimPoseHistoryDebugDrawPose.GetValueOnAnyThread())
diff --git a/Engine/Plugins/Animation/PoseSearch/Source/Runtime/Private/PoseSearchTrajectoryTypes.cpp b/Engine/Plugins/Animation/PoseSearch/Source/Runtime/Private/PoseSearchTrajectoryTypes.cpp
--- a/Engine/Plugins/Animation/PoseSearch/Source/Runtime/Private/PoseSearchTrajectoryTypes.cpp
+++ b/Engine/Plugins/Animation/PoseSearch/Source/Runtime/Private/PoseSearchTrajectoryTypes.cpp
@@ -61,7 +61,7 @@ FPoseSearchQueryTrajectorySample FPoseSearchQueryTrajectory::GetSampleAtTime(flo
 }
 
 #if ENABLE_ANIM_DEBUG
-void FPoseSearchQueryTrajectory::DebugDrawTrajectory(const UWorld* World, float HeightOffset) const
+void FPoseSearchQueryTrajectory::DebugDrawTrajectory(const UWorld* World, const float DebugThickness, float HeightOffset) const
 {
 	const FVector OffsetVector = FVector::UpVector * HeightOffset;
 
@@ -72,7 +72,7 @@ void FPoseSearchQueryTrajectory::DebugDrawTrajectory(const UWorld* World, float
 		{
 			const FVector Pos = Samples[Index].Position + OffsetVector;
 
-			DrawDebugSphere(World, Pos, 1.f, 4, FColor::Black, false, -1.f, SDPG_Foreground);
+			DrawDebugSphere(World, Pos, 1.f, 4, FColor::Black, false, -1.f, SDPG_Foreground, DebugThickness);
 
 			const FRotationMatrix R(FRotator(Samples[Index].Facing));
 			const FVector X = R.GetScaledAxis( EAxis::X );
@@ -82,8 +82,8 @@ void FPoseSearchQueryTrajectory::DebugDrawTrajectory(const UWorld* World, float
 
 			const bool IsPast = Samples[Index].AccumulatedSeconds <= 0.f;
 
-			DrawDebugLine(World, Pos, Pos + X * Scale, IsPast ? FColor::Red : FColor::Blue, false, -1.f, SDPG_Foreground);
-			DrawDebugLine(World, Pos, Pos + Y * Scale, IsPast ? FColor::Orange : FColor::Turquoise, false, -1.f, SDPG_Foreground);
+			DrawDebugLine(World, Pos, Pos + X * Scale, IsPast ? FColor::Red : FColor::Blue, false, -1.f, SDPG_Foreground, DebugThickness);
+			DrawDebugLine(World, Pos, Pos + Y * Scale, IsPast ? FColor::Orange : FColor::Turquoise, false, -1.f, SDPG_Foreground, DebugThickness);
 
 			if (Index == LastIndex)
 			{
@@ -91,12 +91,12 @@ void FPoseSearchQueryTrajectory::DebugDrawTrajectory(const UWorld* World, float
 			}
 			
 			const FVector NextPos = Samples[Index + 1].Position + OffsetVector;
-			DrawDebugLine(World, Pos, NextPos, FColor::Black, false, -1.f, SDPG_Foreground);
+			DrawDebugLine(World, Pos, NextPos, FColor::Black, false, -1.f, SDPG_Foreground, DebugThickness);
 		}
 	}
 }
 
-void FPoseSearchQueryTrajectory::DebugDrawTrajectory(FAnimInstanceProxy& AnimInstanceProxy, float HeightOffset) const
+void FPoseSearchQueryTrajectory::DebugDrawTrajectory(FAnimInstanceProxy& AnimInstanceProxy, const float DebugThickness, float HeightOffset) const
 {
 	const FVector OffsetVector = FVector::UpVector * HeightOffset;
 
@@ -107,7 +107,7 @@ void FPoseSearchQueryTrajectory::DebugDrawTrajectory(FAnimInstanceProxy& AnimIns
 		{
 			const FVector Pos = Samples[Index].Position + OffsetVector;
 
-			AnimInstanceProxy.AnimDrawDebugSphere(Samples[Index].Position + OffsetVector, 1.f, 4, FColor::Black, false, -1.f, 0.f, SDPG_Foreground);
+			AnimInstanceProxy.AnimDrawDebugSphere(Samples[Index].Position + OffsetVector, 1.f, 4, FColor::Black, false, -1.f, DebugThickness, SDPG_Foreground);
 
 			const FRotationMatrix R(FRotator(Samples[Index].Facing));
 			const FVector X = R.GetScaledAxis( EAxis::X );
@@ -117,8 +117,8 @@ void FPoseSearchQueryTrajectory::DebugDrawTrajectory(FAnimInstanceProxy& AnimIns
 
 			const bool IsPast = Samples[Index].AccumulatedSeconds <= 0.f;
 
-			AnimInstanceProxy.AnimDrawDebugLine(Pos, Pos + X * Scale, IsPast ? FColor::Red : FColor::Blue, false, -1.f, 0.f, SDPG_Foreground);
-			AnimInstanceProxy.AnimDrawDebugLine(Pos, Pos + Y * Scale, IsPast ? FColor::Orange : FColor::Turquoise, false, -1.f, 0.f, SDPG_Foreground);
+			AnimInstanceProxy.AnimDrawDebugLine(Pos, Pos + X * Scale, IsPast ? FColor::Red : FColor::Blue, false, -1.f, DebugThickness, SDPG_Foreground);
+			AnimInstanceProxy.AnimDrawDebugLine(Pos, Pos + Y * Scale, IsPast ? FColor::Orange : FColor::Turquoise, false, -1.f, DebugThickness, SDPG_Foreground);
 
 			if (Index == LastIndex)
 			{
@@ -126,7 +126,7 @@ void FPoseSearchQueryTrajectory::DebugDrawTrajectory(FAnimInstanceProxy& AnimIns
 			}
 			
 			const FVector NextPos = Samples[Index + 1].Position + OffsetVector;
-			AnimInstanceProxy.AnimDrawDebugLine(Pos, NextPos, FColor::Black, false, -1.f, 0.f, SDPG_Foreground);
+			AnimInstanceProxy.AnimDrawDebugLine(Pos, NextPos, FColor::Black, false, -1.f, DebugThickness, SDPG_Foreground);
 		}
 	}
 }
diff --git a/Engine/Plugins/Animation/PoseSearch/Source/Runtime/Public/PoseSearch/PoseSearchTrajectoryTypes.h b/Engine/Plugins/Animation/PoseSearch/Source/Runtime/Public/PoseSearch/PoseSearchTrajectoryTypes.h
--- a/Engine/Plugins/Animation/PoseSearch/Source/Runtime/Public/PoseSearch/PoseSearchTrajectoryTypes.h
+++ b/Engine/Plugins/Animation/PoseSearch/Source/Runtime/Public/PoseSearch/PoseSearchTrajectoryTypes.h
@@ -41,8 +41,8 @@ struct POSESEARCH_API FPoseSearchQueryTrajectory
 	FPoseSearchQueryTrajectorySample GetSampleAtTime(float Time, bool bExtrapolate = false) const;
 	
 #if ENABLE_ANIM_DEBUG
-	void DebugDrawTrajectory(const UWorld* World, float HeightOffset = 0.f) const;
-	void DebugDrawTrajectory(FAnimInstanceProxy& AnimInstanceProxy, float HeightOffset = 0.f) const;
+	void DebugDrawTrajectory(const UWorld* World, const float DebugThickness = 0.f, float HeightOffset = 0.f) const;
+	void DebugDrawTrajectory(FAnimInstanceProxy& AnimInstanceProxy, const float DebugThickness = 0.f, float HeightOffset = 0.f) const;
 #endif // ENABLE_ANIM_DEBUG
 };
 POSESEARCH_API FArchive& operator<<(FArchive& Ar, FPoseSearchQueryTrajectory& Trajectory);
