motion matching - fix for trajectory generation when gravity and collisions are applied

diff --git a/Engine/Plugins/Animation/PoseSearch/Source/Runtime/Private/PoseSearchTrajectoryLibrary.cpp b/Engine/Plugins/Animation/PoseSearch/Source/Runtime/Private/PoseSearchTrajectoryLibrary.cpp
--- a/Engine/Plugins/Animation/PoseSearch/Source/Runtime/Private/PoseSearchTrajectoryLibrary.cpp
+++ b/Engine/Plugins/Animation/PoseSearch/Source/Runtime/Private/PoseSearchTrajectoryLibrary.cpp
@@ -43,7 +43,6 @@ void FPoseSearchTrajectoryData::UpdateData(
 		TrajectoryDataDerived.Acceleration = MoveComp->GetCurrentAcceleration();
 		
 		TrajectoryDataDerived.bStepGroundPrediction = !MoveComp->IsFalling() && !MoveComp->IsFlying();
-		TrajectoryDataDerived.Gravity = -MoveComp->GetGravityDirection() * MoveComp->GetGravityZ();
 
 		if (TrajectoryDataDerived.Acceleration.IsZero())
 		{
@@ -264,8 +263,7 @@ void UPoseSearchTrajectoryLibrary::UpdatePrediction_SimulateCharacterMovement(
 	const FPoseSearchTrajectoryData& TrajectoryData,
 	const FPoseSearchTrajectoryData::FDerived& TrajectoryDataDerived,
 	const FPoseSearchTrajectoryData::FSampling& TrajectoryDataSampling,
-	float DeltaTime,
-	bool bAlwaysApplyGravity)
+	float DeltaTime)
 {
 	FVector CurrentPositionWS = TrajectoryDataDerived.Position;
 	FVector CurrentVelocityWS = RemapVectorMagnitudeWithCurve(TrajectoryDataDerived.Velocity, TrajectoryData.bUseSpeedRemappingCurve, TrajectoryData.SpeedRemappingCurve);
@@ -338,15 +336,6 @@ void UPoseSearchTrajectoryLibrary::UpdatePrediction_SimulateCharacterMovement(
 				}
 			}
 		}
-
-		if (!TrajectoryDataDerived.bStepGroundPrediction || bAlwaysApplyGravity)
-		{
-			for (int32 Index = NumHistorySamples + 1; Index < Trajectory.Samples.Num(); ++Index)
-			{
-				const float Time = Trajectory.Samples[Index].AccumulatedSeconds;
-				Trajectory.Samples[Index].Position += 0.5f * TrajectoryDataDerived.Gravity * (Time * Time);
-			}
-		}
 	}
 }
 
@@ -360,8 +349,7 @@ void UPoseSearchTrajectoryLibrary::PoseSearchGenerateTrajectory(
 	float InHistorySamplingInterval,
 	int32 InTrajectoryHistoryCount,
 	float InPredictionSamplingInterval,
-	int32 InTrajectoryPredictionCount,
-	bool bAlwaysApplyGravity)
+	int32 InTrajectoryPredictionCount)
 {
 	FPoseSearchTrajectoryData::FSampling TrajectoryDataSampling;
 	TrajectoryDataSampling.NumHistorySamples = InTrajectoryHistoryCount;
@@ -376,32 +364,89 @@ void UPoseSearchTrajectoryLibrary::PoseSearchGenerateTrajectory(
 	InTrajectoryData.UpdateData(InDeltaTime, InAnimInstance, TrajectoryDataDerived, TrajectoryDataState);
 	InitTrajectorySamples(InOutTrajectory, InTrajectoryData, TrajectoryDataDerived, TrajectoryDataSampling, InDeltaTime);
 	UpdateHistory_TransformHistory(InOutTrajectory, InTrajectoryData, TrajectoryDataDerived, TrajectoryDataSampling, InDeltaTime);
-	UpdatePrediction_SimulateCharacterMovement(InOutTrajectory, InTrajectoryData, TrajectoryDataDerived, TrajectoryDataSampling, InDeltaTime, bAlwaysApplyGravity);
+	UpdatePrediction_SimulateCharacterMovement(InOutTrajectory, InTrajectoryData, TrajectoryDataDerived, TrajectoryDataSampling, InDeltaTime);
 
 	InOutDesiredControllerYawLastUpdate = TrajectoryDataState.DesiredControllerYawLastUpdate;
 
 	OutTrajectory = InOutTrajectory;
 }
 
-void UPoseSearchTrajectoryLibrary::HandleTrajectoryWorldCollisions(const UObject* WorldContextObject, UPARAM(ref) const FPoseSearchQueryTrajectory& InTrajectory, FPoseSearchQueryTrajectory& OutTrajectory,
-	ETraceTypeQuery TraceChannel, bool bTraceComplex, const TArray<AActor*>& ActorsToIgnore, EDrawDebugTrace::Type DrawDebugType, bool bIgnoreSelf, FLinearColor TraceColor, FLinearColor TraceHitColor, float DrawTime)
+void UPoseSearchTrajectoryLibrary::HandleTrajectoryWorldCollisions(const UObject* WorldContextObject, const UAnimInstance* AnimInstance, UPARAM(ref) const FPoseSearchQueryTrajectory& InTrajectory, bool bApplyGravity, float FloorCollisionsOffset, FPoseSearchQueryTrajectory& OutTrajectory,
+	ETraceTypeQuery TraceChannel, bool bTraceComplex, const TArray<AActor*>& ActorsToIgnore, EDrawDebugTrace::Type DrawDebugType, bool bIgnoreSelf, float MaxObstacleHeight, FLinearColor TraceColor, FLinearColor TraceHitColor, float DrawTime)
 {
 	OutTrajectory = InTrajectory;
 
-	FHitResult OutHit;
-	const int32 NumSamples = OutTrajectory.Samples.Num();
+	TArray<FPoseSearchQueryTrajectorySample>& Samples = OutTrajectory.Samples;
+	const int32 NumSamples = Samples.Num();
+
+	FVector GravityDirection = FVector::ZeroVector;
+	float GravityZ = 0.f;
+	if (bApplyGravity && AnimInstance)
+	{
+		if (const ACharacter* Character = Cast<ACharacter>(AnimInstance->GetOwningActor()))
+		{
+			if (const UCharacterMovementComponent* MoveComp = Character->GetCharacterMovement())
+			{
+				GravityZ = MoveComp->GetGravityZ();
+				GravityDirection = MoveComp->GetGravityDirection();
+			}
+		}
+	}
+
+	if (!FMath::IsNearlyZero(GravityZ))
+	{
+		FVector LastImpactPoint;
+		FVector LastImpactNormal;
+		bool bIsLastImpactValid = false;
+
+		const FVector Gravity = GravityDirection * -GravityZ;
+		float FreeFallAccumulatedSeconds = 0.f;
+		for (int32 SampleIndex = 1; SampleIndex < NumSamples; ++SampleIndex)
+		{
+			FPoseSearchQueryTrajectorySample& Sample = Samples[SampleIndex];
+			if (Sample.AccumulatedSeconds > 0.f)
+			{
+				const int32 PrevSampleIndex = SampleIndex - 1;
+				const FPoseSearchQueryTrajectorySample& PrevSample = Samples[PrevSampleIndex];
+
+				FreeFallAccumulatedSeconds += Sample.AccumulatedSeconds - PrevSample.AccumulatedSeconds;
+
+				// projecting Sample.Position on the HitResult plane and offsetting it by FloorCollisionsOffset
+				if (bIsLastImpactValid)
+				{
+					const FVector DeltaImpactPoint = Sample.Position - LastImpactPoint;
+					Sample.Position += (FloorCollisionsOffset - (DeltaImpactPoint | LastImpactNormal)) * LastImpactNormal;
+				}
+
+				// applying gravity
+				const FVector FreeFallOffset = Gravity * (0.5f * FreeFallAccumulatedSeconds * FreeFallAccumulatedSeconds);
+				Sample.Position += FreeFallOffset;
+
+				FHitResult HitResult;
+				if (FloorCollisionsOffset > 0.f && UKismetSystemLibrary::LineTraceSingle(WorldContextObject, Sample.Position + (GravityDirection * -MaxObstacleHeight), Sample.Position, TraceChannel, bTraceComplex, ActorsToIgnore, DrawDebugType, HitResult, bIgnoreSelf, TraceColor, TraceHitColor, DrawTime))
+				{
+					LastImpactPoint = HitResult.ImpactPoint;
+					LastImpactNormal = HitResult.Normal;
+					bIsLastImpactValid = true;
+
+					Sample.Position = LastImpactPoint + (LastImpactNormal * FloorCollisionsOffset);
+					FreeFallAccumulatedSeconds = 0.f;
+				}
+			}
+		}
+	}
+	else if (FloorCollisionsOffset > 0.f)
+	{
 	for (int32 SampleIndex = 0; SampleIndex < NumSamples; ++SampleIndex)
 	{
 		FPoseSearchQueryTrajectorySample& Sample = OutTrajectory.Samples[SampleIndex];
 		if (Sample.AccumulatedSeconds > 0.f)
 		{
-			FVector End = Sample.Position;
-			FVector Start = Sample.Position;
-			Start.Z += 3000.f;
-
-			if (UKismetSystemLibrary::LineTraceSingle(WorldContextObject, Start, End, TraceChannel, bTraceComplex, ActorsToIgnore, DrawDebugType, OutHit, bIgnoreSelf, TraceColor, TraceHitColor, DrawTime))
+				FHitResult HitResult;
+				if (UKismetSystemLibrary::LineTraceSingle(WorldContextObject, Sample.Position + FVector::UpVector * 3000.f, Sample.Position, TraceChannel, bTraceComplex, ActorsToIgnore, DrawDebugType, HitResult, bIgnoreSelf, TraceColor, TraceHitColor, DrawTime))
 			{
-				Sample.Position.Z = OutHit.ImpactPoint.Z;
+					Sample.Position.Z = HitResult.ImpactPoint.Z + FloorCollisionsOffset;
+				}
 			}
 		}
 	}
diff --git a/Engine/Plugins/Animation/PoseSearch/Source/Runtime/Public/PoseSearch/PoseSearchTrajectoryLibrary.h b/Engine/Plugins/Animation/PoseSearch/Source/Runtime/Public/PoseSearch/PoseSearchTrajectoryLibrary.h
--- a/Engine/Plugins/Animation/PoseSearch/Source/Runtime/Public/PoseSearch/PoseSearchTrajectoryLibrary.h
+++ b/Engine/Plugins/Animation/PoseSearch/Source/Runtime/Public/PoseSearch/PoseSearchTrajectoryLibrary.h
@@ -30,7 +30,6 @@ public:
 
 		FVector Velocity = FVector::ZeroVector;
 		FVector Acceleration = FVector::ZeroVector;
-		FVector Gravity = FVector::ZeroVector;
 		
 		FVector Position = FVector::ZeroVector;
 		FQuat Facing = FQuat::Identity;
@@ -99,18 +98,20 @@ public:
 	static void UpdateHistory_TransformHistory(FPoseSearchQueryTrajectory& Trajectory, const FPoseSearchTrajectoryData& TrajectoryData, const FPoseSearchTrajectoryData::FDerived& TrajectoryDataDerived, const FPoseSearchTrajectoryData::FSampling& TrajectoryDataSampling, float DeltaTime);
 
 	// Update prediction by simulating the movement math for ground locomotion from UCharacterMovementComponent.
-	static void UpdatePrediction_SimulateCharacterMovement(FPoseSearchQueryTrajectory& Trajectory, const FPoseSearchTrajectoryData& TrajectoryData, const FPoseSearchTrajectoryData::FDerived& TrajectoryDataDerived, const FPoseSearchTrajectoryData::FSampling& TrajectoryDataSampling, float DeltaTime, bool bAlwaysApplyGravity = false);
+	static void UpdatePrediction_SimulateCharacterMovement(FPoseSearchQueryTrajectory& Trajectory, const FPoseSearchTrajectoryData& TrajectoryData, const FPoseSearchTrajectoryData::FDerived& TrajectoryDataDerived, const FPoseSearchTrajectoryData::FSampling& TrajectoryDataSampling, float DeltaTime);
 
-	/** Get a Pose History node context from an anim node context (pure) */
+	// Get a Pose History node context from an anim node context (pure)
 	UFUNCTION(BlueprintCallable, Category = "Animation|PoseSearch", meta = (BlueprintThreadSafe, DisplayName = "Pose Search Generate Trajectory"))
 	static void PoseSearchGenerateTrajectory(const UAnimInstance* InAnimInstance, UPARAM(ref) const FPoseSearchTrajectoryData& InTrajectoryData, float InDeltaTime,
 		UPARAM(ref) FPoseSearchQueryTrajectory& InOutTrajectory, UPARAM(ref) float& InOutDesiredControllerYawLastUpdate, FPoseSearchQueryTrajectory& OutTrajectory,
-		float InHistorySamplingInterval = 0.04f, int32 InTrajectoryHistoryCount = 10, float InPredictionSamplingInterval = 0.2f, int32 InTrajectoryPredictionCount = 8,
-		bool bAlwaysApplyGravity = false);
-
-	UFUNCTION(BlueprintCallable, Category="Animation|PoseSearch", meta=(bIgnoreSelf="true", WorldContext="WorldContextObject", AutoCreateRefTerm="ActorsToIgnore", AdvancedDisplay="TraceColor,TraceHitColor,DrawTime"))
-	static void HandleTrajectoryWorldCollisions(const UObject* WorldContextObject, UPARAM(ref) const FPoseSearchQueryTrajectory& InTrajectory, FPoseSearchQueryTrajectory& OutTrajectory,
-		ETraceTypeQuery TraceChannel, bool bTraceComplex, const TArray<AActor*>& ActorsToIgnore, EDrawDebugTrace::Type DrawDebugType, bool bIgnoreSelf, FLinearColor TraceColor = FLinearColor::Red, FLinearColor TraceHitColor = FLinearColor::Green, float DrawTime = 5.0f);
+		float InHistorySamplingInterval = 0.04f, int32 InTrajectoryHistoryCount = 10, float InPredictionSamplingInterval = 0.2f, int32 InTrajectoryPredictionCount = 8);
+
+	// this function takes process InTrajectory and returns the eventually modified OutTrajectory
+	// if bApplyGravity is true, gravity from the UCharacterMovementComponent will be applied
+	// and, if FloorCollisionsOffset > 0, vertical collision will be performed to every sample of the trajectory to have the samples float over the geometry (by FloorCollisionsOffset)
+	UFUNCTION(BlueprintCallable, Category="Animation|PoseSearch", meta=(WorldContext="WorldContextObject", AutoCreateRefTerm="ActorsToIgnore", AdvancedDisplay="TraceChannel,bTraceComplex,ActorsToIgnore,DrawDebugType,bIgnoreSelf,MaxObstacleHeight,TraceColor,TraceHitColor,DrawTime"))
+	static void HandleTrajectoryWorldCollisions(const UObject* WorldContextObject, const UAnimInstance* AnimInstance, UPARAM(ref) const FPoseSearchQueryTrajectory& InTrajectory, bool bApplyGravity, float FloorCollisionsOffset, FPoseSearchQueryTrajectory& OutTrajectory,
+		ETraceTypeQuery TraceChannel, bool bTraceComplex, const TArray<AActor*>& ActorsToIgnore, EDrawDebugTrace::Type DrawDebugType, bool bIgnoreSelf = true, float MaxObstacleHeight = 10000.f, FLinearColor TraceColor = FLinearColor::Red, FLinearColor TraceHitColor = FLinearColor::Green, float DrawTime = 5.0f);
 
 	UFUNCTION(BlueprintPure, meta = (BlueprintThreadSafe), Category="Animation|PoseSearch")
 	static void GetTrajectorySampleAtTime(UPARAM(ref) const FPoseSearchQueryTrajectory& InTrajectory, float Time, FPoseSearchQueryTrajectorySample& OutTrajectorySample, bool bExtrapolate = false);
