motion matching - fix for normalization issues between run/walk and jump/fall databases: adding EPoseSearchDataPreprocessor::NormalizeWithCommonSchema to index all the normalization set databases with the schema of the indexing database

diff --git a/Engine/Plugins/Animation/PoseSearch/Source/Runtime/Private/PoseSearchDerivedData.cpp b/Engine/Plugins/Animation/PoseSearch/Source/Runtime/Private/PoseSearchDerivedData.cpp
--- a/Engine/Plugins/Animation/PoseSearch/Source/Runtime/Private/PoseSearchDerivedData.cpp
+++ b/Engine/Plugins/Animation/PoseSearch/Source/Runtime/Private/PoseSearchDerivedData.cpp
@@ -241,7 +241,7 @@ public:
 		RowMajorVectorMap MeanDeviationsMap(MeanDeviations.GetData(), 1, NumDimensions);
 
 		const EPoseSearchDataPreprocessor DataPreprocessor = ThisSchema->DataPreprocessor;
-		if (SearchIndexBases[ThisSchemaIndex].GetNumPoses() > 0 && (DataPreprocessor == EPoseSearchDataPreprocessor::Normalize || DataPreprocessor == EPoseSearchDataPreprocessor::NormalizeOnlyByDeviation))
+		if (SearchIndexBases[ThisSchemaIndex].GetNumPoses() > 0 && (DataPreprocessor != EPoseSearchDataPreprocessor::None))
 		{
 			FEntriesGroup EntriesGroup;
 
@@ -314,7 +314,7 @@ static void FindValidSequenceIntervals(const UAnimSequenceBase* SequenceBase, FF
 	}
 }
 
-static void InitSearchIndexAssets(FSearchIndexBase& SearchIndex, const UPoseSearchDatabase* Database)
+static void InitSearchIndexAssets(FSearchIndexBase& SearchIndex, const TArray<FInstancedStruct>& DatabaseAnimationAssets, const UPoseSearchSchema* Schema, const FFloatInterval& ExcludeFromDatabaseParameters)
 {
 	using namespace UE::PoseSearch;
 
@@ -323,9 +323,9 @@ static void InitSearchIndexAssets(FSearchIndexBase& SearchIndex, const UPoseSear
 	TArray<FBlendSampleData> BlendSamples;
 
 	int32 TotalPoses = 0;
-	for (int32 AnimationAssetIndex = 0; AnimationAssetIndex < Database->GetAnimationAssets().Num(); ++AnimationAssetIndex)
+	for (int32 AnimationAssetIndex = 0; AnimationAssetIndex < DatabaseAnimationAssets.Num(); ++AnimationAssetIndex)
 	{
-		const FInstancedStruct& DatabaseAssetStruct = Database->GetAnimationAssetStruct(AnimationAssetIndex);
+		const FInstancedStruct& DatabaseAssetStruct = DatabaseAnimationAssets[AnimationAssetIndex];
 		if (const FPoseSearchDatabaseAnimationAssetBase* DatabaseAsset = DatabaseAssetStruct.GetPtr<FPoseSearchDatabaseAnimationAssetBase>())
 		{
 			if (!DatabaseAsset->IsEnabled() || !DatabaseAsset->GetAnimationAsset())
@@ -337,7 +337,6 @@ static void InitSearchIndexAssets(FSearchIndexBase& SearchIndex, const UPoseSear
 			const bool bAddMirrored = DatabaseAsset->GetMirrorOption() == EPoseSearchMirrorOption::MirroredOnly || DatabaseAsset->GetMirrorOption() == EPoseSearchMirrorOption::UnmirroredAndMirrored;
 			const bool bIsLooping = DatabaseAsset->IsLooping();
 			const bool bDisableReselection = DatabaseAsset->IsDisableReselection();
-			const int32 SchemaSampleRate = Database->Schema->SampleRate;
 
 			if (const FPoseSearchDatabaseBlendSpace* DatabaseBlendSpace = DatabaseAssetStruct.GetPtr<FPoseSearchDatabaseBlendSpace>())
 			{
@@ -357,12 +356,12 @@ static void InitSearchIndexAssets(FSearchIndexBase& SearchIndex, const UPoseSear
 
 						const float PlayLength = DatabaseBlendSpace->BlendSpace->GetAnimationLengthFromSampleData(BlendSamples);
 
-						for (int32 PermutationIdx = 0; PermutationIdx < Database->Schema->NumberOfPermutations; ++PermutationIdx)
+						for (int32 PermutationIdx = 0; PermutationIdx < Schema->NumberOfPermutations; ++PermutationIdx)
 						{
 							if (bAddUnmirrored)
 							{
 								const FSearchIndexAsset PoseSearchIndexAsset(AnimationAssetIndex, TotalPoses, false, bIsLooping, 
-									bDisableReselection, FFloatInterval(0.f, PlayLength), SchemaSampleRate, PermutationIdx, BlendParameters);
+									bDisableReselection, FFloatInterval(0.f, PlayLength), Schema->SampleRate, PermutationIdx, BlendParameters);
 								if (PoseSearchIndexAsset.GetNumPoses() > 0)
 								{
 									SearchIndex.Assets.Add(PoseSearchIndexAsset);
@@ -373,7 +372,7 @@ static void InitSearchIndexAssets(FSearchIndexBase& SearchIndex, const UPoseSear
 							if (bAddMirrored)
 							{
 								const FSearchIndexAsset PoseSearchIndexAsset(AnimationAssetIndex, TotalPoses, true, bIsLooping,
-									bDisableReselection, FFloatInterval(0.f, PlayLength), SchemaSampleRate, PermutationIdx, BlendParameters);
+									bDisableReselection, FFloatInterval(0.f, PlayLength), Schema->SampleRate, PermutationIdx, BlendParameters);
 								if (PoseSearchIndexAsset.GetNumPoses() > 0)
 								{
 									SearchIndex.Assets.Add(PoseSearchIndexAsset);
@@ -389,12 +388,12 @@ static void InitSearchIndexAssets(FSearchIndexBase& SearchIndex, const UPoseSear
 				// @todo: support FindValidSequenceIntervals(SequenceBase, DatabaseAsset->GetSamplingRange(), bIsLooping, Database->ExcludeFromDatabaseParameters, ValidRanges);
 				const float PlayLength = DatabaseMultiSequence->GetPlayLength();
 
-				for (int32 PermutationIdx = 0; PermutationIdx < Database->Schema->NumberOfPermutations; ++PermutationIdx)
+				for (int32 PermutationIdx = 0; PermutationIdx < Schema->NumberOfPermutations; ++PermutationIdx)
 				{
 					if (bAddUnmirrored)
 					{
 						const FSearchIndexAsset PoseSearchIndexAsset(AnimationAssetIndex, TotalPoses, false, bIsLooping, 
-							bDisableReselection, FFloatInterval(0.f, PlayLength), SchemaSampleRate, PermutationIdx);
+							bDisableReselection, FFloatInterval(0.f, PlayLength), Schema->SampleRate, PermutationIdx);
 						if (PoseSearchIndexAsset.GetNumPoses() > 0)
 						{
 							SearchIndex.Assets.Add(PoseSearchIndexAsset);
@@ -405,7 +404,7 @@ static void InitSearchIndexAssets(FSearchIndexBase& SearchIndex, const UPoseSear
 					if (bAddMirrored)
 					{
 						const FSearchIndexAsset PoseSearchIndexAsset(AnimationAssetIndex, TotalPoses, true, bIsLooping,
-							bDisableReselection, FFloatInterval(0.f, PlayLength), SchemaSampleRate, PermutationIdx);
+							bDisableReselection, FFloatInterval(0.f, PlayLength), Schema->SampleRate, PermutationIdx);
 						if (PoseSearchIndexAsset.GetNumPoses() > 0)
 						{
 							SearchIndex.Assets.Add(PoseSearchIndexAsset);
@@ -419,15 +418,15 @@ static void InitSearchIndexAssets(FSearchIndexBase& SearchIndex, const UPoseSear
 			{
 				ValidRanges.Reset();
 
-				FindValidSequenceIntervals(SequenceBase, DatabaseAsset->GetSamplingRange(), bIsLooping, Database->ExcludeFromDatabaseParameters, ValidRanges);
+				FindValidSequenceIntervals(SequenceBase, DatabaseAsset->GetSamplingRange(), bIsLooping, ExcludeFromDatabaseParameters, ValidRanges);
 				for (const FFloatRange& Range : ValidRanges)
 				{
-					for (int32 PermutationIdx = 0; PermutationIdx < Database->Schema->NumberOfPermutations; ++PermutationIdx)
+					for (int32 PermutationIdx = 0; PermutationIdx < Schema->NumberOfPermutations; ++PermutationIdx)
 					{
 						if (bAddUnmirrored)
 						{
 							const FSearchIndexAsset PoseSearchIndexAsset(AnimationAssetIndex, TotalPoses, false, bIsLooping,
-								bDisableReselection, FFloatInterval(Range.GetLowerBoundValue(), Range.GetUpperBoundValue()), SchemaSampleRate, PermutationIdx);
+								bDisableReselection, FFloatInterval(Range.GetLowerBoundValue(), Range.GetUpperBoundValue()), Schema->SampleRate, PermutationIdx);
 							if (PoseSearchIndexAsset.GetNumPoses() > 0)
 							{
 								SearchIndex.Assets.Add(PoseSearchIndexAsset);
@@ -438,7 +437,7 @@ static void InitSearchIndexAssets(FSearchIndexBase& SearchIndex, const UPoseSear
 						if (bAddMirrored)
 						{
 							const FSearchIndexAsset PoseSearchIndexAsset(AnimationAssetIndex, TotalPoses, true, bIsLooping,
-								bDisableReselection, FFloatInterval(Range.GetLowerBoundValue(), Range.GetUpperBoundValue()), SchemaSampleRate, PermutationIdx);
+								bDisableReselection, FFloatInterval(Range.GetLowerBoundValue(), Range.GetUpperBoundValue()), Schema->SampleRate, PermutationIdx);
 							if (PoseSearchIndexAsset.GetNumPoses() > 0)
 							{
 								SearchIndex.Assets.Add(PoseSearchIndexAsset);
@@ -467,7 +466,7 @@ static void PreprocessSearchIndexWeights(FSearchIndex& SearchIndex, const UPoseS
 	}
 
 	EPoseSearchDataPreprocessor DataPreprocessor = Schema->DataPreprocessor;
-	if (DataPreprocessor == EPoseSearchDataPreprocessor::Normalize)
+	if (DataPreprocessor == EPoseSearchDataPreprocessor::Normalize || DataPreprocessor == EPoseSearchDataPreprocessor::NormalizeWithCommonSchema)
 	{
 		// normalizing user weights: the idea behind this step is to be able to compare poses from databases using different schemas
 		RowMajorVectorMap MapWeights(SearchIndex.WeightsSqrt.GetData(), 1, NumDimensions);
@@ -484,7 +483,7 @@ static void PreprocessSearchIndexWeights(FSearchIndex& SearchIndex, const UPoseS
 		SearchIndex.WeightsSqrt[Dimension] = FMath::Sqrt(SearchIndex.WeightsSqrt[Dimension]);
 	}
 
-	if (DataPreprocessor == EPoseSearchDataPreprocessor::Normalize || DataPreprocessor == EPoseSearchDataPreprocessor::NormalizeOnlyByDeviation)
+	if (DataPreprocessor != EPoseSearchDataPreprocessor::None)
 	{
 		for (int32 Dimension = 0; Dimension != NumDimensions; ++Dimension)
 		{
@@ -841,9 +840,8 @@ struct FSamplerMapKey
 	FVector BlendParameters = FVector::ZeroVector;
 };
 
-static bool IndexDatabase(FSearchIndexBase& SearchIndexBase, const UPoseSearchDatabase& Database, UE::DerivedData::FRequestOwner& Owner)
+static bool IndexDatabase(FSearchIndexBase& SearchIndexBase, const TArray<FInstancedStruct>& DatabaseAnimationAssets, const UPoseSearchSchema* Schema, const FAssetSamplingContext& SamplingContext, const FFloatInterval& AdditionalExtrapolationTime, UE::DerivedData::FRequestOwner& Owner)
 {
-	const UPoseSearchSchema* Schema = Database.Schema;
 	check(Schema);
 
 	// Prepare samplers for all animation assets.
@@ -853,9 +851,9 @@ static bool IndexDatabase(FSearchIndexBase& SearchIndexBase, const UPoseSearchDa
 	TMap<FSamplerMapKey, int32> SamplerMap;
 	SamplerMap.Reserve(256);
 
-	for (int32 AnimationAssetIndex = 0; AnimationAssetIndex < Database.GetAnimationAssets().Num(); ++AnimationAssetIndex)
+	for (int32 AnimationAssetIndex = 0; AnimationAssetIndex < DatabaseAnimationAssets.Num(); ++AnimationAssetIndex)
 	{
-		const FInstancedStruct& DatabaseAssetStruct = Database.GetAnimationAssetStruct(AnimationAssetIndex);
+		const FInstancedStruct& DatabaseAssetStruct = DatabaseAnimationAssets[AnimationAssetIndex];
 		if (const FPoseSearchDatabaseBlendSpace* DatabaseBlendSpace = DatabaseAssetStruct.GetPtr<FPoseSearchDatabaseBlendSpace>())
 		{
 			if (DatabaseBlendSpace->BlendSpace)
@@ -920,11 +918,9 @@ static bool IndexDatabase(FSearchIndexBase& SearchIndexBase, const UPoseSearchDa
 	for (const TPair<FRole, FBoneContainer>& RoledBoneContainerPair : RoledBoneContainers)
 	{
 		const FRole& Role = RoledBoneContainerPair.Key;
-		RoledMirrorDataCaches.Add(Role).Init(Database.Schema ? Database.Schema->GetMirrorDataTable(Role) : nullptr, RoledBoneContainers[Role]);
+		RoledMirrorDataCaches.Add(Role).Init(Schema ? Schema->GetMirrorDataTable(Role) : nullptr, RoledBoneContainers[Role]);
 	}
 
-	const FAssetSamplingContext SamplingContext(Database);
-
 	FAnimationAssetSamplers TempAssetSamplers;
 	TArray<FBoneContainer> TempBoneContainers;
 	FRoleToIndex TempRoleToIndex;
@@ -935,7 +931,7 @@ static bool IndexDatabase(FSearchIndexBase& SearchIndexBase, const UPoseSearchDa
 		FSearchIndexAsset& SearchIndexAsset = SearchIndexBase.Assets[AssetIdx];
 		check(SearchIndexAsset.GetFirstPoseIdx() == TotalPoses);
 
-		const FInstancedStruct& DatabaseAssetStruct = Database.GetAnimationAssetStruct(SearchIndexAsset);
+		const FInstancedStruct& DatabaseAssetStruct = DatabaseAnimationAssets[SearchIndexAsset.GetSourceAssetIdx()];
 		const FPoseSearchDatabaseAnimationAssetBase* DatabaseAnimationAssetBase = DatabaseAssetStruct.GetPtr<FPoseSearchDatabaseAnimationAssetBase>();
 		check(DatabaseAnimationAssetBase);
 
@@ -960,12 +956,12 @@ static bool IndexDatabase(FSearchIndexBase& SearchIndexBase, const UPoseSearchDa
 			}
 			else
 			{
-				UE_LOG(LogPoseSearch, Warning, TEXT("Role '%s' from asset '%s' contained in database '%s' cannot be found in Schema::Skeletons '%s'"), *Role.ToString(), *DatabaseAnimationAssetBase->GetName(), *Database.GetName(), *Database.Schema->GetName());
+				UE_LOG(LogPoseSearch, Warning, TEXT("Role '%s' from asset '%s' cannot be found in Skeletons of Schema '%s'"), *Role.ToString(), *DatabaseAnimationAssetBase->GetName(), *Schema->GetName());
 				return false;
 			}
 		}
 
-		const FFloatInterval ExtrapolationTimeInterval = SearchIndexAsset.GetExtrapolationTimeInterval(Schema->SampleRate, Database.AdditionalExtrapolationTime);
+		const FFloatInterval ExtrapolationTimeInterval = SearchIndexAsset.GetExtrapolationTimeInterval(Schema->SampleRate, AdditionalExtrapolationTime);
 		Indexers.Emplace(TempBoneContainers, SearchIndexAsset, SamplingContext, *Schema, TempAssetSamplers, TempRoleToIndex, ExtrapolationTimeInterval);
 		TotalPoses += SearchIndexAsset.GetNumPoses();
 	}
@@ -1236,6 +1232,7 @@ private:
 
 	void OnGetComplete(UE::DerivedData::FCacheGetResponse&& Response);
 	void SetState(EState State) { ThreadSafeState.Set(int32(State)); }
+	void ResetSearchIndex();
 
 	TWeakObjectPtr<UPoseSearchDatabase> Database;
 	FSearchIndex SearchIndex;
@@ -1277,13 +1274,18 @@ FPoseSearchDatabaseAsyncCacheTask::~FPoseSearchDatabaseAsyncCacheTask()
 
 	Database = nullptr;
 
+	ResetSearchIndex();
+
+	DerivedDataKey = FIoHash::Zero;
+	DatabaseDependencies.Reset();
+}
+
+void FPoseSearchDatabaseAsyncCacheTask::ResetSearchIndex()
+{
 	SearchIndex.Reset();
 #if ENABLE_ANIM_DEBUG
 	SearchIndexCompare.Reset();
 #endif //ENABLE_ANIM_DEBUG
-
-	DerivedDataKey = FIoHash::Zero;
-	DatabaseDependencies.Reset();
 }
 
 #if ENABLE_ANIM_DEBUG
@@ -1379,10 +1381,7 @@ void FPoseSearchDatabaseAsyncCacheTask::Cancel()
 	// Owner.Cancel must be performed before SearchIndex.Reset() in case any task is flying (launched by Owner.LaunchTask)
 	Owner.Cancel();
 
-	SearchIndex.Reset();
-#if ENABLE_ANIM_DEBUG
-	SearchIndexCompare.Reset();
-#endif //ENABLE_ANIM_DEBUG
+	ResetSearchIndex();
 
 	DerivedDataKey = FIoHash::Zero;
 	SetState(EState::Cancelled);
@@ -1438,10 +1437,7 @@ void FPoseSearchDatabaseAsyncCacheTask::Wait(FCriticalSection& OuterMutex)
 		check(!bBroadcastOnDerivedDataRebuild);
 		SetState(EState::Failed);
 	}
-	SearchIndex.Reset();
-#if ENABLE_ANIM_DEBUG
-	SearchIndexCompare.Reset();
-#endif //ENABLE_ANIM_DEBUG
+	ResetSearchIndex();
 }
 
 // true is the task is done executing
@@ -1471,10 +1467,7 @@ void FPoseSearchDatabaseAsyncCacheTask::OnGetComplete(UE::DerivedData::FCacheGet
 		COOK_STAT(auto Timer = UsageStats.TimeAsyncWait());
 
 		// we found the cached data associated to the PendingDerivedDataKey: we'll deserialized into SearchIndex
-		SearchIndex.Reset();
-#if ENABLE_ANIM_DEBUG
-		SearchIndexCompare.Reset();
-#endif //ENABLE_ANIM_DEBUG
+		ResetSearchIndex();
 
 		FSharedBuffer RawData = Response.Record.GetValue(Id).GetData().Decompress();
 		FMemoryReaderView Reader(RawData);
@@ -1499,10 +1492,7 @@ void FPoseSearchDatabaseAsyncCacheTask::OnGetComplete(UE::DerivedData::FCacheGet
 	
 	if (Response.Status == EStatus::Canceled)
 	{
-		SearchIndex.Reset();
-#if ENABLE_ANIM_DEBUG
-		SearchIndexCompare.Reset();
-#endif //ENABLE_ANIM_DEBUG
+		ResetSearchIndex();
 		UE_LOG(LogPoseSearch, Log, TEXT("%s - %s BuildIndex Cancelled"), *LexToString(FullIndexKey.Hash), *Database->GetName());
 	}
 	
@@ -1527,24 +1517,35 @@ void FPoseSearchDatabaseAsyncCacheTask::OnGetComplete(UE::DerivedData::FCacheGet
 			{
 				COOK_STAT(auto Timer = UsageStats.TimeSyncWork());
 
+				const UPoseSearchDatabase* MainDatabase = Database.Get();
+
 				// collecting all the databases that need to be built to gather their FSearchIndexBase
+				// the first one is always the main database (the one we're calculating the index on)
 				TArray<const UPoseSearchDatabase*> IndexBaseDatabases;
-				IndexBaseDatabases.Add(Database.Get()); // the first one is always this Database
-				if (!IndexBaseDatabases[0])
+				IndexBaseDatabases.Add(MainDatabase);
+				if (!MainDatabase)
 				{
 					UE_LOG(LogPoseSearch, Log, TEXT("%s - BuildIndex Cancelled because associated Database weak pointer has been released."), *LexToString(FullIndexKey.Hash));
-					SearchIndex.Reset();
-#if ENABLE_ANIM_DEBUG
-					SearchIndexCompare.Reset();
-#endif //ENABLE_ANIM_DEBUG
+					ResetSearchIndex();
 					return;
 				}
 
-				if (IndexBaseDatabases[0]->NormalizationSet)
+				if (MainDatabase->NormalizationSet)
 				{
-					IndexBaseDatabases[0]->NormalizationSet->AddUniqueDatabases(IndexBaseDatabases);
+					MainDatabase->NormalizationSet->AddUniqueDatabases(IndexBaseDatabases);
+				}
+
+				const FString MainDatabaseName = MainDatabase->GetName();
+				const UPoseSearchSchema* MainDatabaseSchema = MainDatabase->Schema;
+				if (!MainDatabaseSchema || MainDatabaseSchema->SchemaCardinality <= 0)
+				{
+					UE_LOG(LogPoseSearch, Error, TEXT("%s - %s BuildIndex Failed because of invalid Schema"), *LexToString(FullIndexKey.Hash), *MainDatabaseName);
+					ResetSearchIndex();
+					return;
 				}
 
+				const bool bNormalizeWithCommonSchema = MainDatabaseSchema->DataPreprocessor == EPoseSearchDataPreprocessor::NormalizeWithCommonSchema;
+
 				// @todo: DDC or parallelize this code
 				TArray<FSearchIndexBase> SearchIndexBases;
 				TArray<const UPoseSearchSchema*> Schemas;
@@ -1552,44 +1553,47 @@ void FPoseSearchDatabaseAsyncCacheTask::OnGetComplete(UE::DerivedData::FCacheGet
 				Schemas.AddDefaulted(IndexBaseDatabases.Num());
 				for (int32 IndexBaseIdx = 0; IndexBaseIdx < IndexBaseDatabases.Num(); ++IndexBaseIdx)
 				{
-					const UPoseSearchDatabase* IndexBaseDatabase = IndexBaseDatabases[IndexBaseIdx];
-					check(IndexBaseDatabase);
+					const UPoseSearchDatabase* DependentDatabase = IndexBaseDatabases[IndexBaseIdx];
+					check(DependentDatabase);
+
+					const UPoseSearchSchema* DependentDatabaseSchema = bNormalizeWithCommonSchema ? MainDatabaseSchema : DependentDatabase->Schema.Get();
+					const FString DependentDatabaseName = DependentDatabase->GetName();
+
 					FSearchIndexBase& SearchIndexBase = SearchIndexBases[IndexBaseIdx];
-					Schemas[IndexBaseIdx] = IndexBaseDatabase->Schema;
+
+					Schemas[IndexBaseIdx] = DependentDatabaseSchema;
+					const TArray<FInstancedStruct>& DependentDatabaseAnimationAssets = DependentDatabase->GetAnimationAssets();
+					const FAssetSamplingContext DependentSamplingContext(*DependentDatabase);
+					const FFloatInterval& DependentExcludeFromDatabaseParameters = DependentDatabase->ExcludeFromDatabaseParameters;
+					const FFloatInterval& DependentAdditionalExtrapolationTime = DependentDatabase->AdditionalExtrapolationTime;
 
 					// early out for invalid indexing conditions
-					if (!IndexBaseDatabase->Schema || IndexBaseDatabase->Schema->SchemaCardinality <= 0)
+					if (!DependentDatabaseSchema || DependentDatabaseSchema->SchemaCardinality <= 0)
 					{
 						if (IndexBaseIdx == 0)
 						{
-							UE_LOG(LogPoseSearch, Error, TEXT("%s - %s BuildIndex Failed because of invalid Schema"), *LexToString(FullIndexKey.Hash), *IndexBaseDatabases[0]->GetName());
+							UE_LOG(LogPoseSearch, Error, TEXT("%s - %s BuildIndex Failed because of invalid Schema"), *LexToString(FullIndexKey.Hash), *MainDatabaseName);
 						}
 						else
 						{
-							UE_LOG(LogPoseSearch, Error, TEXT("%s - %s BuildIndex Failed because dependent database '%s' has an invalid Schema"), *LexToString(FullIndexKey.Hash), *IndexBaseDatabases[0]->GetName(), *IndexBaseDatabase->GetName());
+							UE_LOG(LogPoseSearch, Error, TEXT("%s - %s BuildIndex Failed because dependent database '%s' has an invalid Schema"), *LexToString(FullIndexKey.Hash), *MainDatabaseName, *DependentDatabaseName);
 						}
-						SearchIndex.Reset();
-#if ENABLE_ANIM_DEBUG
-						SearchIndexCompare.Reset();
-#endif //ENABLE_ANIM_DEBUG
+						ResetSearchIndex();
 						return;
 					}
 
 					// validating that the missing MirrorDataTable(s) are not necessary 
-					if (!IndexBaseDatabase->Schema->AllRoledSkeletonHaveMirrorDataTable())
+					if (!DependentDatabaseSchema->AllRoledSkeletonHaveMirrorDataTable())
 					{
-						for (int32 AnimationAssetIndex = 0; AnimationAssetIndex < IndexBaseDatabase->GetAnimationAssets().Num(); ++AnimationAssetIndex)
+						for (int32 AnimationAssetIndex = 0; AnimationAssetIndex < DependentDatabaseAnimationAssets.Num(); ++AnimationAssetIndex)
 						{
-							if (const FPoseSearchDatabaseAnimationAssetBase* DatabaseAsset = IndexBaseDatabase->GetAnimationAssetBase(AnimationAssetIndex))
+							if (const FPoseSearchDatabaseAnimationAssetBase* DatabaseAsset = DependentDatabaseAnimationAssets[AnimationAssetIndex].GetPtr<FPoseSearchDatabaseAnimationAssetBase>())
 							{
 								if (DatabaseAsset->GetMirrorOption() == EPoseSearchMirrorOption::MirroredOnly || DatabaseAsset->GetMirrorOption() == EPoseSearchMirrorOption::UnmirroredAndMirrored)
 								{
 									// want to sample a mirrored asset
-									UE_LOG(LogPoseSearch, Error, TEXT("%s - %s BuildIndex Failed because '%s' requires MirrorDataTable(s) to sample mirrored animation assets"), *LexToString(FullIndexKey.Hash), *IndexBaseDatabases[0]->GetName(), *IndexBaseDatabase->Schema->GetName());
-									SearchIndex.Reset();
-#if ENABLE_ANIM_DEBUG
-									SearchIndexCompare.Reset();
-#endif //ENABLE_ANIM_DEBUG
+									UE_LOG(LogPoseSearch, Error, TEXT("%s - %s BuildIndex Failed because '%s' requires MirrorDataTable(s) to sample mirrored animation assets"), *LexToString(FullIndexKey.Hash), *MainDatabaseName, *DependentDatabaseName);
+									ResetSearchIndex();
 									return;
 								}
 							}
@@ -1598,34 +1602,25 @@ void FPoseSearchDatabaseAsyncCacheTask::OnGetComplete(UE::DerivedData::FCacheGet
 
 					if (Owner.IsCanceled())
 					{
-						UE_LOG(LogPoseSearch, Log, TEXT("%s - %s BuildIndex Cancelled"), *LexToString(FullIndexKey.Hash), *IndexBaseDatabases[0]->GetName());
-						SearchIndex.Reset();
-#if ENABLE_ANIM_DEBUG
-						SearchIndexCompare.Reset();
-#endif //ENABLE_ANIM_DEBUG
+						UE_LOG(LogPoseSearch, Log, TEXT("%s - %s BuildIndex Cancelled"), *LexToString(FullIndexKey.Hash), *MainDatabaseName);
+						ResetSearchIndex();
 						return;
 					}
 
 					// Building all the related FPoseSearchBaseIndex first
-					InitSearchIndexAssets(SearchIndexBase, IndexBaseDatabase);
+					InitSearchIndexAssets(SearchIndexBase, DependentDatabaseAnimationAssets, DependentDatabaseSchema, DependentExcludeFromDatabaseParameters);
 
 					if (Owner.IsCanceled())
 					{
-						UE_LOG(LogPoseSearch, Log, TEXT("%s - %s BuildIndex Cancelled"), *LexToString(FullIndexKey.Hash), *IndexBaseDatabases[0]->GetName());
-						SearchIndex.Reset();
-#if ENABLE_ANIM_DEBUG
-						SearchIndexCompare.Reset();
-#endif //ENABLE_ANIM_DEBUG
+						UE_LOG(LogPoseSearch, Log, TEXT("%s - %s BuildIndex Cancelled"), *LexToString(FullIndexKey.Hash), *MainDatabaseName);
+						ResetSearchIndex();
 						return;
 					}
 
-					if (!IndexDatabase(SearchIndexBase, *IndexBaseDatabase, Owner))
+					if (!IndexDatabase(SearchIndexBase, DependentDatabaseAnimationAssets, DependentDatabaseSchema, DependentSamplingContext, DependentAdditionalExtrapolationTime, Owner))
 					{
-						UE_LOG(LogPoseSearch, Log, TEXT("%s - %s BuildIndex Cancelled"), *LexToString(FullIndexKey.Hash), *IndexBaseDatabases[0]->GetName());
-						SearchIndex.Reset();
-#if ENABLE_ANIM_DEBUG
-						SearchIndexCompare.Reset();
-#endif //ENABLE_ANIM_DEBUG
+						UE_LOG(LogPoseSearch, Log, TEXT("%s - %s BuildIndex Cancelled"), *LexToString(FullIndexKey.Hash), *MainDatabaseName);
+						ResetSearchIndex();
 						return;
 					}
 
@@ -1636,12 +1631,12 @@ void FPoseSearchDatabaseAsyncCacheTask::OnGetComplete(UE::DerivedData::FCacheGet
 						for (int32 Iteration = 0; Iteration < NumIterations; ++Iteration)
 						{
 							FSearchIndexBase TestSearchIndexBase = SearchIndexBase;
-							if (IndexDatabase(TestSearchIndexBase, *IndexBaseDatabase, Owner))
+							if (IndexDatabase(TestSearchIndexBase, DependentDatabaseAnimationAssets, DependentDatabaseSchema, DependentSamplingContext, DependentAdditionalExtrapolationTime, Owner))
 							{
 								if (TestSearchIndexBase != SearchIndexBase)
 								{
 									FStringBuilderBase Message;
-									CompareSearchIndexBase(TestSearchIndexBase, SearchIndexBase, IndexBaseDatabase->Schema, Message);
+									CompareSearchIndexBase(TestSearchIndexBase, SearchIndexBase, DependentDatabaseSchema, Message);
 									UE_LOG(LogPoseSearch, Warning, TEXT("OnGetComplete - IndexDatabase is not deterministic\n%s"), *Message);
 								}
 							}
@@ -1659,11 +1654,11 @@ void FPoseSearchDatabaseAsyncCacheTask::OnGetComplete(UE::DerivedData::FCacheGet
 					const int32 NumIterations = CVarMotionMatchTestNumIterations.GetValueOnAnyThread();
 
 					FSearchIndex TestSearchIndexA = SearchIndex;
-					TestSearchIndexA.PruneDuplicateValues(IndexBaseDatabases[0]->PosePruningSimilarityThreshold, IndexBaseDatabases[0]->Schema->SchemaCardinality, false);
+					TestSearchIndexA.PruneDuplicateValues(MainDatabase->PosePruningSimilarityThreshold, MainDatabaseSchema->SchemaCardinality, false);
 					for (int32 Iteration = 0; Iteration < NumIterations; ++Iteration)
 					{
 						FSearchIndex TestSearchIndexB = SearchIndex;
-						TestSearchIndexB.PruneDuplicateValues(IndexBaseDatabases[0]->PosePruningSimilarityThreshold, IndexBaseDatabases[0]->Schema->SchemaCardinality, false);
+						TestSearchIndexB.PruneDuplicateValues(MainDatabase->PosePruningSimilarityThreshold, MainDatabaseSchema->SchemaCardinality, false);
 
 						if (TestSearchIndexA.Values != TestSearchIndexB.Values)
 						{
@@ -1679,24 +1674,21 @@ void FPoseSearchDatabaseAsyncCacheTask::OnGetComplete(UE::DerivedData::FCacheGet
 #endif // ENABLE_ANIM_DEBUG
 
 				// VPTree requires ValuesVectorToPoseIndexes if there's any Values pruning
-				const bool bDoNotGenerateValuesVectorToPoseIndexes = IndexBaseDatabases[0]->PoseSearchMode != EPoseSearchMode::VPTree;
-				SearchIndex.PruneDuplicateValues(IndexBaseDatabases[0]->PosePruningSimilarityThreshold, IndexBaseDatabases[0]->Schema->SchemaCardinality, bDoNotGenerateValuesVectorToPoseIndexes);
+				const bool bDoNotGenerateValuesVectorToPoseIndexes = MainDatabase->PoseSearchMode != EPoseSearchMode::VPTree;
+				SearchIndex.PruneDuplicateValues(MainDatabase->PosePruningSimilarityThreshold, MainDatabaseSchema->SchemaCardinality, bDoNotGenerateValuesVectorToPoseIndexes);
 
 				TArray<float> Deviation = FMeanDeviationCalculator::Calculate(SearchIndexBases, Schemas);
 
 				// Building FSearchIndex
-				PreprocessSearchIndexWeights(SearchIndex, IndexBaseDatabases[0]->Schema, Deviation);
+				PreprocessSearchIndexWeights(SearchIndex, MainDatabaseSchema, Deviation);
 				if (Owner.IsCanceled())
 				{
-					UE_LOG(LogPoseSearch, Log, TEXT("%s - %s BuildIndex Cancelled"), *LexToString(FullIndexKey.Hash), *IndexBaseDatabases[0]->GetName());
-					SearchIndex.Reset();
-#if ENABLE_ANIM_DEBUG
-					SearchIndexCompare.Reset();
-#endif //ENABLE_ANIM_DEBUG
+					UE_LOG(LogPoseSearch, Log, TEXT("%s - %s BuildIndex Cancelled"), *LexToString(FullIndexKey.Hash), *MainDatabaseName);
+					ResetSearchIndex();
 					return;
 				}
 
-				const Eigen::ComputationInfo ComputationInfo = PreprocessSearchIndexPCAData(SearchIndex, IndexBaseDatabases[0]->Schema->SchemaCardinality, IndexBaseDatabases[0]->GetNumberOfPrincipalComponents(), IndexBaseDatabases[0]->PoseSearchMode);
+				const Eigen::ComputationInfo ComputationInfo = PreprocessSearchIndexPCAData(SearchIndex, MainDatabaseSchema->SchemaCardinality, MainDatabase->GetNumberOfPrincipalComponents(), MainDatabase->PoseSearchMode);
 				if (ComputationInfo != Eigen::Success)
 				{
 					FString Reason;
@@ -1707,22 +1699,16 @@ void FPoseSearchDatabaseAsyncCacheTask::OnGetComplete(UE::DerivedData::FCacheGet
 					case Eigen::InvalidInput:	Reason = "Invalid Input"; 	 break;
 					default:					Reason = "Unknown Reasons";  break;
 					}
-					UE_LOG(LogPoseSearch, Error, TEXT("%s - %s BuildIndex Failed because of '%s' while calculating PCA data. Try with a different dataset or change the database 'Pose Search Mode'"), *LexToString(FullIndexKey.Hash), *IndexBaseDatabases[0]->GetName(), *Reason);
+					UE_LOG(LogPoseSearch, Error, TEXT("%s - %s BuildIndex Failed because of '%s' while calculating PCA data. Try with a different dataset or change the database 'Pose Search Mode'"), *LexToString(FullIndexKey.Hash), *MainDatabaseName, *Reason);
 
-					SearchIndex.Reset();
-#if ENABLE_ANIM_DEBUG
-					SearchIndexCompare.Reset();
-#endif //ENABLE_ANIM_DEBUG
+					ResetSearchIndex();
 					return;
 				}
 
 				if (Owner.IsCanceled())
 				{
-					UE_LOG(LogPoseSearch, Log, TEXT("%s - %s BuildIndex Cancelled"), *LexToString(FullIndexKey.Hash), *IndexBaseDatabases[0]->GetName());
-					SearchIndex.Reset();
-#if ENABLE_ANIM_DEBUG
-					SearchIndexCompare.Reset();
-#endif //ENABLE_ANIM_DEBUG
+					UE_LOG(LogPoseSearch, Log, TEXT("%s - %s BuildIndex Cancelled"), *LexToString(FullIndexKey.Hash), *MainDatabaseName);
+					ResetSearchIndex();
 					return;
 				}
 
@@ -1733,11 +1719,11 @@ void FPoseSearchDatabaseAsyncCacheTask::OnGetComplete(UE::DerivedData::FCacheGet
 					const int32 NumIterations = CVarMotionMatchTestNumIterations.GetValueOnAnyThread();
 
 					FSearchIndex TestSearchIndexA = SearchIndex;
-					TestSearchIndexA.PruneDuplicatePCAValues(IndexBaseDatabases[0]->PCAValuesPruningSimilarityThreshold, IndexBaseDatabases[0]->GetNumberOfPrincipalComponents());
+					TestSearchIndexA.PruneDuplicatePCAValues(MainDatabase->PCAValuesPruningSimilarityThreshold, MainDatabase->GetNumberOfPrincipalComponents());
 					for (int32 Iteration = 0; Iteration < NumIterations; ++Iteration)
 					{
 						FSearchIndex TestSearchIndexB = SearchIndex;
-						TestSearchIndexB.PruneDuplicatePCAValues(IndexBaseDatabases[0]->PCAValuesPruningSimilarityThreshold, IndexBaseDatabases[0]->GetNumberOfPrincipalComponents());
+						TestSearchIndexB.PruneDuplicatePCAValues(MainDatabase->PCAValuesPruningSimilarityThreshold, MainDatabase->GetNumberOfPrincipalComponents());
 
 						if (TestSearchIndexA.PCAValues != TestSearchIndexB.PCAValues)
 						{
@@ -1752,62 +1738,50 @@ void FPoseSearchDatabaseAsyncCacheTask::OnGetComplete(UE::DerivedData::FCacheGet
 				}
 #endif // ENABLE_ANIM_DEBUG
 
-				SearchIndex.PruneDuplicatePCAValues(IndexBaseDatabases[0]->PCAValuesPruningSimilarityThreshold, IndexBaseDatabases[0]->GetNumberOfPrincipalComponents());
+				SearchIndex.PruneDuplicatePCAValues(MainDatabase->PCAValuesPruningSimilarityThreshold, MainDatabase->GetNumberOfPrincipalComponents());
 				if (Owner.IsCanceled())
 				{
-					UE_LOG(LogPoseSearch, Log, TEXT("%s - %s BuildIndex Cancelled"), *LexToString(FullIndexKey.Hash), *IndexBaseDatabases[0]->GetName());
-					SearchIndex.Reset();
-#if ENABLE_ANIM_DEBUG
-					SearchIndexCompare.Reset();
-#endif //ENABLE_ANIM_DEBUG
+					UE_LOG(LogPoseSearch, Log, TEXT("%s - %s BuildIndex Cancelled"), *LexToString(FullIndexKey.Hash), *MainDatabaseName);
+					ResetSearchIndex();
 					return;
 				}
 
-				if (IndexBaseDatabases[0]->bPCAValuesPruningFromBlockTransitionPoses)
+				if (MainDatabase->bPCAValuesPruningFromBlockTransitionPoses)
 				{
-					SearchIndex.PrunePCAValuesFromBlockTransitionPoses(IndexBaseDatabases[0]->GetNumberOfPrincipalComponents());
+					SearchIndex.PrunePCAValuesFromBlockTransitionPoses(MainDatabase->GetNumberOfPrincipalComponents());
 				}
 
-				PreprocessSearchIndexKDTree(SearchIndex, IndexBaseDatabases[0]);
+				PreprocessSearchIndexKDTree(SearchIndex, MainDatabase);
 				if (Owner.IsCanceled())
 				{
-					UE_LOG(LogPoseSearch, Log, TEXT("%s - %s BuildIndex Cancelled"), *LexToString(FullIndexKey.Hash), *IndexBaseDatabases[0]->GetName());
-					SearchIndex.Reset();
-#if ENABLE_ANIM_DEBUG
-					SearchIndexCompare.Reset();
-#endif //ENABLE_ANIM_DEBUG
+					UE_LOG(LogPoseSearch, Log, TEXT("%s - %s BuildIndex Cancelled"), *LexToString(FullIndexKey.Hash), *MainDatabaseName);
+					ResetSearchIndex();
 					return;
 				}
 
 				// removing SearchIndex.Values and relying on FSearchIndex::GetReconstructedPoseValues to reconstruct the Values data from the PCAValues
-				if (IndexBaseDatabases[0]->PoseSearchMode == EPoseSearchMode::PCAKDTree && IndexBaseDatabases[0]->KDTreeQueryNumNeighbors <= 1)
+				if (MainDatabase->PoseSearchMode == EPoseSearchMode::PCAKDTree && MainDatabase->KDTreeQueryNumNeighbors <= 1)
 				{
 					SearchIndex.ResetValues();
 				}
 
 				const int32 RandomSeed = GetTypeHash(FullIndexKey.Hash);
-				PreprocessSearchIndexVPTree(SearchIndex, IndexBaseDatabases[0], RandomSeed);
+				PreprocessSearchIndexVPTree(SearchIndex, MainDatabase, RandomSeed);
 				if (Owner.IsCanceled())
 				{
-					UE_LOG(LogPoseSearch, Log, TEXT("%s - %s BuildIndex Cancelled"), *LexToString(FullIndexKey.Hash), *IndexBaseDatabases[0]->GetName());
-					SearchIndex.Reset();
-#if ENABLE_ANIM_DEBUG
-					SearchIndexCompare.Reset();
-#endif //ENABLE_ANIM_DEBUG
+					UE_LOG(LogPoseSearch, Log, TEXT("%s - %s BuildIndex Cancelled"), *LexToString(FullIndexKey.Hash), *MainDatabaseName);
+					ResetSearchIndex();
 					return;
 				}
 
-				UE_LOG(LogPoseSearch, Log, TEXT("%s - %s BuildIndex Succeeded"), *LexToString(FullIndexKey.Hash), *IndexBaseDatabases[0]->GetName());
+				UE_LOG(LogPoseSearch, Log, TEXT("%s - %s BuildIndex Succeeded"), *LexToString(FullIndexKey.Hash), *MainDatabaseName);
 
 #if ENABLE_ANIM_DEBUG
-				if (bCompareSearchIndex)
+				if (bCompareSearchIndex && SearchIndexCompare != SearchIndex)
 				{
-					if (SearchIndexCompare != SearchIndex)
-					{
-						FStringBuilderBase Message;
-						CompareSearchIndex(SearchIndexCompare, SearchIndex, IndexBaseDatabases[0]->Schema, Message);
-						UE_LOG(LogPoseSearch, Warning, TEXT("%s - %s BuildIndex mismatch with DDC Index\n%s"), *LexToString(FullIndexKey.Hash), *IndexBaseDatabases[0]->GetName(), *Message);
-					}
+					FStringBuilderBase Message;
+					CompareSearchIndex(SearchIndexCompare, SearchIndex, MainDatabaseSchema, Message);
+					UE_LOG(LogPoseSearch, Warning, TEXT("%s - %s BuildIndex mismatch with DDC Index\n%s"), *LexToString(FullIndexKey.Hash), *MainDatabaseName, *Message);
 				}
 #endif // ENABLE_ANIM_DEBUG
 
@@ -1821,18 +1795,18 @@ void FPoseSearchDatabaseAsyncCacheTask::OnGetComplete(UE::DerivedData::FCacheGet
 				FCacheRecordBuilder Builder(FullIndexKey);
 				Builder.AddValue(Id, RawData);
 
-				GetCache().Put({ { { IndexBaseDatabases[0]->GetPathName() }, Builder.Build() } }, Owner, [this, IndexBaseDatabases, FullIndexKey](FCachePutResponse&& Response)
+				GetCache().Put({ { { MainDatabase->GetPathName() }, Builder.Build() } }, Owner, [this, MainDatabaseSchema, MainDatabaseName, FullIndexKey](FCachePutResponse&& Response)
 					{
 						switch (Response.Status)
 						{
 						case EStatus::Error:
-							UE_LOG(LogPoseSearch, Log, TEXT("%s - %s Failed to store DDC"), *LexToString(FullIndexKey.Hash), *IndexBaseDatabases[0]->GetName());
+							UE_LOG(LogPoseSearch, Log, TEXT("%s - %s Failed to store DDC"), *LexToString(FullIndexKey.Hash), *MainDatabaseName);
 							break;
 						case EStatus::Canceled:
-							UE_LOG(LogPoseSearch, Log, TEXT("%s - %s Canceled to store DDC"), *LexToString(FullIndexKey.Hash), *IndexBaseDatabases[0]->GetName());
+							UE_LOG(LogPoseSearch, Log, TEXT("%s - %s Canceled to store DDC"), *LexToString(FullIndexKey.Hash), *MainDatabaseName);
 							break;
 						case EStatus::Ok:
-							UE_LOG(LogPoseSearch, Log, TEXT("%s - %s BuildIndex stored to DDC"), *LexToString(FullIndexKey.Hash), *IndexBaseDatabases[0]->GetName());
+							UE_LOG(LogPoseSearch, Log, TEXT("%s - %s BuildIndex stored to DDC"), *LexToString(FullIndexKey.Hash), *MainDatabaseName);
 
 #if ENABLE_ANIM_DEBUG
 							if (AnyTestFlags(EMotionMatchTestFlags::ValidateDDC))
@@ -1840,7 +1814,7 @@ void FPoseSearchDatabaseAsyncCacheTask::OnGetComplete(UE::DerivedData::FCacheGet
 								TArray<FCacheGetRequest> CacheRequests;
 								const FCacheKey CacheKey{ Bucket, DerivedDataKey };
 								CacheRequests.Add({ { Database->GetPathName() }, CacheKey, ECachePolicy::Default });
-								GetCache().Get(CacheRequests, Owner, [this, IndexBaseDatabases](FCacheGetResponse&& Response)
+								GetCache().Get(CacheRequests, Owner, [this, MainDatabaseSchema, MainDatabaseName](FCacheGetResponse&& Response)
 									{
 										const FCacheKey& FullIndexKey = Response.Record.GetKey();
 										check(FullIndexKey.Hash == DerivedDataKey);
@@ -1856,8 +1830,8 @@ void FPoseSearchDatabaseAsyncCacheTask::OnGetComplete(UE::DerivedData::FCacheGet
 											if (TestSearchIndex != SearchIndex)
 											{
 												FStringBuilderBase Message;
-												CompareSearchIndex(TestSearchIndex, SearchIndex, IndexBaseDatabases[0]->Schema, Message);
-												UE_LOG(LogPoseSearch, Warning, TEXT("%s - %s DDC Index mismatch with BuildIndex\n%s"), *LexToString(FullIndexKey.Hash), *IndexBaseDatabases[0]->GetName(), *Message);
+												CompareSearchIndex(TestSearchIndex, SearchIndex, MainDatabaseSchema, Message);
+												UE_LOG(LogPoseSearch, Warning, TEXT("%s - %s DDC Index mismatch with BuildIndex\n%s"), *LexToString(FullIndexKey.Hash), *MainDatabaseName, *Message);
 											}
 										}
 									});
diff --git a/Engine/Plugins/Animation/PoseSearch/Source/Runtime/Private/PoseSearchDerivedDataKey.cpp b/Engine/Plugins/Animation/PoseSearch/Source/Runtime/Private/PoseSearchDerivedDataKey.cpp
--- a/Engine/Plugins/Animation/PoseSearch/Source/Runtime/Private/PoseSearchDerivedDataKey.cpp
+++ b/Engine/Plugins/Animation/PoseSearch/Source/Runtime/Private/PoseSearchDerivedDataKey.cpp
@@ -41,7 +41,7 @@ FKeyBuilder::FKeyBuilder(const UObject* Object, bool bUseDataVer, bool bPerformC
 	if (bUseDataVer)
 	{
 		// used to invalidate the key without having to change POSESEARCHDB_DERIVEDDATA_VER all the times
-		int32 POSESEARCHDB_DERIVEDDATA_VER_SMALL = 239;
+		int32 POSESEARCHDB_DERIVEDDATA_VER_SMALL = 247;
 		FGuid VersionGuid = FDevSystemGuids::GetSystemGuid(FDevSystemGuids::Get().POSESEARCHDB_DERIVEDDATA_VER);
 
 		*this << VersionGuid;
diff --git a/Engine/Plugins/Animation/PoseSearch/Source/Runtime/Public/PoseSearch/PoseSearchSchema.h b/Engine/Plugins/Animation/PoseSearch/Source/Runtime/Public/PoseSearch/PoseSearchSchema.h
--- a/Engine/Plugins/Animation/PoseSearch/Source/Runtime/Public/PoseSearch/PoseSearchSchema.h
+++ b/Engine/Plugins/Animation/PoseSearch/Source/Runtime/Public/PoseSearch/PoseSearchSchema.h
@@ -20,7 +20,10 @@ enum class EPoseSearchDataPreprocessor : int32
 	Normalize,
 
 	// The data will be normalized against its deviation
-	NormalizeOnlyByDeviation
+	NormalizeOnlyByDeviation UMETA(DisplayName = "Normalize Only By Deviation (Experimental)"),
+
+	// same behavior as Normalize, but it'll index all the databases in the normalization set with the same schema
+	NormalizeWithCommonSchema UMETA(DisplayName = "Normalize With Common Schema (Experimental)"),
 };
 
 
