motion matching - fix for Static Analysis PoseSearchDerivedData.cpp(914): warning V595: The 'Schema' pointer was utilized before it was verified against nullptr. Check lines: 914, 921.

diff --git a/Engine/Plugins/Animation/PoseSearch/Source/Runtime/Private/PoseSearchDerivedData.cpp b/Engine/Plugins/Animation/PoseSearch/Source/Runtime/Private/PoseSearchDerivedData.cpp
--- a/Engine/Plugins/Animation/PoseSearch/Source/Runtime/Private/PoseSearchDerivedData.cpp
+++ b/Engine/Plugins/Animation/PoseSearch/Source/Runtime/Private/PoseSearchDerivedData.cpp
@@ -842,7 +842,10 @@ struct FSamplerMapKey
 
 static bool IndexDatabase(FSearchIndexBase& SearchIndexBase, const TArray<FInstancedStruct>& DatabaseAnimationAssets, const UPoseSearchSchema* Schema, const FAssetSamplingContext& SamplingContext, const FFloatInterval& AdditionalExtrapolationTime, UE::DerivedData::FRequestOwner& Owner)
 {
-	check(Schema);
+	if (!ensure(Schema))
+	{
+		return false;
+	}
 
 	// Prepare samplers for all animation assets.
 	TArray<FAnimationAssetSampler> Samplers;
@@ -918,7 +921,7 @@ static bool IndexDatabase(FSearchIndexBase& SearchIndexBase, const TArray<FInsta
 	for (const TPair<FRole, FBoneContainer>& RoledBoneContainerPair : RoledBoneContainers)
 	{
 		const FRole& Role = RoledBoneContainerPair.Key;
-		RoledMirrorDataCaches.Add(Role).Init(Schema ? Schema->GetMirrorDataTable(Role) : nullptr, RoledBoneContainers[Role]);
+		RoledMirrorDataCaches.Add(Role).Init(Schema->GetMirrorDataTable(Role), RoledBoneContainers[Role]);
 	}
 
 	FAnimationAssetSamplers TempAssetSamplers;
