motion matching - adding normalization groups to channels to be able to normalize, for example, feet positions and velocities together

diff --git a/Engine/Plugins/Animation/PoseSearch/Source/Runtime/Private/PoseSearchFeatureChannel.cpp b/Engine/Plugins/Animation/PoseSearch/Source/Runtime/Private/PoseSearchFeatureChannel.cpp
--- a/Engine/Plugins/Animation/PoseSearch/Source/Runtime/Private/PoseSearchFeatureChannel.cpp
+++ b/Engine/Plugins/Animation/PoseSearch/Source/Runtime/Private/PoseSearchFeatureChannel.cpp
@@ -227,6 +227,11 @@ bool UPoseSearchFeatureChannel::CanBeNormalizedWith(const UPoseSearchFeatureChan
 		return false;
 	}
 
+	if (!GetNormalizationGroup().IsNone() && GetNormalizationGroup() == Other->GetNormalizationGroup())
+	{
+		return true;
+	}
+
 	TLabelBuilder ThisLabelBuilder, OtherLabelBuilder;
 	if (GetLabel(ThisLabelBuilder).ToString() != Other->GetLabel(OtherLabelBuilder).ToString())
 	{
diff --git a/Engine/Plugins/Animation/PoseSearch/Source/Runtime/Private/PoseSearchFeatureChannel_Heading.h b/Engine/Plugins/Animation/PoseSearch/Source/Runtime/Private/PoseSearchFeatureChannel_Heading.h
--- a/Engine/Plugins/Animation/PoseSearch/Source/Runtime/Private/PoseSearchFeatureChannel_Heading.h
+++ b/Engine/Plugins/Animation/PoseSearch/Source/Runtime/Private/PoseSearchFeatureChannel_Heading.h
@@ -80,6 +80,14 @@ public:
 	UPROPERTY(EditAnywhere, Category = "Settings")
 	EPermutationTimeType PermutationTimeType = EPermutationTimeType::UseSampleTime;
 
+#if WITH_EDITORONLY_DATA
+	// if set, all the channels of the same class with the same cardinality, and the same NormalizationGroup, will be normalized together.
+	// for example in a locomotion database of a character holding a weapon, containing non mirrorable animations, you'd still want to normalize togeter 
+	// left foot and right foot position and velocity
+	UPROPERTY(EditAnywhere, Category = "Settings")
+	FName NormalizationGroup;
+#endif //WITH_EDITORONLY_DATA
+
 	UFUNCTION(BlueprintPure, BlueprintImplementableEvent, meta=(BlueprintThreadSafe, DisplayName = "Get World Rotation"), Category = "Settings")
 	FQuat BP_GetWorldRotation(const UAnimInstance* AnimInstance) const;
 
@@ -102,6 +110,7 @@ public:
 	virtual void FillWeights(TArrayView<float> Weights) const override;
 	virtual bool IndexAsset(UE::PoseSearch::FAssetIndexer& Indexer) const override;
 	virtual UE::PoseSearch::TLabelBuilder& GetLabel(UE::PoseSearch::TLabelBuilder& LabelBuilder, UE::PoseSearch::ELabelFormat LabelFormat = UE::PoseSearch::ELabelFormat::Full_Horizontal) const override;
+	virtual FName GetNormalizationGroup() const override { return NormalizationGroup; }
 
 	// IBoneReferenceSkeletonProvider interface
 	USkeleton* GetSkeleton(bool& bInvalidSkeletonIsError, const IPropertyHandle* PropertyHandle) override;
diff --git a/Engine/Plugins/Animation/PoseSearch/Source/Runtime/Private/PoseSearchFeatureChannel_Phase.h b/Engine/Plugins/Animation/PoseSearch/Source/Runtime/Private/PoseSearchFeatureChannel_Phase.h
--- a/Engine/Plugins/Animation/PoseSearch/Source/Runtime/Private/PoseSearchFeatureChannel_Phase.h
+++ b/Engine/Plugins/Animation/PoseSearch/Source/Runtime/Private/PoseSearchFeatureChannel_Phase.h
@@ -35,6 +35,14 @@ public:
 	UPROPERTY(EditAnywhere, Category = "Settings")
 	EInputQueryPose InputQueryPose = EInputQueryPose::UseContinuingPose;
 
+#if WITH_EDITORONLY_DATA
+	// if set, all the channels of the same class with the same cardinality, and the same NormalizationGroup, will be normalized together.
+	// for example in a locomotion database of a character holding a weapon, containing non mirrorable animations, you'd still want to normalize togeter 
+	// left foot and right foot position and velocity
+	UPROPERTY(EditAnywhere, Category = "Settings")
+	FName NormalizationGroup;
+#endif //WITH_EDITORONLY_DATA
+
 	// UPoseSearchFeatureChannel interface
 	virtual bool Finalize(UPoseSearchSchema* Schema) override;
 	virtual void BuildQuery(UE::PoseSearch::FSearchContext& SearchContext) const override;
@@ -49,6 +57,7 @@ public:
 	virtual void FillWeights(TArrayView<float> Weights) const override;
 	virtual bool IndexAsset(UE::PoseSearch::FAssetIndexer& Indexer) const override;
 	virtual UE::PoseSearch::TLabelBuilder& GetLabel(UE::PoseSearch::TLabelBuilder& LabelBuilder, UE::PoseSearch::ELabelFormat LabelFormat = UE::PoseSearch::ELabelFormat::Full_Horizontal) const override;
+	virtual FName GetNormalizationGroup() const override { return NormalizationGroup; }
 	virtual const UE::PoseSearch::FRole GetDefaultRole() const override { return SampleRole; }
 #endif
 };
\ No newline at end of file
diff --git a/Engine/Plugins/Animation/PoseSearch/Source/Runtime/Private/PoseSearchFeatureChannel_Pose.cpp b/Engine/Plugins/Animation/PoseSearch/Source/Runtime/Private/PoseSearchFeatureChannel_Pose.cpp
--- a/Engine/Plugins/Animation/PoseSearch/Source/Runtime/Private/PoseSearchFeatureChannel_Pose.cpp
+++ b/Engine/Plugins/Animation/PoseSearch/Source/Runtime/Private/PoseSearchFeatureChannel_Pose.cpp
@@ -20,13 +20,13 @@ UPoseSearchFeatureChannel_Pose::UPoseSearchFeatureChannel_Pose()
 
 	SampledBones.Add(FPoseSearchBone({ {"foot_l"}, int32(EPoseSearchBoneFlags::Position | EPoseSearchBoneFlags::Velocity)
 #if WITH_EDITORONLY_DATA
-		, 1.f, FLinearColor::Green
+		, FName(), 1.f, FLinearColor::Green
 #endif // WITH_EDITORONLY_DATA
 		}));
 
 	SampledBones.Add(FPoseSearchBone({ {"foot_r"}, int32(EPoseSearchBoneFlags::Position | EPoseSearchBoneFlags::Velocity)
 #if WITH_EDITORONLY_DATA
-		, 1.f, FLinearColor::Green
+		, FName(), 1.f, FLinearColor::Green
 #endif // WITH_EDITORONLY_DATA
 		}));
 }
@@ -45,6 +45,7 @@ bool UPoseSearchFeatureChannel_Pose::Finalize(UPoseSearchSchema* Schema)
 			Position->SampleRole = SampleRole;
 			Position->OriginRole = SampleRole;
 #if WITH_EDITORONLY_DATA
+			Position->NormalizationGroup = SampledBone.NormalizationGroup;
 			Position->Weight = SampledBone.Weight * Weight;
 			Position->DebugColor = SampledBone.DebugColor;
 #endif // WITH_EDITORONLY_DATA
@@ -60,6 +61,7 @@ bool UPoseSearchFeatureChannel_Pose::Finalize(UPoseSearchSchema* Schema)
 			HeadingX->SampleRole = SampleRole;
 			HeadingX->OriginRole = SampleRole;
 #if WITH_EDITORONLY_DATA
+			HeadingX->NormalizationGroup = SampledBone.NormalizationGroup;
 			HeadingX->Weight = SampledBone.Weight * Weight;
 			HeadingX->DebugColor = SampledBone.DebugColor;
 #endif // WITH_EDITORONLY_DATA
@@ -73,6 +75,7 @@ bool UPoseSearchFeatureChannel_Pose::Finalize(UPoseSearchSchema* Schema)
 			HeadingY->SampleRole = SampleRole;
 			HeadingY->OriginRole = SampleRole;
 #if WITH_EDITORONLY_DATA
+			HeadingY->NormalizationGroup = SampledBone.NormalizationGroup;
 			HeadingY->Weight = SampledBone.Weight * Weight;
 			HeadingY->DebugColor = SampledBone.DebugColor;
 #endif // WITH_EDITORONLY_DATA
@@ -89,6 +92,7 @@ bool UPoseSearchFeatureChannel_Pose::Finalize(UPoseSearchSchema* Schema)
 			Velocity->SampleRole = SampleRole;
 			Velocity->OriginRole = SampleRole;
 #if WITH_EDITORONLY_DATA
+			Velocity->NormalizationGroup = SampledBone.NormalizationGroup;
 			Velocity->Weight = SampledBone.Weight * Weight;
 			Velocity->DebugColor = SampledBone.DebugColor;
 #endif // WITH_EDITORONLY_DATA
@@ -104,6 +108,7 @@ bool UPoseSearchFeatureChannel_Pose::Finalize(UPoseSearchSchema* Schema)
 			Phase->Bone = SampledBone.Reference;
 			Phase->SampleRole = SampleRole;
 #if WITH_EDITORONLY_DATA
+			Phase->NormalizationGroup = SampledBone.NormalizationGroup;
 			Phase->Weight = SampledBone.Weight * Weight;
 			Phase->DebugColor = SampledBone.DebugColor;
 #endif // WITH_EDITORONLY_DATA
diff --git a/Engine/Plugins/Animation/PoseSearch/Source/Runtime/Private/PoseSearchFeatureChannel_Pose.h b/Engine/Plugins/Animation/PoseSearch/Source/Runtime/Private/PoseSearchFeatureChannel_Pose.h
--- a/Engine/Plugins/Animation/PoseSearch/Source/Runtime/Private/PoseSearchFeatureChannel_Pose.h
+++ b/Engine/Plugins/Animation/PoseSearch/Source/Runtime/Private/PoseSearchFeatureChannel_Pose.h
@@ -30,6 +30,12 @@ struct POSESEARCH_API FPoseSearchBone
 	int32 Flags = int32(EPoseSearchBoneFlags::Position);
 
 #if WITH_EDITORONLY_DATA
+	// if set, all the channels of the same class with the same cardinality, and the same NormalizationGroup, will be normalized together.
+	// for example in a locomotion database of a character holding a weapon, containing non mirrorable animations, you'd still want to normalize togeter 
+	// left foot and right foot position and velocity
+	UPROPERTY(EditAnywhere, Category = "Settings")
+	FName NormalizationGroup;
+
 	UPROPERTY(EditAnywhere, Category = Config)
 	float Weight = 1.f;
 
diff --git a/Engine/Plugins/Animation/PoseSearch/Source/Runtime/Private/PoseSearchFeatureChannel_Position.h b/Engine/Plugins/Animation/PoseSearch/Source/Runtime/Private/PoseSearchFeatureChannel_Position.h
--- a/Engine/Plugins/Animation/PoseSearch/Source/Runtime/Private/PoseSearchFeatureChannel_Position.h
+++ b/Engine/Plugins/Animation/PoseSearch/Source/Runtime/Private/PoseSearchFeatureChannel_Position.h
@@ -67,6 +67,14 @@ public:
 	UPROPERTY(EditAnywhere, Category = "Settings")
 	EPermutationTimeType PermutationTimeType = EPermutationTimeType::UseSampleTime;
 
+#if WITH_EDITORONLY_DATA
+	// if set, all the channels of the same class with the same cardinality, and the same NormalizationGroup, will be normalized together.
+	// for example in a locomotion database of a character holding a weapon, containing non mirrorable animations, you'd still want to normalize togeter 
+	// left foot and right foot position and velocity
+	UPROPERTY(EditAnywhere, Category = "Settings")
+	FName NormalizationGroup;
+#endif //WITH_EDITORONLY_DATA
+
 	UFUNCTION(BlueprintPure, BlueprintImplementableEvent, meta=(BlueprintThreadSafe, DisplayName = "Get World Position"), Category = "Settings")
 	FVector BP_GetWorldPosition(const UAnimInstance* AnimInstance) const;
 
@@ -89,6 +97,7 @@ public:
 	virtual void FillWeights(TArrayView<float> Weights) const override;
 	virtual bool IndexAsset(UE::PoseSearch::FAssetIndexer& Indexer) const override;
 	virtual UE::PoseSearch::TLabelBuilder& GetLabel(UE::PoseSearch::TLabelBuilder& LabelBuilder, UE::PoseSearch::ELabelFormat LabelFormat = UE::PoseSearch::ELabelFormat::Full_Horizontal) const override;
+	virtual FName GetNormalizationGroup() const override { return NormalizationGroup; }
 
 	// IBoneReferenceSkeletonProvider interface
 	USkeleton* GetSkeleton(bool& bInvalidSkeletonIsError, const IPropertyHandle* PropertyHandle) override;
diff --git a/Engine/Plugins/Animation/PoseSearch/Source/Runtime/Private/PoseSearchFeatureChannel_Velocity.h b/Engine/Plugins/Animation/PoseSearch/Source/Runtime/Private/PoseSearchFeatureChannel_Velocity.h
--- a/Engine/Plugins/Animation/PoseSearch/Source/Runtime/Private/PoseSearchFeatureChannel_Velocity.h
+++ b/Engine/Plugins/Animation/PoseSearch/Source/Runtime/Private/PoseSearchFeatureChannel_Velocity.h
@@ -75,6 +75,14 @@ public:
 	UPROPERTY(EditAnywhere, Category = "Settings")
 	EPermutationTimeType PermutationTimeType = EPermutationTimeType::UseSampleTime;
 
+#if WITH_EDITORONLY_DATA
+	// if set, all the channels of the same class with the same cardinality, and the same NormalizationGroup, will be normalized together.
+	// for example in a locomotion database of a character holding a weapon, containing non mirrorable animations, you'd still want to normalize togeter 
+	// left foot and right foot position and velocity
+	UPROPERTY(EditAnywhere, Category = "Settings")
+	FName NormalizationGroup;
+#endif //WITH_EDITORONLY_DATA
+
 	UFUNCTION(BlueprintPure, BlueprintImplementableEvent, meta=(BlueprintThreadSafe, DisplayName = "Get World Velocity"), Category = "Settings")
 	FVector BP_GetWorldVelocity(const UAnimInstance* AnimInstance) const;
 
@@ -97,6 +105,7 @@ public:
 	virtual void FillWeights(TArrayView<float> Weights) const override;
 	virtual bool IndexAsset(UE::PoseSearch::FAssetIndexer& Indexer) const override;
 	virtual UE::PoseSearch::TLabelBuilder& GetLabel(UE::PoseSearch::TLabelBuilder& LabelBuilder, UE::PoseSearch::ELabelFormat LabelFormat = UE::PoseSearch::ELabelFormat::Full_Horizontal) const override;
+	virtual FName GetNormalizationGroup() const override { return NormalizationGroup; }
 	
 	// IBoneReferenceSkeletonProvider interface
 	USkeleton* GetSkeleton(bool& bInvalidSkeletonIsError, const IPropertyHandle* PropertyHandle) override;
diff --git a/Engine/Plugins/Animation/PoseSearch/Source/Runtime/Public/PoseSearch/PoseSearchFeatureChannel.h b/Engine/Plugins/Animation/PoseSearch/Source/Runtime/Public/PoseSearch/PoseSearchFeatureChannel.h
--- a/Engine/Plugins/Animation/PoseSearch/Source/Runtime/Public/PoseSearch/PoseSearchFeatureChannel.h
+++ b/Engine/Plugins/Animation/PoseSearch/Source/Runtime/Public/PoseSearch/PoseSearchFeatureChannel.h
@@ -152,7 +152,15 @@ public:
 	// returns the TLabelBuilder used editor side to identify this UPoseSearchFeatureChannel (for instance in the pose search debugger)
 	virtual UE::PoseSearch::TLabelBuilder& GetLabel(UE::PoseSearch::TLabelBuilder& LabelBuilder, UE::PoseSearch::ELabelFormat LabelFormat = UE::PoseSearch::ELabelFormat::Full_Horizontal) const;
 
+	// returns true if the data associated to this channel can be normalized toghether with the associated with Other
 	virtual bool CanBeNormalizedWith(const UPoseSearchFeatureChannel* Other) const;
+
+	// if this channel GetNormalizationGroup returns a valid FName, all the channels of the same class with the same cardinality, and the same NormalizationGroup, 
+	// will make CanBeNormalizedWith return true and will be normalized together.
+	// for example in a locomotion database of a character holding a weapon, containing non mirrorable animations, you'd still want to normalize togeter 
+	// left foot and right foot position and velocity, so you'd want those channels returning the same GetNormalizationGroup value
+	virtual FName GetNormalizationGroup() const { return FName(); }
+
 	const UPoseSearchSchema* GetSchema() const;
 
 	virtual const UE::PoseSearch::FRole GetDefaultRole() const;
