motion matching - fix for MM node generating a T-pose when the kdtree search returns only block transition poses (removing UPoseSearchDatabase::bPCAValuesPruningFromBlockTransitionPoses, by defaulting its value to true)

diff --git a/Engine/Plugins/Animation/PoseSearch/Source/Runtime/Private/PoseSearchDatabase.cpp b/Engine/Plugins/Animation/PoseSearch/Source/Runtime/Private/PoseSearchDatabase.cpp
--- a/Engine/Plugins/Animation/PoseSearch/Source/Runtime/Private/PoseSearchDatabase.cpp
+++ b/Engine/Plugins/Animation/PoseSearch/Source/Runtime/Private/PoseSearchDatabase.cpp
@@ -1302,8 +1302,9 @@ UE::PoseSearch::FSearchResult UPoseSearchDatabase::SearchPCAKDTree(UE::PoseSearc
 		}
 #endif // WITH_EDITOR && ENABLE_ANIM_DEBUG
 
-		// NonSelectableIdx are already filtered out inside the kdtree search
-		const FSearchFilters SearchFilters(Schema, bRunNonSelectableIdxPostKDTree ? NonSelectableIdx : TConstArrayView<int32>(), SelectableAssetIdx, SearchIndex.bAnyBlockTransition && !bPCAValuesPruningFromBlockTransitionPoses);
+		// NonSelectableIdx are already filtered out inside the kdtree search.
+		// Also kdtrees don't contain block transition poses by construction, so FSearchFilters input bAddBlockTransitionFilter can be set to false
+		const FSearchFilters SearchFilters(Schema, bRunNonSelectableIdxPostKDTree ? NonSelectableIdx : TConstArrayView<int32>(), SelectableAssetIdx, false);
 		
 		// are the PCAValues pruned out of duplicates (multiple poses are associated with the same PCAValuesVectorIdx)
 		if (bArePCAValuesPruned)
diff --git a/Engine/Plugins/Animation/PoseSearch/Source/Runtime/Private/PoseSearchDerivedData.cpp b/Engine/Plugins/Animation/PoseSearch/Source/Runtime/Private/PoseSearchDerivedData.cpp
--- a/Engine/Plugins/Animation/PoseSearch/Source/Runtime/Private/PoseSearchDerivedData.cpp
+++ b/Engine/Plugins/Animation/PoseSearch/Source/Runtime/Private/PoseSearchDerivedData.cpp
@@ -1749,10 +1749,7 @@ void FPoseSearchDatabaseAsyncCacheTask::OnGetComplete(UE::DerivedData::FCacheGet
 					return;
 				}
 
-				if (MainDatabase->bPCAValuesPruningFromBlockTransitionPoses)
-				{
-					SearchIndex.PrunePCAValuesFromBlockTransitionPoses(MainDatabase->GetNumberOfPrincipalComponents());
-				}
+				SearchIndex.PrunePCAValuesFromBlockTransitionPoses(MainDatabase->GetNumberOfPrincipalComponents());
 
 				PreprocessSearchIndexKDTree(SearchIndex, MainDatabase);
 				if (Owner.IsCanceled())
diff --git a/Engine/Plugins/Animation/PoseSearch/Source/Runtime/Public/PoseSearch/PoseSearchDatabase.h b/Engine/Plugins/Animation/PoseSearch/Source/Runtime/Public/PoseSearch/PoseSearchDatabase.h
--- a/Engine/Plugins/Animation/PoseSearch/Source/Runtime/Public/PoseSearch/PoseSearchDatabase.h
+++ b/Engine/Plugins/Animation/PoseSearch/Source/Runtime/Public/PoseSearch/PoseSearchDatabase.h
@@ -364,11 +364,6 @@ public:
 	UPROPERTY(EditAnywhere, Category = "Performance", meta = (DisplayName = "KNNQueryNumNeighborsWithDuplicates", EditCondition = "PoseSearchMode == EPoseSearchMode::PCAKDTree && PCAValuesPruningSimilarityThreshold > 0", EditConditionHides, ClampMin = "0", ClampMax = "600", UIMin = "1"))
 	int32 KDTreeQueryNumNeighborsWithDuplicates = 0;
 	
-	// pruning block transition poses from the PCA values (the KDTree will be smaller and faster). This will force pose filtering to be done after tree search,
-	// and in case all the KNNQueryNumNeighbors requested poses get filtered out, the search could end up with no valid results
-	UPROPERTY(EditAnywhere, Category = "Performance", meta = (EditCondition = "PoseSearchMode == EPoseSearchMode::PCAKDTree", EditConditionHides))
-	bool bPCAValuesPruningFromBlockTransitionPoses = false;
-
 private:
 	// Do not use it directly. Use GetSearchIndex / SetSearchIndex interact with it and validate that is ok to do so.
 	UE::PoseSearch::FSearchIndex SearchIndexPrivate;
