motion matching - fix for FDatabaseEditor::FindOrAddSelectionWidget when DetailsContainer is null

diff --git a/Engine/Plugins/Animation/PoseSearch/Source/Editor/Private/PoseSearchDatabaseEditor.cpp b/Engine/Plugins/Animation/PoseSearch/Source/Editor/Private/PoseSearchDatabaseEditor.cpp
--- a/Engine/Plugins/Animation/PoseSearch/Source/Editor/Private/PoseSearchDatabaseEditor.cpp
+++ b/Engine/Plugins/Animation/PoseSearch/Source/Editor/Private/PoseSearchDatabaseEditor.cpp
@@ -636,11 +636,14 @@ namespace UE::PoseSearch
 			this,
 			&FDatabaseEditor::OnFinishedChangingSelectionProperties);
 
-		DetailsContainer->AddSlot()
-			.AutoHeight()
-			[
-				SelectionWidget.DetailView.ToSharedRef()
-			];
+		if (DetailsContainer)
+		{
+			DetailsContainer->AddSlot()
+				.AutoHeight()
+				[
+					SelectionWidget.DetailView.ToSharedRef()
+				];
+		}
 
 		return SelectionWidget;
 	}
