Added TimeToLand and LandSpeed based on trajectory collision prediction
#rnx

diff --git a/Engine/Plugins/Animation/PoseSearch/Source/Runtime/Private/PoseSearchTrajectoryLibrary.cpp b/Engine/Plugins/Animation/PoseSearch/Source/Runtime/Private/PoseSearchTrajectoryLibrary.cpp
--- a/Engine/Plugins/Animation/PoseSearch/Source/Runtime/Private/PoseSearchTrajectoryLibrary.cpp
+++ b/Engine/Plugins/Animation/PoseSearch/Source/Runtime/Private/PoseSearchTrajectoryLibrary.cpp
@@ -371,7 +371,7 @@ void UPoseSearchTrajectoryLibrary::PoseSearchGenerateTrajectory(
 	OutTrajectory = InOutTrajectory;
 }
 
-void UPoseSearchTrajectoryLibrary::HandleTrajectoryWorldCollisions(const UObject* WorldContextObject, const UAnimInstance* AnimInstance, UPARAM(ref) const FPoseSearchQueryTrajectory& InTrajectory, bool bApplyGravity, float FloorCollisionsOffset, FPoseSearchQueryTrajectory& OutTrajectory,
+void UPoseSearchTrajectoryLibrary::HandleTrajectoryWorldCollisions(const UObject* WorldContextObject, const UAnimInstance* AnimInstance, UPARAM(ref) const FPoseSearchQueryTrajectory& InTrajectory, bool bApplyGravity, float FloorCollisionsOffset, FPoseSearchQueryTrajectory& OutTrajectory, FPoseSearchTrajectory_WorldCollisionResults& CollisionResult,
 	ETraceTypeQuery TraceChannel, bool bTraceComplex, const TArray<AActor*>& ActorsToIgnore, EDrawDebugTrace::Type DrawDebugType, bool bIgnoreSelf, float MaxObstacleHeight, FLinearColor TraceColor, FLinearColor TraceHitColor, float DrawTime)
 {
 	OutTrajectory = InTrajectory;
@@ -381,6 +381,7 @@ void UPoseSearchTrajectoryLibrary::HandleTrajectoryWorldCollisions(const UObject
 
 	FVector GravityDirection = FVector::ZeroVector;
 	float GravityZ = 0.f;
+	float InitialVelocityZ = 0.0f;
 	if (bApplyGravity && AnimInstance)
 	{
 		if (const ACharacter* Character = Cast<ACharacter>(AnimInstance->GetOwningActor()))
@@ -389,15 +390,18 @@ void UPoseSearchTrajectoryLibrary::HandleTrajectoryWorldCollisions(const UObject
 			{
 				GravityZ = MoveComp->GetGravityZ();
 				GravityDirection = MoveComp->GetGravityDirection();
+				InitialVelocityZ = Character->GetVelocity().Z;
 			}
 		}
 	}
+	CollisionResult.TimeToLand = OutTrajectory.Samples.Last().AccumulatedSeconds;
 
 	if (!FMath::IsNearlyZero(GravityZ))
 	{
 		FVector LastImpactPoint;
 		FVector LastImpactNormal;
 		bool bIsLastImpactValid = false;
+		bool bIsFirstFall = true;
 
 		const FVector Gravity = GravityDirection * -GravityZ;
 		float FreeFallAccumulatedSeconds = 0.f;
@@ -430,6 +434,19 @@ void UPoseSearchTrajectoryLibrary::HandleTrajectoryWorldCollisions(const UObject
 					bIsLastImpactValid = true;
 
 					Sample.Position = LastImpactPoint + (LastImpactNormal * FloorCollisionsOffset);
+
+					if (bIsFirstFall)
+					{
+						const float InitialHeight = OutTrajectory.GetSampleAtTime(0.0f).Position.Z;
+						const float FinalHeight = Sample.Position.Z;
+						const float FallHeight = FMath::Abs(FinalHeight - InitialHeight);
+
+						bIsFirstFall = false;
+						CollisionResult.TimeToLand = (InitialVelocityZ / -GravityZ) + ((FMath::Sqrt(FMath::Square(InitialVelocityZ) + (2.f * -GravityZ * FallHeight))) / -GravityZ);
+						CollisionResult.LandSpeed = InitialVelocityZ + GravityZ * CollisionResult.TimeToLand;
+						CollisionResult.LandHitResult = HitResult;
+					}
+
 					FreeFallAccumulatedSeconds = 0.f;
 				}
 			}
@@ -450,6 +467,8 @@ void UPoseSearchTrajectoryLibrary::HandleTrajectoryWorldCollisions(const UObject
 			}
 		}
 	}
+
+	CollisionResult.LandSpeed = InitialVelocityZ + GravityZ * CollisionResult.TimeToLand;
 }
 
 void UPoseSearchTrajectoryLibrary::GetTrajectorySampleAtTime(UPARAM(ref) const FPoseSearchQueryTrajectory& InTrajectory, float Time, FPoseSearchQueryTrajectorySample& OutTrajectorySample, bool bExtrapolate)
diff --git a/Engine/Plugins/Animation/PoseSearch/Source/Runtime/Public/PoseSearch/PoseSearchTrajectoryLibrary.h b/Engine/Plugins/Animation/PoseSearch/Source/Runtime/Public/PoseSearch/PoseSearchTrajectoryLibrary.h
--- a/Engine/Plugins/Animation/PoseSearch/Source/Runtime/Public/PoseSearch/PoseSearchTrajectoryLibrary.h
+++ b/Engine/Plugins/Animation/PoseSearch/Source/Runtime/Public/PoseSearch/PoseSearchTrajectoryLibrary.h
@@ -81,6 +81,21 @@ public:
 	FRuntimeFloatCurve AccelerationRemappingCurve;
 };
 
+USTRUCT(BlueprintType)
+struct POSESEARCH_API FPoseSearchTrajectory_WorldCollisionResults
+{
+	GENERATED_BODY()
+
+	UPROPERTY(BlueprintReadOnly, VisibleAnywhere, Category = "Trajectory Settings")
+	float TimeToLand  = 0.0f;
+
+	UPROPERTY(BlueprintReadOnly, VisibleAnywhere, Category = "Trajectory Settings")
+	float LandSpeed  = 0.0f;
+
+	UPROPERTY(BlueprintReadOnly, VisibleAnywhere, Category = "Trajectory Settings")
+	FHitResult LandHitResult;
+};
+
 /**
  * Set of functions to help populate a FPoseSearchQueryTrajectory for motion matching.
  */
@@ -110,7 +125,7 @@ public:
 	// if bApplyGravity is true, gravity from the UCharacterMovementComponent will be applied
 	// and, if FloorCollisionsOffset > 0, vertical collision will be performed to every sample of the trajectory to have the samples float over the geometry (by FloorCollisionsOffset)
 	UFUNCTION(BlueprintCallable, Category="Animation|PoseSearch", meta=(WorldContext="WorldContextObject", AutoCreateRefTerm="ActorsToIgnore", AdvancedDisplay="TraceChannel,bTraceComplex,ActorsToIgnore,DrawDebugType,bIgnoreSelf,MaxObstacleHeight,TraceColor,TraceHitColor,DrawTime"))
-	static void HandleTrajectoryWorldCollisions(const UObject* WorldContextObject, const UAnimInstance* AnimInstance, UPARAM(ref) const FPoseSearchQueryTrajectory& InTrajectory, bool bApplyGravity, float FloorCollisionsOffset, FPoseSearchQueryTrajectory& OutTrajectory,
+	static void HandleTrajectoryWorldCollisions(const UObject* WorldContextObject, const UAnimInstance* AnimInstance, UPARAM(ref) const FPoseSearchQueryTrajectory& InTrajectory, bool bApplyGravity, float FloorCollisionsOffset, FPoseSearchQueryTrajectory& OutTrajectory, FPoseSearchTrajectory_WorldCollisionResults& CollisionResult,
 		ETraceTypeQuery TraceChannel, bool bTraceComplex, const TArray<AActor*>& ActorsToIgnore, EDrawDebugTrace::Type DrawDebugType, bool bIgnoreSelf = true, float MaxObstacleHeight = 10000.f, FLinearColor TraceColor = FLinearColor::Red, FLinearColor TraceHitColor = FLinearColor::Green, float DrawTime = 5.0f);
 
 	UFUNCTION(BlueprintPure, meta = (BlueprintThreadSafe), Category="Animation|PoseSearch")
