Removed hit normal  from trajectory collision

diff --git a/Engine/Plugins/Animation/PoseSearch/Source/Runtime/Private/PoseSearchTrajectoryLibrary.cpp b/Engine/Plugins/Animation/PoseSearch/Source/Runtime/Private/PoseSearchTrajectoryLibrary.cpp
--- a/Engine/Plugins/Animation/PoseSearch/Source/Runtime/Private/PoseSearchTrajectoryLibrary.cpp
+++ b/Engine/Plugins/Animation/PoseSearch/Source/Runtime/Private/PoseSearchTrajectoryLibrary.cpp
@@ -444,7 +444,6 @@ void UPoseSearchTrajectoryLibrary::HandleTrajectoryWorldCollisions(const UObject
 						bIsFirstFall = false;
 						CollisionResult.TimeToLand = (InitialVelocityZ / -GravityZ) + ((FMath::Sqrt(FMath::Square(InitialVelocityZ) + (2.f * -GravityZ * FallHeight))) / -GravityZ);
 						CollisionResult.LandSpeed = InitialVelocityZ + GravityZ * CollisionResult.TimeToLand;
-						CollisionResult.LandHitResult = HitResult;
 					}
 
 					FreeFallAccumulatedSeconds = 0.f;
diff --git a/Engine/Plugins/Animation/PoseSearch/Source/Runtime/Public/PoseSearch/PoseSearchTrajectoryLibrary.h b/Engine/Plugins/Animation/PoseSearch/Source/Runtime/Public/PoseSearch/PoseSearchTrajectoryLibrary.h
--- a/Engine/Plugins/Animation/PoseSearch/Source/Runtime/Public/PoseSearch/PoseSearchTrajectoryLibrary.h
+++ b/Engine/Plugins/Animation/PoseSearch/Source/Runtime/Public/PoseSearch/PoseSearchTrajectoryLibrary.h
@@ -91,9 +91,6 @@ struct POSESEARCH_API FPoseSearchTrajectory_WorldCollisionResults
 
 	UPROPERTY(BlueprintReadOnly, VisibleAnywhere, Category = "Trajectory Settings")
 	float LandSpeed  = 0.0f;
-
-	UPROPERTY(BlueprintReadOnly, VisibleAnywhere, Category = "Trajectory Settings")
-	FHitResult LandHitResult;
 };
 
 /**
