jasper - marked HandleTrajectoryCollisions experimental

#rnx

diff --git a/Engine/Plugins/Animation/PoseSearch/Source/Runtime/Private/PoseSearchTrajectoryLibrary.cpp b/Engine/Plugins/Animation/PoseSearch/Source/Runtime/Private/PoseSearchTrajectoryLibrary.cpp
--- a/Engine/Plugins/Animation/PoseSearch/Source/Runtime/Private/PoseSearchTrajectoryLibrary.cpp
+++ b/Engine/Plugins/Animation/PoseSearch/Source/Runtime/Private/PoseSearchTrajectoryLibrary.cpp
@@ -453,14 +453,14 @@ void UPoseSearchTrajectoryLibrary::HandleTrajectoryWorldCollisions(const UObject
 	}
 	else if (FloorCollisionsOffset > 0.f)
 	{
-	for (int32 SampleIndex = 0; SampleIndex < NumSamples; ++SampleIndex)
-	{
-		FPoseSearchQueryTrajectorySample& Sample = OutTrajectory.Samples[SampleIndex];
-		if (Sample.AccumulatedSeconds > 0.f)
+		for (int32 SampleIndex = 0; SampleIndex < NumSamples; ++SampleIndex)
 		{
+			FPoseSearchQueryTrajectorySample& Sample = OutTrajectory.Samples[SampleIndex];
+			if (Sample.AccumulatedSeconds > 0.f)
+			{
 				FHitResult HitResult;
 				if (UKismetSystemLibrary::LineTraceSingle(WorldContextObject, Sample.Position + FVector::UpVector * 3000.f, Sample.Position, TraceChannel, bTraceComplex, ActorsToIgnore, DrawDebugType, HitResult, bIgnoreSelf, TraceColor, TraceHitColor, DrawTime))
-			{
+				{
 					Sample.Position.Z = HitResult.ImpactPoint.Z + FloorCollisionsOffset;
 				}
 			}
diff --git a/Engine/Plugins/Animation/PoseSearch/Source/Runtime/Public/PoseSearch/PoseSearchTrajectoryLibrary.h b/Engine/Plugins/Animation/PoseSearch/Source/Runtime/Public/PoseSearch/PoseSearchTrajectoryLibrary.h
--- a/Engine/Plugins/Animation/PoseSearch/Source/Runtime/Public/PoseSearch/PoseSearchTrajectoryLibrary.h
+++ b/Engine/Plugins/Animation/PoseSearch/Source/Runtime/Public/PoseSearch/PoseSearchTrajectoryLibrary.h
@@ -112,16 +112,16 @@ public:
 	// Update prediction by simulating the movement math for ground locomotion from UCharacterMovementComponent.
 	static void UpdatePrediction_SimulateCharacterMovement(FPoseSearchQueryTrajectory& Trajectory, const FPoseSearchTrajectoryData& TrajectoryData, const FPoseSearchTrajectoryData::FDerived& TrajectoryDataDerived, const FPoseSearchTrajectoryData::FSampling& TrajectoryDataSampling, float DeltaTime);
 
-	// Get a Pose History node context from an anim node context (pure)
+	// Generates a prediction trajectory based of the current character intent.
 	UFUNCTION(BlueprintCallable, Category = "Animation|PoseSearch", meta = (BlueprintThreadSafe, DisplayName = "Pose Search Generate Trajectory"))
 	static void PoseSearchGenerateTrajectory(const UAnimInstance* InAnimInstance, UPARAM(ref) const FPoseSearchTrajectoryData& InTrajectoryData, float InDeltaTime,
 		UPARAM(ref) FPoseSearchQueryTrajectory& InOutTrajectory, UPARAM(ref) float& InOutDesiredControllerYawLastUpdate, FPoseSearchQueryTrajectory& OutTrajectory,
 		float InHistorySamplingInterval = 0.04f, int32 InTrajectoryHistoryCount = 10, float InPredictionSamplingInterval = 0.2f, int32 InTrajectoryPredictionCount = 8);
 
-	// this function takes process InTrajectory and returns the eventually modified OutTrajectory
-	// if bApplyGravity is true, gravity from the UCharacterMovementComponent will be applied
-	// and, if FloorCollisionsOffset > 0, vertical collision will be performed to every sample of the trajectory to have the samples float over the geometry (by FloorCollisionsOffset)
-	UFUNCTION(BlueprintCallable, Category="Animation|PoseSearch", meta=(WorldContext="WorldContextObject", AutoCreateRefTerm="ActorsToIgnore", AdvancedDisplay="TraceChannel,bTraceComplex,ActorsToIgnore,DrawDebugType,bIgnoreSelf,MaxObstacleHeight,TraceColor,TraceHitColor,DrawTime"))
+	// Experimental: Process InTrajectory to apply gravity and handle collisions. Eventually returns the modified OutTrajectory.
+	// If bApplyGravity is true, gravity from the UCharacterMovementComponent will be applied.
+	// If FloorCollisionsOffset > 0, vertical collision will be performed to every sample of the trajectory to have the samples float over the geometry (by FloorCollisionsOffset).
+	UFUNCTION(BlueprintCallable, Category="Animation|PoseSearch|Experimental", meta=(WorldContext="WorldContextObject", AutoCreateRefTerm="ActorsToIgnore", AdvancedDisplay="TraceChannel,bTraceComplex,ActorsToIgnore,DrawDebugType,bIgnoreSelf,MaxObstacleHeight,TraceColor,TraceHitColor,DrawTime"))
 	static void HandleTrajectoryWorldCollisions(const UObject* WorldContextObject, const UAnimInstance* AnimInstance, UPARAM(ref) const FPoseSearchQueryTrajectory& InTrajectory, bool bApplyGravity, float FloorCollisionsOffset, FPoseSearchQueryTrajectory& OutTrajectory, FPoseSearchTrajectory_WorldCollisionResults& CollisionResult,
 		ETraceTypeQuery TraceChannel, bool bTraceComplex, const TArray<AActor*>& ActorsToIgnore, EDrawDebugTrace::Type DrawDebugType, bool bIgnoreSelf = true, float MaxObstacleHeight = 10000.f, FLinearColor TraceColor = FLinearColor::Red, FLinearColor TraceHitColor = FLinearColor::Green, float DrawTime = 5.0f);
 
