jasper - Added option to specify max num of history and/or prediction samples to be draw when drawing a pose history node trajectory

#rnx minor

diff --git a/Engine/Plugins/Animation/PoseSearch/Source/Runtime/Private/PoseSearchHistory.cpp b/Engine/Plugins/Animation/PoseSearch/Source/Runtime/Private/PoseSearchHistory.cpp
--- a/Engine/Plugins/Animation/PoseSearch/Source/Runtime/Private/PoseSearchHistory.cpp
+++ b/Engine/Plugins/Animation/PoseSearch/Source/Runtime/Private/PoseSearchHistory.cpp
@@ -22,6 +22,8 @@ namespace UE::PoseSearch
 TAutoConsoleVariable<bool> CVarAnimPoseHistoryDebugDrawPose(TEXT("a.AnimNode.PoseHistory.DebugDrawPose"), false, TEXT("Enable / Disable Pose History Pose DebugDraw"));
 TAutoConsoleVariable<bool> CVarAnimPoseHistoryDebugDrawTrajectory(TEXT("a.AnimNode.PoseHistory.DebugDrawTrajectory"), false, TEXT("Enable / Disable Pose History Trajectory DebugDraw"));
 TAutoConsoleVariable<float> CVarAnimPoseHistoryDebugDrawTrajectoryThickness(TEXT("a.AnimNode.PoseHistory.DebugDrawTrajectoryThickness"), 0.0f, TEXT("Thickness of the trajectory debug draw (Default 0.0f)"));
+TAutoConsoleVariable<int> CVarAnimPoseHistoryDebugDrawTrajectoryMaxNumOfHistorySamples(TEXT("a.AnimNode.PoseHistory.DebugDrawMaxNumOfHistorySamples"), -1, TEXT("Max number of history samples to debug draw. All history samples will be drawn if value is negative. (Default -1)"));
+TAutoConsoleVariable<int> CVarAnimPoseHistoryDebugDrawTrajectoryMaxNumOfPredictionSamples(TEXT("a.AnimNode.PoseHistory.DebugDrawMaxNumOfPredictionSamples"), -1, TEXT("Max number of prediction samples to debug draw. All prediction samples will be drawn if value is negative. (Default -1)"));
 #endif
 
 /**
@@ -692,7 +694,9 @@ void FPoseHistory::DebugDraw(FAnimInstanceProxy& AnimInstanceProxy, FColor Color
 	if (CVarAnimPoseHistoryDebugDrawTrajectory.GetValueOnAnyThread())
 	{
 		const float DebugThickness = CVarAnimPoseHistoryDebugDrawTrajectoryThickness.GetValueOnAnyThread();
-		Trajectory.DebugDrawTrajectory(AnimInstanceProxy, DebugThickness);
+		const int MaxHistorySamples = CVarAnimPoseHistoryDebugDrawTrajectoryMaxNumOfHistorySamples.GetValueOnAnyThread();
+		const int MaxPredictionSamples = CVarAnimPoseHistoryDebugDrawTrajectoryMaxNumOfPredictionSamples.GetValueOnAnyThread();
+		Trajectory.DebugDrawTrajectory(AnimInstanceProxy, DebugThickness, 0, MaxHistorySamples, MaxPredictionSamples);
 	}
 
 	if (Color.A > 0 && CVarAnimPoseHistoryDebugDrawPose.GetValueOnAnyThread())
diff --git a/Engine/Plugins/Animation/PoseSearch/Source/Runtime/Private/PoseSearchTrajectoryTypes.cpp b/Engine/Plugins/Animation/PoseSearch/Source/Runtime/Private/PoseSearchTrajectoryTypes.cpp
--- a/Engine/Plugins/Animation/PoseSearch/Source/Runtime/Private/PoseSearchTrajectoryTypes.cpp
+++ b/Engine/Plugins/Animation/PoseSearch/Source/Runtime/Private/PoseSearchTrajectoryTypes.cpp
@@ -96,14 +96,35 @@ void FPoseSearchQueryTrajectory::DebugDrawTrajectory(const UWorld* World, const
 	}
 }
 
-void FPoseSearchQueryTrajectory::DebugDrawTrajectory(FAnimInstanceProxy& AnimInstanceProxy, const float DebugThickness, float HeightOffset) const
+void FPoseSearchQueryTrajectory::DebugDrawTrajectory(FAnimInstanceProxy& AnimInstanceProxy, const float DebugThickness, float HeightOffset, int MaxHistorySamples, int MaxPredictionSamples) const
 {
 	const FVector OffsetVector = FVector::UpVector * HeightOffset;
-
+	
+	int AvailableHistorySamplesCount = 0;
+	
+	if (MaxHistorySamples != -1 || MaxPredictionSamples != -1)
+	{
+		for (int32 i = 0; i < Samples.Num(); ++i)
+		{
+			if (Samples[i].AccumulatedSeconds <= 0)
+			{
+				++AvailableHistorySamplesCount;
+			}
+			else
+			{
+				break;
+			}
+		}
+	}
+	
 	const int32 LastIndex = Samples.Num() - 1;
-	if (LastIndex >= 0)
+	const int32 ZeroTimeIndex = AvailableHistorySamplesCount - 1;
+	const int32 StartIndex = MaxHistorySamples < 0 ? 0 : FMath::Max(AvailableHistorySamplesCount - MaxHistorySamples, 0); 
+	const int32 EndIndex = MaxPredictionSamples < 0 ? LastIndex : FMath::Min(ZeroTimeIndex + MaxPredictionSamples, LastIndex);
+	
+	if (LastIndex >= 0 && StartIndex <= EndIndex)
 	{
-		for (int32 Index = 0; ; ++Index)
+		for (int32 Index = StartIndex; ; ++Index)
 		{
 			const FVector Pos = Samples[Index].Position + OffsetVector;
 
@@ -120,7 +141,7 @@ void FPoseSearchQueryTrajectory::DebugDrawTrajectory(FAnimInstanceProxy& AnimIns
 			AnimInstanceProxy.AnimDrawDebugLine(Pos, Pos + X * Scale, IsPast ? FColor::Red : FColor::Blue, false, -1.f, DebugThickness, SDPG_Foreground);
 			AnimInstanceProxy.AnimDrawDebugLine(Pos, Pos + Y * Scale, IsPast ? FColor::Orange : FColor::Turquoise, false, -1.f, DebugThickness, SDPG_Foreground);
 
-			if (Index == LastIndex)
+			if (Index == EndIndex)
 			{
 				break;
 			}
@@ -130,6 +151,7 @@ void FPoseSearchQueryTrajectory::DebugDrawTrajectory(FAnimInstanceProxy& AnimIns
 		}
 	}
 }
+
 #endif // ENABLE_ANIM_DEBUG
 
 FArchive& operator<<(FArchive& Ar, FPoseSearchQueryTrajectorySample& TrajectorySample)
diff --git a/Engine/Plugins/Animation/PoseSearch/Source/Runtime/Public/PoseSearch/PoseSearchTrajectoryTypes.h b/Engine/Plugins/Animation/PoseSearch/Source/Runtime/Public/PoseSearch/PoseSearchTrajectoryTypes.h
--- a/Engine/Plugins/Animation/PoseSearch/Source/Runtime/Public/PoseSearch/PoseSearchTrajectoryTypes.h
+++ b/Engine/Plugins/Animation/PoseSearch/Source/Runtime/Public/PoseSearch/PoseSearchTrajectoryTypes.h
@@ -42,7 +42,7 @@ struct POSESEARCH_API FPoseSearchQueryTrajectory
 	
 #if ENABLE_ANIM_DEBUG
 	void DebugDrawTrajectory(const UWorld* World, const float DebugThickness = 0.f, float HeightOffset = 0.f) const;
-	void DebugDrawTrajectory(FAnimInstanceProxy& AnimInstanceProxy, const float DebugThickness = 0.f, float HeightOffset = 0.f) const;
+	void DebugDrawTrajectory(FAnimInstanceProxy& AnimInstanceProxy, const float DebugThickness = 0.f, float HeightOffset = 0.f, int MaxHistorySamples = -1, int MaxPredictionSamples = -1) const;
 #endif // ENABLE_ANIM_DEBUG
 };
 POSESEARCH_API FArchive& operator<<(FArchive& Ar, FPoseSearchQueryTrajectory& Trajectory);
