motion matching - fix for stitching MotionMatch

diff --git a/Engine/Plugins/Animation/PoseSearch/Source/Runtime/Private/PoseSearchFeatureChannel_Heading.cpp b/Engine/Plugins/Animation/PoseSearch/Source/Runtime/Private/PoseSearchFeatureChannel_Heading.cpp
--- a/Engine/Plugins/Animation/PoseSearch/Source/Runtime/Private/PoseSearchFeatureChannel_Heading.cpp
+++ b/Engine/Plugins/Animation/PoseSearch/Source/Runtime/Private/PoseSearchFeatureChannel_Heading.cpp
@@ -264,6 +264,15 @@ UE::PoseSearch::TLabelBuilder& UPoseSearchFeatureChannel_Heading::GetLabel(UE::P
 		LabelBuilder.Append(Schema->GetBoneReferences(SampleRole)[SchemaBoneIdx].BoneName.ToString());
 	}
 
+	if (PermutationTimeType == EPermutationTimeType::UsePermutationTime)
+	{
+		LabelBuilder.Append(TEXT("_PT"));
+	}
+	else if (PermutationTimeType == EPermutationTimeType::UseSampleToPermutationTime)
+	{
+		LabelBuilder.Append(TEXT("_SPT"));
+	}
+
 	if (SampleRole != DefaultRole)
 	{
 		LabelBuilder.Append(TEXT("["));
diff --git a/Engine/Plugins/Animation/PoseSearch/Source/Runtime/Private/PoseSearchFeatureChannel_Position.cpp b/Engine/Plugins/Animation/PoseSearch/Source/Runtime/Private/PoseSearchFeatureChannel_Position.cpp
--- a/Engine/Plugins/Animation/PoseSearch/Source/Runtime/Private/PoseSearchFeatureChannel_Position.cpp
+++ b/Engine/Plugins/Animation/PoseSearch/Source/Runtime/Private/PoseSearchFeatureChannel_Position.cpp
@@ -250,6 +250,15 @@ UE::PoseSearch::TLabelBuilder& UPoseSearchFeatureChannel_Position::GetLabel(UE::
 		LabelBuilder.Append(Schema->GetBoneReferences(OriginRole)[SchemaOriginBoneIdx].BoneName.ToString());
 	}
 
+	if (PermutationTimeType == EPermutationTimeType::UsePermutationTime)
+	{
+		LabelBuilder.Append(TEXT("_PT"));
+	}
+	else if (PermutationTimeType == EPermutationTimeType::UseSampleToPermutationTime)
+	{
+		LabelBuilder.Append(TEXT("_SPT"));
+	}
+
 	if (OriginRole != DefaultRole)
 	{
 		LabelBuilder.Append(TEXT("["));
diff --git a/Engine/Plugins/Animation/PoseSearch/Source/Runtime/Private/PoseSearchFeatureChannel_Velocity.cpp b/Engine/Plugins/Animation/PoseSearch/Source/Runtime/Private/PoseSearchFeatureChannel_Velocity.cpp
--- a/Engine/Plugins/Animation/PoseSearch/Source/Runtime/Private/PoseSearchFeatureChannel_Velocity.cpp
+++ b/Engine/Plugins/Animation/PoseSearch/Source/Runtime/Private/PoseSearchFeatureChannel_Velocity.cpp
@@ -213,6 +213,15 @@ UE::PoseSearch::TLabelBuilder& UPoseSearchFeatureChannel_Velocity::GetLabel(UE::
 		LabelBuilder.Append(Schema->GetBoneReferences(SampleRole)[SchemaBoneIdx].BoneName.ToString());
 	}
 
+	if (PermutationTimeType == EPermutationTimeType::UsePermutationTime)
+	{
+		LabelBuilder.Append(TEXT("_PT"));
+	}
+	else if (PermutationTimeType == EPermutationTimeType::UseSampleToPermutationTime)
+	{
+		LabelBuilder.Append(TEXT("_SPT"));
+	}
+
 	if (SampleRole != DefaultRole)
 	{
 		LabelBuilder.Append(TEXT("["));
diff --git a/Engine/Plugins/Animation/PoseSearch/Source/Runtime/Private/PoseSearchLibrary.cpp b/Engine/Plugins/Animation/PoseSearch/Source/Runtime/Private/PoseSearchLibrary.cpp
--- a/Engine/Plugins/Animation/PoseSearch/Source/Runtime/Private/PoseSearchLibrary.cpp
+++ b/Engine/Plugins/Animation/PoseSearch/Source/Runtime/Private/PoseSearchLibrary.cpp
@@ -733,7 +733,7 @@ void UPoseSearchLibrary::MotionMatch(
 		PoseHistories.Add(MemStackPoseHistory.GetThisOrPoseHistory());
 	}
 
-	const FSearchResult SearchResult = MotionMatch(AnimInstances, Roles, PoseHistories, AssetsToSearch, nullptr, 0.f, DebugSessionUniqueIdentifier);
+	const FSearchResult SearchResult = MotionMatch(AnimInstances, Roles, PoseHistories, AssetsToSearch, nullptr, 0.f, DebugSessionUniqueIdentifier, FutureIntervalTime);
 	if (SearchResult.IsValid())
 	{
 		const UPoseSearchDatabase* Database = SearchResult.Database.Get();
@@ -794,7 +794,8 @@ UE::PoseSearch::FSearchResult UPoseSearchLibrary::MotionMatch(
 	TArrayView<const UObject*> AssetsToSearch,
 	const UObject* PlayingAsset,
 	float PlayingAssetAccumulatedTime,
-	const int32 DebugSessionUniqueIdentifier)
+	const int32 DebugSessionUniqueIdentifier,
+	float DesiredPermutationTimeOffset)
 {
 	check(!AnimInstances.IsEmpty() && AnimInstances.Num() == Roles.Num() && AnimInstances.Num() == PoseHistories.Num());
 
@@ -804,7 +805,7 @@ UE::PoseSearch::FSearchResult UPoseSearchLibrary::MotionMatch(
 
 	FMemMark Mark(FMemStack::Get());
 	FSearchResult ReconstructedPreviousSearchResult;
-	FSearchContext SearchContext(0.f, nullptr, ReconstructedPreviousSearchResult);
+	FSearchContext SearchContext(DesiredPermutationTimeOffset, nullptr, ReconstructedPreviousSearchResult);
 
 	// @todo: all assets in AssetsToSearch should have a consistent Roles requirements, or else the search will throw an error!
 	for (int32 RoleIndex = 0; RoleIndex < Roles.Num(); ++RoleIndex)
diff --git a/Engine/Plugins/Animation/PoseSearch/Source/Runtime/Public/PoseSearch/PoseSearchLibrary.h b/Engine/Plugins/Animation/PoseSearch/Source/Runtime/Public/PoseSearch/PoseSearchLibrary.h
--- a/Engine/Plugins/Animation/PoseSearch/Source/Runtime/Public/PoseSearch/PoseSearchLibrary.h
+++ b/Engine/Plugins/Animation/PoseSearch/Source/Runtime/Public/PoseSearch/PoseSearchLibrary.h
@@ -217,6 +217,7 @@ public:
 		TArrayView<const UObject*> AssetsToSearch,
 		const UObject* PlayingAsset,
 		float PlayingAssetAccumulatedTime,
-		const int32 DebugSessionUniqueIdentifier);
+		const int32 DebugSessionUniqueIdentifier,
+		float DesiredPermutationTimeOffset = 0.f);
 };
 
