motion matching - FPoseHistory code cleanup

diff --git a/Engine/Plugins/Animation/PoseSearch/Source/Runtime/Private/PoseSearchHistory.cpp b/Engine/Plugins/Animation/PoseSearch/Source/Runtime/Private/PoseSearchHistory.cpp
--- a/Engine/Plugins/Animation/PoseSearch/Source/Runtime/Private/PoseSearchHistory.cpp
+++ b/Engine/Plugins/Animation/PoseSearch/Source/Runtime/Private/PoseSearchHistory.cpp
@@ -440,9 +440,23 @@ FPoseHistory::FPoseHistory(FPoseHistory&& Other)
 
 FPoseHistory& FPoseHistory::operator=(const FPoseHistory& Other)
 {
+#if ENABLE_ANIM_DEBUG
 	CheckThreadSafetyWrite(ReadDataThreadSafeCounter);
 	CheckThreadSafetyWrite(WriteDataThreadSafeCounter);
 
+	FThreadSafeCounter& OtherReadDataThreadSafeCounter = Other.ReadDataThreadSafeCounter;
+	FThreadSafeCounter& OtherWriteDataThreadSafeCounter = Other.WriteDataThreadSafeCounter;
+	CheckThreadSafetyWrite(OtherReadDataThreadSafeCounter);
+	CheckThreadSafetyWrite(OtherWriteDataThreadSafeCounter);
+#endif ENABLE_ANIM_DEBUG
+
+	MaxNumPoses = Other.MaxNumPoses;
+	SamplingInterval = Other.SamplingInterval;
+
+	Trajectory = Other.Trajectory;
+	TrajectoryDataState = Other.TrajectoryDataState;
+	TrajectorySpeedMultiplier = Other.TrajectorySpeedMultiplier;
+
 	ReadData = Other.ReadData;
 	WriteData = Other.WriteData;
 	return *this;
@@ -450,9 +464,23 @@ FPoseHistory& FPoseHistory::operator=(const FPoseHistory& Other)
 
 FPoseHistory& FPoseHistory::operator=(FPoseHistory&& Other)
 {
+#if ENABLE_ANIM_DEBUG
 	CheckThreadSafetyWrite(ReadDataThreadSafeCounter);
 	CheckThreadSafetyWrite(WriteDataThreadSafeCounter);
 
+	FThreadSafeCounter& OtherReadDataThreadSafeCounter = Other.ReadDataThreadSafeCounter;
+	FThreadSafeCounter& OtherWriteDataThreadSafeCounter = Other.WriteDataThreadSafeCounter;
+	CheckThreadSafetyWrite(OtherReadDataThreadSafeCounter);
+	CheckThreadSafetyWrite(OtherWriteDataThreadSafeCounter);
+#endif ENABLE_ANIM_DEBUG
+
+	MaxNumPoses = MoveTemp(Other.MaxNumPoses);
+	SamplingInterval = MoveTemp(Other.SamplingInterval);
+
+	Trajectory = MoveTemp(Other.Trajectory);
+	TrajectoryDataState = MoveTemp(Other.TrajectoryDataState);
+	TrajectorySpeedMultiplier = MoveTemp(Other.TrajectorySpeedMultiplier);
+
 	ReadData = MoveTemp(Other.ReadData);
 	WriteData = MoveTemp(Other.WriteData);
 	return *this;
@@ -465,6 +493,13 @@ void FPoseHistory::Initialize_AnyThread(int32 InNumPoses, float InSamplingInterv
 
 	MaxNumPoses = InNumPoses;
 	SamplingInterval = InSamplingInterval;
+
+	Trajectory = FPoseSearchQueryTrajectory();
+	TrajectoryDataState = FPoseSearchTrajectoryData::FState();
+	TrajectorySpeedMultiplier = 1.f;
+
+	ReadData = FData();
+	WriteData = FData();
 }
 
 void FPoseHistory::CacheBones_AnyThread(const TArray<FBoneIndexType>& RequiredBones)
@@ -477,14 +512,16 @@ void FPoseHistory::CacheBones_AnyThread(const TArray<FBoneIndexType>& RequiredBo
 	if (!RequiredBones.IsEmpty())
 	{
 		// making sure we always collect the root bone transform (by construction BoneToTransformMap[0] = 0)
-		WriteData.BoneToTransformMap.Add(RootBoneIndexType) = WriteData.BoneToTransformMap.Num();
+		const FComponentSpaceTransformIndex ComponentSpaceTransformRootBoneIndex = 0;
+		WriteData.BoneToTransformMap.Add(RootBoneIndexType) = ComponentSpaceTransformRootBoneIndex;
 
 		for (int32 i = 0; i < RequiredBones.Num(); ++i)
 		{
 			// adding only unique RequiredBones to avoid oversizing Entries::ComponentSpaceTransforms
 			if (!WriteData.BoneToTransformMap.Find(RequiredBones[i]))
 			{
-				WriteData.BoneToTransformMap.Add(RequiredBones[i]) = WriteData.BoneToTransformMap.Num();
+				const FComponentSpaceTransformIndex ComponentSpaceTransformIndex = WriteData.BoneToTransformMap.Num();
+				WriteData.BoneToTransformMap.Add(RequiredBones[i]) = ComponentSpaceTransformIndex;
 			}
 		}
 	}
