motion matching
- improved debug lines for sampler when using MM stitches
- fix for crash when reverting sequences previewed in the database editor

diff --git a/Engine/Plugins/Animation/PoseSearch/Source/Editor/Private/PoseSearchDatabaseAssetListItem.cpp b/Engine/Plugins/Animation/PoseSearch/Source/Editor/Private/PoseSearchDatabaseAssetListItem.cpp
--- a/Engine/Plugins/Animation/PoseSearch/Source/Editor/Private/PoseSearchDatabaseAssetListItem.cpp
+++ b/Engine/Plugins/Animation/PoseSearch/Source/Editor/Private/PoseSearchDatabaseAssetListItem.cpp
@@ -115,40 +115,41 @@ namespace UE::PoseSearch
 
 	FReply SDatabaseAssetListItem::OnMouseButtonDoubleClick(const FGeometry& InMyGeometry, const FPointerEvent& InMouseEvent)
 	{
-		TSharedPtr<FDatabaseAssetTreeNode> Node = WeakAssetTreeNode.Pin();
-		TSharedPtr<FDatabaseViewModel> ViewModel = EditorViewModel.Pin();
-		if (const UPoseSearchDatabase* Database = ViewModel->GetPoseSearchDatabase())
+		if (TSharedPtr<FDatabaseAssetTreeNode> AssetTreeNode = WeakAssetTreeNode.Pin())
 		{
-			if (const FPoseSearchDatabaseAnimationAssetBase* DatabaseAnimationAsset = Database->GetAnimationAssetBase(Node->SourceAssetIdx))
+			TSharedPtr<FDatabaseViewModel> ViewModel = EditorViewModel.Pin();
+			if (const UPoseSearchDatabase* Database = ViewModel->GetPoseSearchDatabase())
 			{
-				if (UAssetEditorSubsystem* AssetEditorSS = GEditor->GetEditorSubsystem<UAssetEditorSubsystem>())
+				if (const FPoseSearchDatabaseAnimationAssetBase* DatabaseAnimationAsset = Database->GetAnimationAssetBase(AssetTreeNode->SourceAssetIdx))
 				{
-					if (UObject* AnimationAsset = DatabaseAnimationAsset->GetAnimationAsset())
+					if (UAssetEditorSubsystem* AssetEditorSS = GEditor->GetEditorSubsystem<UAssetEditorSubsystem>())
 					{
-						AssetEditorSS->OpenEditorForAsset(AnimationAsset);
-
-						if (IAssetEditorInstance* Editor = AssetEditorSS->FindEditorForAsset(AnimationAsset, true))
+						if (UObject* AnimationAsset = DatabaseAnimationAsset->GetAnimationAsset())
 						{
-							if (Editor->GetEditorName() == "AnimationEditor")
-							{
-								float AnimationAssetTime = 0.f;
-								FVector AnimationAssetBlendParameters = FVector::ZeroVector;
-								ViewModel->GetAnimationTime(Node->SourceAssetIdx, AnimationAssetTime, AnimationAssetBlendParameters);
+							AssetEditorSS->OpenEditorForAsset(AnimationAsset);
 
-								const IAnimationEditor* AnimationEditor = static_cast<IAnimationEditor*>(Editor);
-								const UDebugSkelMeshComponent* PreviewComponent = AnimationEditor->GetPersonaToolkit()->GetPreviewMeshComponent();
-
-								// Open asset paused and at specific time as seen on the pose search debugger.
-								PreviewComponent->PreviewInstance->SetPosition(AnimationAssetTime);
-								PreviewComponent->PreviewInstance->SetPlaying(false);
-								PreviewComponent->PreviewInstance->SetBlendSpacePosition(AnimationAssetBlendParameters);
+							if (IAssetEditorInstance* Editor = AssetEditorSS->FindEditorForAsset(AnimationAsset, true))
+							{
+								if (Editor->GetEditorName() == "AnimationEditor")
+								{
+									float AnimationAssetTime = 0.f;
+									FVector AnimationAssetBlendParameters = FVector::ZeroVector;
+									ViewModel->GetAnimationTime(AssetTreeNode->SourceAssetIdx, AnimationAssetTime, AnimationAssetBlendParameters);
+
+									const IAnimationEditor* AnimationEditor = static_cast<IAnimationEditor*>(Editor);
+									const UDebugSkelMeshComponent* PreviewComponent = AnimationEditor->GetPersonaToolkit()->GetPreviewMeshComponent();
+
+									// Open asset paused and at specific time as seen on the pose search debugger.
+									PreviewComponent->PreviewInstance->SetPosition(AnimationAssetTime);
+									PreviewComponent->PreviewInstance->SetPlaying(false);
+									PreviewComponent->PreviewInstance->SetBlendSpacePosition(AnimationAssetBlendParameters);
+								}
 							}
 						}
 					}
 				}
 			}
 		}
-
 		return STableRow<TSharedPtr<FDatabaseAssetTreeNode>>::OnMouseButtonDoubleClick(InMyGeometry, InMouseEvent);
 	}
 
@@ -158,12 +159,13 @@ namespace UE::PoseSearch
 
 		if (const UPoseSearchDatabase* Database = ViewModel->GetPoseSearchDatabase())
 		{
-			TSharedPtr<FDatabaseAssetTreeNode> Node = WeakAssetTreeNode.Pin();
-			if (const FPoseSearchDatabaseAnimationAssetBase* DatabaseAnimationAsset = Database->GetAnimationAssetBase(Node->SourceAssetIdx))
+			if (TSharedPtr<FDatabaseAssetTreeNode> AssetTreeNode = WeakAssetTreeNode.Pin())
 			{
-				return FText::FromString(DatabaseAnimationAsset->GetName());
+				if (const FPoseSearchDatabaseAnimationAssetBase* DatabaseAnimationAsset = Database->GetAnimationAssetBase(AssetTreeNode->SourceAssetIdx))
+				{
+					return FText::FromString(DatabaseAnimationAsset->GetName());
+				}
 			}
-
 			return FText::FromString(Database->GetName());
 		}
 
@@ -172,12 +174,16 @@ namespace UE::PoseSearch
 
 	TSharedRef<SWidget> SDatabaseAssetListItem::GenerateItemWidget()
 	{
-		TSharedPtr<FDatabaseAssetTreeNode> Node = WeakAssetTreeNode.Pin();
-		
+		int32 SourceAssetIdx = INDEX_NONE;
+		if (TSharedPtr<FDatabaseAssetTreeNode> AssetTreeNode = WeakAssetTreeNode.Pin())
+		{
+			SourceAssetIdx = AssetTreeNode->SourceAssetIdx;
+		}
+
 		TSharedPtr<SWidget> ItemWidget;
 		const FDetailColumnSizeData& ColumnSizeData = SkeletonView.Pin()->GetColumnSizeData();
 		
-		if (Node->SourceAssetIdx == INDEX_NONE)
+		if (SourceAssetIdx == INDEX_NONE)
 		{
 			// it's a group
 			SAssignNew(ItemWidget, SBorder)
@@ -211,7 +217,7 @@ namespace UE::PoseSearch
 			TSharedPtr<FDatabaseViewModel> ViewModel = EditorViewModel.Pin();
 			if (UPoseSearchDatabase* Database = ViewModel->GetPoseSearchDatabase())
 			{
-				if (const FPoseSearchDatabaseAnimationAssetBase* DatabaseAnimationAsset = Database->GetAnimationAssetBase(Node->SourceAssetIdx))
+				if (const FPoseSearchDatabaseAnimationAssetBase* DatabaseAnimationAsset = Database->GetAnimationAssetBase(SourceAssetIdx))
 				{
 					SAssignNew(ItemIconWidget, SImage)
 						.Image(FSlateIconFinder::FindIconBrushForClass(DatabaseAnimationAsset->GetAnimationAssetStaticClass()));
@@ -385,15 +391,16 @@ namespace UE::PoseSearch
 	EVisibility SDatabaseAssetListItem::GetSelectedActorIconVisbility() const
 	{
 		TSharedPtr<FDatabaseViewModel> ViewModelPtr = EditorViewModel.Pin();
-		TSharedPtr<FDatabaseAssetTreeNode> TreeNodePtr = WeakAssetTreeNode.Pin();
-		if (const FSearchIndexAsset* SelectedIndexAsset = ViewModelPtr->GetSelectedActorIndexAsset())
+		if (TSharedPtr<FDatabaseAssetTreeNode> AssetTreeNode = WeakAssetTreeNode.Pin())
 		{
-			if (TreeNodePtr->SourceAssetIdx == SelectedIndexAsset->GetSourceAssetIdx())
+			if (const FSearchIndexAsset* SelectedIndexAsset = ViewModelPtr->GetSelectedActorIndexAsset())
 			{
-				return EVisibility::Visible;
+				if (AssetTreeNode->SourceAssetIdx == SelectedIndexAsset->GetSourceAssetIdx())
+				{
+					return EVisibility::Visible;
+				}
 			}
 		}
-
 		return EVisibility::Hidden;
 	}
 
@@ -412,12 +419,14 @@ namespace UE::PoseSearch
 		TSharedPtr<FDatabaseViewModel> ViewModelPtr = EditorViewModel.Pin();
 		if (const UPoseSearchDatabase* Database = ViewModelPtr->GetPoseSearchDatabase())
 		{
-			TSharedPtr<FDatabaseAssetTreeNode> TreeNodePtr = WeakAssetTreeNode.Pin();
-			if (Database->GetAnimationAssets().IsValidIndex(TreeNodePtr->SourceAssetIdx))
+			if (TSharedPtr<FDatabaseAssetTreeNode> AssetTreeNode = WeakAssetTreeNode.Pin())
 			{
-				if (ViewModelPtr->IsDisableReselection(TreeNodePtr->SourceAssetIdx))
+				if (Database->GetAnimationAssets().IsValidIndex(AssetTreeNode->SourceAssetIdx))
 				{
-					return ECheckBoxState::Checked;
+					if (ViewModelPtr->IsDisableReselection(AssetTreeNode->SourceAssetIdx))
+					{
+						return ECheckBoxState::Checked;
+					}
 				}
 			}
 		}
@@ -430,17 +439,18 @@ namespace UE::PoseSearch
 		const TSharedPtr<FDatabaseViewModel> ViewModelPtr = EditorViewModel.Pin();
 		if (UPoseSearchDatabase* PoseSearchDatabase = ViewModelPtr->GetPoseSearchDatabase())
 		{
-			const FScopedTransaction Transaction(LOCTEXT("EnableChangedForAssetInPoseSearchDatabase", "Update enabled flag for item from Pose Search Database"));
-
-			const TSharedPtr<FDatabaseAssetTreeNode> TreeNodePtr = WeakAssetTreeNode.Pin();
+			if (const TSharedPtr<FDatabaseAssetTreeNode> AssetTreeNode = WeakAssetTreeNode.Pin())
+			{
+				const FScopedTransaction Transaction(LOCTEXT("EnableChangedForAssetInPoseSearchDatabase", "Update enabled flag for item from Pose Search Database"));
 
-			PoseSearchDatabase->Modify();
+				PoseSearchDatabase->Modify();
 
-			ViewModelPtr->SetDisableReselection(TreeNodePtr->SourceAssetIdx, NewCheckboxState == ECheckBoxState::Checked ? true : false);
+				ViewModelPtr->SetDisableReselection(AssetTreeNode->SourceAssetIdx, NewCheckboxState == ECheckBoxState::Checked ? true : false);
 
-			SkeletonView.Pin()->RefreshTreeView(false, true);
+				SkeletonView.Pin()->RefreshTreeView(false, true);
 
-			// no need to rebuild the SearchIndex (ViewModelPtr->BuildSearchIndex()), since bDisableReselection is a runtime only parameter
+				// no need to rebuild the SearchIndex (ViewModelPtr->BuildSearchIndex()), since bDisableReselection is a runtime only parameter
+			}
 		}
 	}
 
@@ -449,12 +459,14 @@ namespace UE::PoseSearch
 		TSharedPtr<FDatabaseViewModel> ViewModelPtr = EditorViewModel.Pin();
 		if (const UPoseSearchDatabase* Database = ViewModelPtr->GetPoseSearchDatabase())
 		{
-			TSharedPtr<FDatabaseAssetTreeNode> TreeNodePtr = WeakAssetTreeNode.Pin();
-			if (Database->GetAnimationAssets().IsValidIndex(TreeNodePtr->SourceAssetIdx))
+			if (TSharedPtr<FDatabaseAssetTreeNode> AssetTreeNode = WeakAssetTreeNode.Pin())
 			{
-				if (ViewModelPtr->IsEnabled(TreeNodePtr->SourceAssetIdx))
+				if (Database->GetAnimationAssets().IsValidIndex(AssetTreeNode->SourceAssetIdx))
 				{
-					return ECheckBoxState::Checked;
+					if (ViewModelPtr->IsEnabled(AssetTreeNode->SourceAssetIdx))
+					{
+						return ECheckBoxState::Checked;
+					}
 				}
 			}
 		}
@@ -466,12 +478,13 @@ namespace UE::PoseSearch
 		const FScopedTransaction Transaction(LOCTEXT("EnableChangedForAssetInPoseSearchDatabase", "Update enabled flag for item from Pose Search Database"));
 
 		const TSharedPtr<FDatabaseViewModel> ViewModelPtr = EditorViewModel.Pin();
-		const TSharedPtr<FDatabaseAssetTreeNode> TreeNodePtr = WeakAssetTreeNode.Pin();
-
-		ViewModelPtr->SetIsEnabled(TreeNodePtr->SourceAssetIdx, NewCheckboxState == ECheckBoxState::Checked);
+		if (const TSharedPtr<FDatabaseAssetTreeNode> AssetTreeNode = WeakAssetTreeNode.Pin())
+		{
+			ViewModelPtr->SetIsEnabled(AssetTreeNode->SourceAssetIdx, NewCheckboxState == ECheckBoxState::Checked);
 
-		SkeletonView.Pin()->RefreshTreeView(false, true);
-		ViewModelPtr->BuildSearchIndex();
+			SkeletonView.Pin()->RefreshTreeView(false, true);
+			ViewModelPtr->BuildSearchIndex();
+		}
 	}
 
 	FSlateColor SDatabaseAssetListItem::GetNameTextColorAndOpacity() const
@@ -479,17 +492,19 @@ namespace UE::PoseSearch
 		TSharedPtr<FDatabaseViewModel> ViewModelPtr = EditorViewModel.Pin();
 		if (const UPoseSearchDatabase* Database = ViewModelPtr->GetPoseSearchDatabase())
 		{
-			TSharedPtr<FDatabaseAssetTreeNode> TreeNodePtr = WeakAssetTreeNode.Pin();
-			if (const FPoseSearchDatabaseAnimationAssetBase* DatabaseAnimationAssetBase = Database->GetAnimationAssetBase(TreeNodePtr->SourceAssetIdx))
+			if (TSharedPtr<FDatabaseAssetTreeNode> AssetTreeNode = WeakAssetTreeNode.Pin())
 			{
-				if (DatabaseAnimationAssetBase->IsEnabled())
+				if (const FPoseSearchDatabaseAnimationAssetBase* DatabaseAnimationAssetBase = Database->GetAnimationAssetBase(AssetTreeNode->SourceAssetIdx))
 				{
-					if (DatabaseAnimationAssetBase->bSynchronizeWithExternalDependency)
+					if (DatabaseAnimationAssetBase->IsEnabled())
 					{
-						return FColor::Turquoise;
-					}
+						if (DatabaseAnimationAssetBase->bSynchronizeWithExternalDependency)
+						{
+							return FColor::Turquoise;
+						}
 
-					return FLinearColor::White;
+						return FLinearColor::White;
+					}
 				}
 			}
 		}
@@ -498,84 +513,113 @@ namespace UE::PoseSearch
 
 	FSlateColor SDatabaseAssetListItem::GetLoopingColorAndOpacity() const
 	{
-		const TSharedPtr<FDatabaseAssetTreeNode> Node = WeakAssetTreeNode.Pin();
-		return Node->IsLooping() ? FLinearColor::White : DisabledColor;
+		if (const TSharedPtr<FDatabaseAssetTreeNode> AssetTreeNode = WeakAssetTreeNode.Pin())
+		{
+			if (AssetTreeNode->IsLooping())
+			{
+				return FLinearColor::White;
+			}
+		}
+		
+		return DisabledColor;
 	}
 
 	FText SDatabaseAssetListItem::GetLoopingToolTip() const
 	{
-		const TSharedPtr<FDatabaseAssetTreeNode> Node = WeakAssetTreeNode.Pin();
-		return Node->IsLooping() ? LOCTEXT("NodeLoopEnabledToolTip", "Looping (Read only)") : LOCTEXT("NodeLoopDisabledToolTip", "Not looping (Read only)");
+		if (const TSharedPtr<FDatabaseAssetTreeNode> AssetTreeNode = WeakAssetTreeNode.Pin())
+		{
+			if (AssetTreeNode->IsLooping())
+			{
+				return LOCTEXT("NodeLoopEnabledToolTip", "Looping (Read only)");
+			}
+		}
+
+		return LOCTEXT("NodeLoopDisabledToolTip", "Not looping (Read only)");
 	}
 
 	FSlateColor SDatabaseAssetListItem::GetRootMotionColorAndOpacity() const
 	{
-		const TSharedPtr<FDatabaseAssetTreeNode> Node = WeakAssetTreeNode.Pin();
-		return Node->IsRootMotionEnabled() ? FLinearColor::White : DisabledColor;
+		if (const TSharedPtr<FDatabaseAssetTreeNode> AssetTreeNode = WeakAssetTreeNode.Pin())
+		{
+			if (AssetTreeNode->IsRootMotionEnabled())
+			{
+				return FLinearColor::White;
+			}
+		}
+		
+		return DisabledColor;
 	}
 
 	FText SDatabaseAssetListItem::GetRootMotionOptionToolTip() const
 	{
-		const TSharedPtr<FDatabaseAssetTreeNode> Node = WeakAssetTreeNode.Pin();
-		return Node->IsRootMotionEnabled() ? LOCTEXT("NodeRootMotionEnabledToolTip", "Root motion enabled (Read only)") : LOCTEXT("NodeRootMotionDisabledToolTip", "No root motion enabled (Read only)");
-
+		if (const TSharedPtr<FDatabaseAssetTreeNode> AssetTreeNode = WeakAssetTreeNode.Pin())
+		{
+			if (AssetTreeNode->IsRootMotionEnabled())
+			{
+				return LOCTEXT("NodeRootMotionEnabledToolTip", "Root motion enabled (Read only)");
+			}
+		}
+		
+		return LOCTEXT("NodeRootMotionDisabledToolTip", "No root motion enabled (Read only)");
 	}
+
 	const FSlateBrush* SDatabaseAssetListItem::GetMirrorOptionSlateBrush() const
 	{
-		const TSharedPtr<FDatabaseAssetTreeNode> Node = WeakAssetTreeNode.Pin();
-
-		// TODO: Update icons when appropriate assets become available.
-		switch (Node->GetMirrorOption())
+		if (const TSharedPtr<FDatabaseAssetTreeNode> AssetTreeNode = WeakAssetTreeNode.Pin())
 		{
-			case EPoseSearchMirrorOption::UnmirroredOnly: 
+			// TODO: Update icons when appropriate assets become available.
+			switch (AssetTreeNode->GetMirrorOption())
+			{
+			case EPoseSearchMirrorOption::UnmirroredOnly:
 				return FAppStyle::Get().GetBrush("GraphEditor.AlignNodesRight");
-			
-			case EPoseSearchMirrorOption::MirroredOnly: 
+
+			case EPoseSearchMirrorOption::MirroredOnly:
 				return FAppStyle::Get().GetBrush("GraphEditor.AlignNodesLeft");
-			
+
 			case EPoseSearchMirrorOption::UnmirroredAndMirrored:
 				return FAppStyle::Get().GetBrush("GraphEditor.AlignNodesCenter");
-			
-			default:
-				return nullptr;
+			}
 		}
+		
+		return nullptr;
 	}
 
 	FText SDatabaseAssetListItem::GetMirrorOptionToolTip() const
 	{
-		const TSharedPtr<FDatabaseAssetTreeNode> Node = WeakAssetTreeNode.Pin();
-		return FText::FromString(LOCTEXT("ToolTipMirrorOption", "Mirror Option: ").ToString() + (Node ? UEnum::GetDisplayValueAsText(Node->GetMirrorOption()).ToString() : LOCTEXT("ToolTipMirrorOption_Invalid", "Invalid").ToString()));
+		const TSharedPtr<FDatabaseAssetTreeNode> AssetTreeNode = WeakAssetTreeNode.Pin();
+		return FText::FromString(LOCTEXT("ToolTipMirrorOption", "Mirror Option: ").ToString() + (AssetTreeNode ? UEnum::GetDisplayValueAsText(AssetTreeNode->GetMirrorOption()).ToString() : LOCTEXT("ToolTipMirrorOption_Invalid", "Invalid").ToString()));
 	}
 
 	FReply SDatabaseAssetListItem::MirrorOptionOnMouseButtonDown(const FGeometry& InMyGeometry, const FPointerEvent& InMouseEvent)
 	{
-		const TSharedPtr<FDatabaseAssetTreeNode> Node = WeakAssetTreeNode.Pin();
-		const TSharedPtr<FDatabaseViewModel> ViewModel = EditorViewModel.Pin();
-
-		if(InMouseEvent.GetEffectingButton() == EKeys::LeftMouseButton)
+		if (const TSharedPtr<FDatabaseAssetTreeNode> AssetTreeNode = WeakAssetTreeNode.Pin())
 		{
-			if (UPoseSearchDatabase* Database = ViewModel->GetPoseSearchDatabase())
+			const TSharedPtr<FDatabaseViewModel> ViewModel = EditorViewModel.Pin();
+
+			if (InMouseEvent.GetEffectingButton() == EKeys::LeftMouseButton)
 			{
-				if (FPoseSearchDatabaseAnimationAssetBase* DatabaseAnimationAsset = Database->GetMutableAnimationAssetBase(Node->SourceAssetIdx))
+				if (UPoseSearchDatabase* Database = ViewModel->GetPoseSearchDatabase())
 				{
-					const FScopedTransaction Transaction(LOCTEXT("OnClickEditMirrorOptionPoseSearchDatabase", "Edit Mirror Option"));
-				
-					// Get next mirror option
-					static const TArray<EPoseSearchMirrorOption> OptionArray = { EPoseSearchMirrorOption::UnmirroredOnly, EPoseSearchMirrorOption::MirroredOnly, EPoseSearchMirrorOption::UnmirroredAndMirrored };
-					const int32 NextOption = (static_cast<int32>(DatabaseAnimationAsset->MirrorOption) + 1) % OptionArray.Num();
+					if (FPoseSearchDatabaseAnimationAssetBase* DatabaseAnimationAsset = Database->GetMutableAnimationAssetBase(AssetTreeNode->SourceAssetIdx))
+					{
+						const FScopedTransaction Transaction(LOCTEXT("OnClickEditMirrorOptionPoseSearchDatabase", "Edit Mirror Option"));
 
-					// Modify asset (@todo: this should be done through the viewmodel).
-					Database->Modify();
-					DatabaseAnimationAsset->MirrorOption = OptionArray[NextOption];
-				
-					SkeletonView.Pin()->RefreshTreeView(false, true);
-					ViewModel->BuildSearchIndex();
-				
-					return FReply::Handled();
+						// Get next mirror option
+						static const TArray<EPoseSearchMirrorOption> OptionArray = { EPoseSearchMirrorOption::UnmirroredOnly, EPoseSearchMirrorOption::MirroredOnly, EPoseSearchMirrorOption::UnmirroredAndMirrored };
+						const int32 NextOption = (static_cast<int32>(DatabaseAnimationAsset->MirrorOption) + 1) % OptionArray.Num();
+
+						// Modify asset (@todo: this should be done through the viewmodel).
+						Database->Modify();
+						DatabaseAnimationAsset->MirrorOption = OptionArray[NextOption];
+
+						SkeletonView.Pin()->RefreshTreeView(false, true);
+						ViewModel->BuildSearchIndex();
+
+						return FReply::Handled();
+					}
 				}
 			}
 		}
-		
 		return FReply::Unhandled();
 	}
 
diff --git a/Engine/Plugins/Animation/PoseSearch/Source/Editor/Private/PoseSearchDatabaseDataDetails.cpp b/Engine/Plugins/Animation/PoseSearch/Source/Editor/Private/PoseSearchDatabaseDataDetails.cpp
--- a/Engine/Plugins/Animation/PoseSearch/Source/Editor/Private/PoseSearchDatabaseDataDetails.cpp
+++ b/Engine/Plugins/Animation/PoseSearch/Source/Editor/Private/PoseSearchDatabaseDataDetails.cpp
@@ -318,11 +318,14 @@ public:
 											{
 												if (*GetNameSafe(PreviewActor.GetActor()) == ColumnName)
 												{
-													const FSearchIndexAsset& IndexAsset = PoseSearchDatabase->GetSearchIndex().Assets[PreviewActor.GetIndexAssetIndex()];
+													const TConstArrayView<FSearchIndexAsset> Assets = PoseSearchDatabase->GetSearchIndex().Assets;
 													const int32 AssetIdx = PreviewActor.GetIndexAssetIndex();
-													FString AssetIdxString;
-													AssetIdxString.AppendInt(IndexAsset.GetSourceAssetIdx());
-													return FText::FromString(AssetIdxString);
+													if (Assets.IsValidIndex(AssetIdx))
+													{
+														FString AssetIdxString;
+														AssetIdxString.AppendInt(Assets[AssetIdx].GetSourceAssetIdx());
+														return FText::FromString(AssetIdxString);
+													}
 												}
 											}
 										}
diff --git a/Engine/Plugins/Animation/PoseSearch/Source/Editor/Private/PoseSearchDatabaseEditorReflection.cpp b/Engine/Plugins/Animation/PoseSearch/Source/Editor/Private/PoseSearchDatabaseEditorReflection.cpp
--- a/Engine/Plugins/Animation/PoseSearch/Source/Editor/Private/PoseSearchDatabaseEditorReflection.cpp
+++ b/Engine/Plugins/Animation/PoseSearch/Source/Editor/Private/PoseSearchDatabaseEditorReflection.cpp
@@ -40,18 +40,21 @@ void UPoseSearchDatabaseReflectionBase::PostEditChangeProperty(FPropertyChangedE
 
 bool UPoseSearchDatabaseSequenceReflection::ApplyChanges() const 
 {
-	if (const TSharedPtr<UE::PoseSearch::FDatabaseViewModel> ViewModel = WeakAssetTreeNode.Pin()->EditorViewModel.Pin())
+	if (const TSharedPtr<UE::PoseSearch::FDatabaseAssetTreeNode> AssetTreeNode = WeakAssetTreeNode.Pin())
 	{
-		UPoseSearchDatabase* Database = ViewModel->GetPoseSearchDatabase();
-		if (IsValid(Database))
+		if (const TSharedPtr<UE::PoseSearch::FDatabaseViewModel> ViewModel = AssetTreeNode->EditorViewModel.Pin())
 		{
-			FInstancedStruct& DatabaseAsset = Database->GetMutableAnimationAssetStruct(WeakAssetTreeNode.Pin()->SourceAssetIdx);
-			if (FPoseSearchDatabaseSequence* DatabaseSequence = DatabaseAsset.GetMutablePtr<FPoseSearchDatabaseSequence>())
+			UPoseSearchDatabase* Database = ViewModel->GetPoseSearchDatabase();
+			if (IsValid(Database))
 			{
-				*DatabaseSequence = Sequence;
-				Database->MarkPackageDirty();
+				FInstancedStruct& DatabaseAsset = Database->GetMutableAnimationAssetStruct(AssetTreeNode->SourceAssetIdx);
+				if (FPoseSearchDatabaseSequence* DatabaseSequence = DatabaseAsset.GetMutablePtr<FPoseSearchDatabaseSequence>())
+				{
+					*DatabaseSequence = Sequence;
+					Database->MarkPackageDirty();
 
-				return true;
+					return true;
+				}
 			}
 		}
 	}
@@ -61,20 +64,23 @@ bool UPoseSearchDatabaseSequenceReflection::ApplyChanges() const
 
 bool UPoseSearchDatabaseBlendSpaceReflection::ApplyChanges() const
 {
-	if (const TSharedPtr<UE::PoseSearch::FDatabaseViewModel> ViewModel = WeakAssetTreeNode.Pin()->EditorViewModel.Pin())
+	if (const TSharedPtr<UE::PoseSearch::FDatabaseAssetTreeNode> AssetTreeNode = WeakAssetTreeNode.Pin())
 	{
-		UPoseSearchDatabase* Database = ViewModel->GetPoseSearchDatabase();
-		if (IsValid(Database))
+		if (const TSharedPtr<UE::PoseSearch::FDatabaseViewModel> ViewModel = AssetTreeNode->EditorViewModel.Pin())
 		{
-			if (UE::PoseSearch::EAsyncBuildIndexResult::Success == UE::PoseSearch::FAsyncPoseSearchDatabasesManagement::RequestAsyncBuildIndex(Database, UE::PoseSearch::ERequestAsyncBuildFlag::ContinueRequest))
+			UPoseSearchDatabase* Database = ViewModel->GetPoseSearchDatabase();
+			if (IsValid(Database))
 			{
-				FInstancedStruct& DatabaseAsset = Database->GetMutableAnimationAssetStruct(WeakAssetTreeNode.Pin()->SourceAssetIdx);
-				if (FPoseSearchDatabaseBlendSpace* DatabaseBlendSpace = DatabaseAsset.GetMutablePtr<FPoseSearchDatabaseBlendSpace>())
+				if (UE::PoseSearch::EAsyncBuildIndexResult::Success == UE::PoseSearch::FAsyncPoseSearchDatabasesManagement::RequestAsyncBuildIndex(Database, UE::PoseSearch::ERequestAsyncBuildFlag::ContinueRequest))
 				{
-					*DatabaseBlendSpace = BlendSpace;
-					Database->MarkPackageDirty();
+					FInstancedStruct& DatabaseAsset = Database->GetMutableAnimationAssetStruct(AssetTreeNode->SourceAssetIdx);
+					if (FPoseSearchDatabaseBlendSpace* DatabaseBlendSpace = DatabaseAsset.GetMutablePtr<FPoseSearchDatabaseBlendSpace>())
+					{
+						*DatabaseBlendSpace = BlendSpace;
+						Database->MarkPackageDirty();
 
-					return true;
+						return true;
+					}
 				}
 			}
 		}
@@ -85,20 +91,23 @@ bool UPoseSearchDatabaseBlendSpaceReflection::ApplyChanges() const
 
 bool UPoseSearchDatabaseAnimCompositeReflection::ApplyChanges() const
 {
-	if (const TSharedPtr<UE::PoseSearch::FDatabaseViewModel> ViewModel = WeakAssetTreeNode.Pin()->EditorViewModel.Pin())
+	if (const TSharedPtr<UE::PoseSearch::FDatabaseAssetTreeNode> AssetTreeNode = WeakAssetTreeNode.Pin())
 	{
-		UPoseSearchDatabase* Database = ViewModel->GetPoseSearchDatabase();
-		if (IsValid(Database))
+		if (const TSharedPtr<UE::PoseSearch::FDatabaseViewModel> ViewModel = AssetTreeNode->EditorViewModel.Pin())
 		{
-			if (UE::PoseSearch::EAsyncBuildIndexResult::Success == UE::PoseSearch::FAsyncPoseSearchDatabasesManagement::RequestAsyncBuildIndex(Database, UE::PoseSearch::ERequestAsyncBuildFlag::ContinueRequest))
+			UPoseSearchDatabase* Database = ViewModel->GetPoseSearchDatabase();
+			if (IsValid(Database))
 			{
-				FInstancedStruct& DatabaseAsset = Database->GetMutableAnimationAssetStruct(WeakAssetTreeNode.Pin()->SourceAssetIdx);
-				if (FPoseSearchDatabaseAnimComposite* DatabaseAnimComposite = DatabaseAsset.GetMutablePtr<FPoseSearchDatabaseAnimComposite>())
+				if (UE::PoseSearch::EAsyncBuildIndexResult::Success == UE::PoseSearch::FAsyncPoseSearchDatabasesManagement::RequestAsyncBuildIndex(Database, UE::PoseSearch::ERequestAsyncBuildFlag::ContinueRequest))
 				{
-					*DatabaseAnimComposite = AnimComposite;
-					Database->MarkPackageDirty();
+					FInstancedStruct& DatabaseAsset = Database->GetMutableAnimationAssetStruct(AssetTreeNode->SourceAssetIdx);
+					if (FPoseSearchDatabaseAnimComposite* DatabaseAnimComposite = DatabaseAsset.GetMutablePtr<FPoseSearchDatabaseAnimComposite>())
+					{
+						*DatabaseAnimComposite = AnimComposite;
+						Database->MarkPackageDirty();
 
-					return true;
+						return true;
+					}
 				}
 			}
 		}
@@ -109,20 +118,23 @@ bool UPoseSearchDatabaseAnimCompositeReflection::ApplyChanges() const
 
 bool UPoseSearchDatabaseAnimMontageReflection::ApplyChanges() const
 {
-	if (const TSharedPtr<UE::PoseSearch::FDatabaseViewModel> ViewModel = WeakAssetTreeNode.Pin()->EditorViewModel.Pin())
+	if (const TSharedPtr<UE::PoseSearch::FDatabaseAssetTreeNode> AssetTreeNode = WeakAssetTreeNode.Pin())
 	{
-		UPoseSearchDatabase* Database = ViewModel->GetPoseSearchDatabase();
-		if (IsValid(Database))
+		if (const TSharedPtr<UE::PoseSearch::FDatabaseViewModel> ViewModel = AssetTreeNode->EditorViewModel.Pin())
 		{
-			if (UE::PoseSearch::EAsyncBuildIndexResult::Success == UE::PoseSearch::FAsyncPoseSearchDatabasesManagement::RequestAsyncBuildIndex(Database, UE::PoseSearch::ERequestAsyncBuildFlag::ContinueRequest))
+			UPoseSearchDatabase* Database = ViewModel->GetPoseSearchDatabase();
+			if (IsValid(Database))
 			{
-				FInstancedStruct& DatabaseAsset = Database->GetMutableAnimationAssetStruct(WeakAssetTreeNode.Pin()->SourceAssetIdx);
-				if (FPoseSearchDatabaseAnimMontage* DatabaseAnimMontage = DatabaseAsset.GetMutablePtr<FPoseSearchDatabaseAnimMontage>())
+				if (UE::PoseSearch::EAsyncBuildIndexResult::Success == UE::PoseSearch::FAsyncPoseSearchDatabasesManagement::RequestAsyncBuildIndex(Database, UE::PoseSearch::ERequestAsyncBuildFlag::ContinueRequest))
 				{
-					*DatabaseAnimMontage = AnimMontage;
-					Database->MarkPackageDirty();
+					FInstancedStruct& DatabaseAsset = Database->GetMutableAnimationAssetStruct(AssetTreeNode->SourceAssetIdx);
+					if (FPoseSearchDatabaseAnimMontage* DatabaseAnimMontage = DatabaseAsset.GetMutablePtr<FPoseSearchDatabaseAnimMontage>())
+					{
+						*DatabaseAnimMontage = AnimMontage;
+						Database->MarkPackageDirty();
 
-					return true;
+						return true;
+					}
 				}
 			}
 		}
@@ -133,18 +145,21 @@ bool UPoseSearchDatabaseAnimMontageReflection::ApplyChanges() const
 
 bool UPoseSearchDatabaseMultiSequenceReflection::ApplyChanges() const
 {
-	if (const TSharedPtr<UE::PoseSearch::FDatabaseViewModel> ViewModel = WeakAssetTreeNode.Pin()->EditorViewModel.Pin())
+	if (const TSharedPtr<UE::PoseSearch::FDatabaseAssetTreeNode> AssetTreeNode = WeakAssetTreeNode.Pin())
 	{
-		UPoseSearchDatabase* Database = ViewModel->GetPoseSearchDatabase();
-		if (IsValid(Database))
+		if (const TSharedPtr<UE::PoseSearch::FDatabaseViewModel> ViewModel = AssetTreeNode->EditorViewModel.Pin())
 		{
-			FInstancedStruct& DatabaseAsset = Database->GetMutableAnimationAssetStruct(WeakAssetTreeNode.Pin()->SourceAssetIdx);
-			if (FPoseSearchDatabaseMultiSequence* DatabaseMultiSequence = DatabaseAsset.GetMutablePtr<FPoseSearchDatabaseMultiSequence>())
+			UPoseSearchDatabase* Database = ViewModel->GetPoseSearchDatabase();
+			if (IsValid(Database))
 			{
-				*DatabaseMultiSequence = MultiSequence;
-				Database->MarkPackageDirty();
+				FInstancedStruct& DatabaseAsset = Database->GetMutableAnimationAssetStruct(AssetTreeNode->SourceAssetIdx);
+				if (FPoseSearchDatabaseMultiSequence* DatabaseMultiSequence = DatabaseAsset.GetMutablePtr<FPoseSearchDatabaseMultiSequence>())
+				{
+					*DatabaseMultiSequence = MultiSequence;
+					Database->MarkPackageDirty();
 
-				return true;
+					return true;
+				}
 			}
 		}
 	}
diff --git a/Engine/Plugins/Animation/PoseSearch/Source/Editor/Private/PoseSearchDatabaseViewModel.cpp b/Engine/Plugins/Animation/PoseSearch/Source/Editor/Private/PoseSearchDatabaseViewModel.cpp
--- a/Engine/Plugins/Animation/PoseSearch/Source/Editor/Private/PoseSearchDatabaseViewModel.cpp
+++ b/Engine/Plugins/Animation/PoseSearch/Source/Editor/Private/PoseSearchDatabaseViewModel.cpp
@@ -32,6 +32,7 @@ namespace UE::PoseSearch
 {
 #if ENABLE_ANIM_DEBUG
 static TAutoConsoleVariable<float> CVarDatabasePreviewDebugDrawSamplerSize(TEXT("a.DatabasePreview.DebugDrawSamplerSize"), 0.f, TEXT("Debug Draw Sampler Positions Size"));
+static TAutoConsoleVariable<float> CVarDatabasePreviewDebugDrawSamplerTimeOffset(TEXT("a.DatabasePreview.DebugDrawSamplerTimeOffset"), 0.f, TEXT("Debug Draw Sampler Positions At Time Offset"));
 #endif
 
 constexpr float StepDeltaTime = 1.0f / 30.0f;
@@ -378,23 +379,32 @@ bool FDatabasePreviewActor::DrawPreviewActors(TArrayView<FDatabasePreviewActor>
 		const float DebugDrawSamplerSize = CVarDatabasePreviewDebugDrawSamplerSize.GetValueOnAnyThread();
 		if (DebugDrawSamplerSize > UE_KINDA_SMALL_NUMBER)
 		{
-			// drawing the pose extracted from the Sampler to visually compare with the pose features and the mesh drawing
+			const float DebugDrawSamplerTimeOffset = CVarDatabasePreviewDebugDrawSamplerTimeOffset.GetValueOnAnyThread();
+			
+			const int32 NumDrawPasses = FMath::IsNearlyZero(DebugDrawSamplerTimeOffset) ? 1 : 2;
+
 			FMemMark Mark(FMemStack::Get());
 			FCompactPose Pose;
-			Pose.SetBoneContainer(&PreviewActor.GetAnimPreviewInstance()->GetRequiredBonesOnAnyThread());
+			FCSPose<FCompactPose> ComponentSpacePose;
 
-			PreviewActor.Sampler.ExtractPose(PreviewActor.CurrentTime, Pose);
-			MirrorDataCache.MirrorPose(Pose);
+			for (int32 DrawPass = 0; DrawPass < NumDrawPasses; ++DrawPass)
+			{
+				// drawing the pose extracted from the Sampler to visually compare with the pose features and the mesh drawing
+				Pose.SetBoneContainer(&PreviewActor.GetAnimPreviewInstance()->GetRequiredBonesOnAnyThread());
 
-			const FTransform RootTransform = MirrorDataCache.MirrorTransform(PreviewActor.Sampler.ExtractRootTransform(PreviewActor.CurrentTime));
+				const float SamplerTime = DrawPass ? PreviewActor.CurrentTime + DebugDrawSamplerTimeOffset : PreviewActor.CurrentTime;
+				const FColor DebugColor = DrawPass ? FColor::Blue : FColor::Red;
 
-			FCSPose<FCompactPose> ComponentSpacePose;
-			ComponentSpacePose.InitPose(MoveTemp(Pose));
+				PreviewActor.Sampler.ExtractPose(SamplerTime, Pose);
+				MirrorDataCache.MirrorPose(Pose);
+				ComponentSpacePose.InitPose(MoveTemp(Pose));
 
-			for (int32 BoneIndex = 0; BoneIndex < ComponentSpacePose.GetPose().GetNumBones(); ++BoneIndex)
-			{
-				const FTransform BoneWorldTransforms = ComponentSpacePose.GetComponentSpaceTransform(FCompactPoseBoneIndex(BoneIndex)) * RootTransform;
-				DrawParams.DrawPoint(BoneWorldTransforms.GetTranslation(), FColor::Red, DebugDrawSamplerSize);
+				const FTransform RootTransform = MirrorDataCache.MirrorTransform(PreviewActor.Sampler.ExtractRootTransform(SamplerTime));
+				for (int32 BoneIndex = 0; BoneIndex < ComponentSpacePose.GetPose().GetNumBones(); ++BoneIndex)
+				{
+					const FTransform BoneWorldTransforms = ComponentSpacePose.GetComponentSpaceTransform(FCompactPoseBoneIndex(BoneIndex)) * RootTransform;
+					DrawParams.DrawPoint(BoneWorldTransforms.GetTranslation(), DebugColor, DebugDrawSamplerSize);
+				}
 			}
 		}
 #endif // ENABLE_ANIM_DEBUG
