motion matching - fix for error: extra tokens at end of #endif directive [-Werror,-Wextra-tokens] #endif ENABLE_ANIM_DEBUG

diff --git a/Engine/Plugins/Animation/PoseSearch/Source/Runtime/Private/PoseSearchHistory.cpp b/Engine/Plugins/Animation/PoseSearch/Source/Runtime/Private/PoseSearchHistory.cpp
--- a/Engine/Plugins/Animation/PoseSearch/Source/Runtime/Private/PoseSearchHistory.cpp
+++ b/Engine/Plugins/Animation/PoseSearch/Source/Runtime/Private/PoseSearchHistory.cpp
@@ -472,7 +472,7 @@ FPoseHistory& FPoseHistory::operator=(FPoseHistory&& Other)
 	FThreadSafeCounter& OtherWriteDataThreadSafeCounter = Other.WriteDataThreadSafeCounter;
 	CheckThreadSafetyWrite(OtherReadDataThreadSafeCounter);
 	CheckThreadSafetyWrite(OtherWriteDataThreadSafeCounter);
-#endif ENABLE_ANIM_DEBUG
+#endif // ENABLE_ANIM_DEBUG
 
 	MaxNumPoses = MoveTemp(Other.MaxNumPoses);
 	SamplingInterval = MoveTemp(Other.SamplingInterval);
