motion matching - adding UE::Anim::Compression::AnimationCompressionVersionString dependency to UE::PoseSearch::FKeyBuilder::FKeyBuilder

diff --git a/Engine/Plugins/Animation/PoseSearch/Source/Runtime/Private/PoseSearchDerivedDataKey.cpp b/Engine/Plugins/Animation/PoseSearch/Source/Runtime/Private/PoseSearchDerivedDataKey.cpp
--- a/Engine/Plugins/Animation/PoseSearch/Source/Runtime/Private/PoseSearchDerivedDataKey.cpp
+++ b/Engine/Plugins/Animation/PoseSearch/Source/Runtime/Private/PoseSearchDerivedDataKey.cpp
@@ -7,6 +7,7 @@
 #include "Animation/AnimCurveCompressionSettings.h"
 #include "Animation/AnimData/IAnimationDataModel.h"
 #include "AnimationModifier.h"
+#include "Animation/AnimCompress.h"
 #include "Animation/MirrorDataTable.h"
 #include "Animation/Skeleton.h"
 #include "Engine/SkeletalMesh.h"
@@ -43,8 +44,10 @@ FKeyBuilder::FKeyBuilder(const UObject* Object, bool bUseDataVer, bool bPerformC
 		// used to invalidate the key without having to change POSESEARCHDB_DERIVEDDATA_VER all the times
 		int32 POSESEARCHDB_DERIVEDDATA_VER_SMALL = 247;
 		FGuid VersionGuid = FDevSystemGuids::GetSystemGuid(FDevSystemGuids::Get().POSESEARCHDB_DERIVEDDATA_VER);
+		FString AnimationCompressionVersionString = UE::Anim::Compression::AnimationCompressionVersionString;
 
 		*this << VersionGuid;
+		*this << AnimationCompressionVersionString;
 		*this << POSESEARCHDB_DERIVEDDATA_VER_SMALL;
 	}
 
