motion matching - fix for schemas referencing invalid skeleton bones

diff --git a/Engine/Plugins/Animation/PoseSearch/Source/Runtime/Private/PoseSearchFeatureChannel_Heading.cpp b/Engine/Plugins/Animation/PoseSearch/Source/Runtime/Private/PoseSearchFeatureChannel_Heading.cpp
--- a/Engine/Plugins/Animation/PoseSearch/Source/Runtime/Private/PoseSearchFeatureChannel_Heading.cpp
+++ b/Engine/Plugins/Animation/PoseSearch/Source/Runtime/Private/PoseSearchFeatureChannel_Heading.cpp
@@ -258,7 +258,7 @@ UE::PoseSearch::TLabelBuilder& UPoseSearchFeatureChannel_Heading::GetLabel(UE::P
 
 	const UPoseSearchSchema* Schema = GetSchema();
 	check(Schema);
-	if (SchemaBoneIdx != RootSchemaBoneIdx)
+	if (SchemaBoneIdx > RootSchemaBoneIdx)
 	{
 		LabelBuilder.Append(TEXT("_"));
 		LabelBuilder.Append(Schema->GetBoneReferences(SampleRole)[SchemaBoneIdx].BoneName.ToString());
@@ -280,7 +280,7 @@ UE::PoseSearch::TLabelBuilder& UPoseSearchFeatureChannel_Heading::GetLabel(UE::P
 		LabelBuilder.Append(TEXT("]"));
 	}
 
-	if (SchemaOriginBoneIdx != RootSchemaBoneIdx)
+	if (SchemaOriginBoneIdx > RootSchemaBoneIdx)
 	{
 		LabelBuilder.Append(TEXT("_"));
 		LabelBuilder.Append(Schema->GetBoneReferences(OriginRole)[SchemaOriginBoneIdx].BoneName.ToString());
diff --git a/Engine/Plugins/Animation/PoseSearch/Source/Runtime/Private/PoseSearchFeatureChannel_Phase.cpp b/Engine/Plugins/Animation/PoseSearch/Source/Runtime/Private/PoseSearchFeatureChannel_Phase.cpp
--- a/Engine/Plugins/Animation/PoseSearch/Source/Runtime/Private/PoseSearchFeatureChannel_Phase.cpp
+++ b/Engine/Plugins/Animation/PoseSearch/Source/Runtime/Private/PoseSearchFeatureChannel_Phase.cpp
@@ -449,7 +449,7 @@ UE::PoseSearch::TLabelBuilder& UPoseSearchFeatureChannel_Phase::GetLabel(UE::Pos
 
 	const UPoseSearchSchema* Schema = GetSchema();
 	check(Schema);
-	if (SchemaBoneIdx != RootSchemaBoneIdx)
+	if (SchemaBoneIdx > RootSchemaBoneIdx)
 	{
 		LabelBuilder.Append(TEXT("_"));
 		LabelBuilder.Append(Schema->GetBoneReferences(SampleRole)[SchemaBoneIdx].BoneName.ToString());
diff --git a/Engine/Plugins/Animation/PoseSearch/Source/Runtime/Private/PoseSearchFeatureChannel_Position.cpp b/Engine/Plugins/Animation/PoseSearch/Source/Runtime/Private/PoseSearchFeatureChannel_Position.cpp
--- a/Engine/Plugins/Animation/PoseSearch/Source/Runtime/Private/PoseSearchFeatureChannel_Position.cpp
+++ b/Engine/Plugins/Animation/PoseSearch/Source/Runtime/Private/PoseSearchFeatureChannel_Position.cpp
@@ -231,7 +231,7 @@ UE::PoseSearch::TLabelBuilder& UPoseSearchFeatureChannel_Position::GetLabel(UE::
 
 	const UPoseSearchSchema* Schema = GetSchema();
 	check(Schema);
-	if (SchemaBoneIdx != RootSchemaBoneIdx)
+	if (SchemaBoneIdx > RootSchemaBoneIdx)
 	{
 		LabelBuilder.Append(TEXT("_"));
 		LabelBuilder.Append(Schema->GetBoneReferences(SampleRole)[SchemaBoneIdx].BoneName.ToString());
@@ -244,7 +244,7 @@ UE::PoseSearch::TLabelBuilder& UPoseSearchFeatureChannel_Position::GetLabel(UE::
 		LabelBuilder.Append(TEXT("]"));
 	}
 
-	if (SchemaOriginBoneIdx != RootSchemaBoneIdx)
+	if (SchemaOriginBoneIdx > RootSchemaBoneIdx)
 	{
 		LabelBuilder.Append(TEXT("_"));
 		LabelBuilder.Append(Schema->GetBoneReferences(OriginRole)[SchemaOriginBoneIdx].BoneName.ToString());
diff --git a/Engine/Plugins/Animation/PoseSearch/Source/Runtime/Private/PoseSearchFeatureChannel_Velocity.cpp b/Engine/Plugins/Animation/PoseSearch/Source/Runtime/Private/PoseSearchFeatureChannel_Velocity.cpp
--- a/Engine/Plugins/Animation/PoseSearch/Source/Runtime/Private/PoseSearchFeatureChannel_Velocity.cpp
+++ b/Engine/Plugins/Animation/PoseSearch/Source/Runtime/Private/PoseSearchFeatureChannel_Velocity.cpp
@@ -207,7 +207,7 @@ UE::PoseSearch::TLabelBuilder& UPoseSearchFeatureChannel_Velocity::GetLabel(UE::
 
 	const UPoseSearchSchema* Schema = GetSchema();
 	check(Schema);
-	if (SchemaBoneIdx != RootSchemaBoneIdx)
+	if (SchemaBoneIdx > RootSchemaBoneIdx)
 	{
 		LabelBuilder.Append(TEXT("_"));
 		LabelBuilder.Append(Schema->GetBoneReferences(SampleRole)[SchemaBoneIdx].BoneName.ToString());
@@ -229,7 +229,7 @@ UE::PoseSearch::TLabelBuilder& UPoseSearchFeatureChannel_Velocity::GetLabel(UE::
 		LabelBuilder.Append(TEXT("]"));
 	}
 
-	if (SchemaOriginBoneIdx != RootSchemaBoneIdx)
+	if (SchemaOriginBoneIdx > RootSchemaBoneIdx)
 	{
 		LabelBuilder.Append(TEXT("_"));
 		LabelBuilder.Append(Schema->GetBoneReferences(OriginRole)[SchemaOriginBoneIdx].BoneName.ToString());
