Make database items with incompatible skeletons have a red tint on their name

[RN] minor

diff --git a/Engine/Plugins/Animation/PoseSearch/Source/Editor/Private/PoseSearchDatabaseAssetListItem.cpp b/Engine/Plugins/Animation/PoseSearch/Source/Editor/Private/PoseSearchDatabaseAssetListItem.cpp
--- a/Engine/Plugins/Animation/PoseSearch/Source/Editor/Private/PoseSearchDatabaseAssetListItem.cpp
+++ b/Engine/Plugins/Animation/PoseSearch/Source/Editor/Private/PoseSearchDatabaseAssetListItem.cpp
@@ -503,6 +503,11 @@ namespace UE::PoseSearch
 							return FColor::Turquoise;
 						}
 
+						if (!DatabaseAnimationAssetBase->IsSkeletonCompatible(Database->Schema))
+						{
+							return FColor::Red;
+						}
+						
 						return FLinearColor::White;
 					}
 				}
diff --git a/Engine/Plugins/Animation/PoseSearch/Source/Runtime/Private/PoseSearchDatabase.cpp b/Engine/Plugins/Animation/PoseSearch/Source/Runtime/Private/PoseSearchDatabase.cpp
--- a/Engine/Plugins/Animation/PoseSearch/Source/Runtime/Private/PoseSearchDatabase.cpp
+++ b/Engine/Plugins/Animation/PoseSearch/Source/Runtime/Private/PoseSearchDatabase.cpp
@@ -330,6 +330,38 @@ int32 FPoseSearchDatabaseAnimationAssetBase::GetFrameAtTime(float Time) const
 	}
 	return 0.f;
 }
+
+bool FPoseSearchDatabaseAnimationAssetBase::IsSkeletonCompatible(TObjectPtr<const UPoseSearchSchema> InSchema) const
+{
+	if (InSchema)
+	{
+		TArray<FPoseSearchRoledSkeleton> RoledSkeletons = InSchema->GetRoledSkeletons();
+
+		if (GetAnimationAsset())
+		{
+			const int32 NumRoles = GetNumRoles();
+			for (int RoleIdx = 0; RoleIdx < NumRoles; ++RoleIdx)
+			{
+				UE::PoseSearch::FRole Role = GetRole(RoleIdx);
+				FAssetData AssetData = IAssetRegistry::Get()->GetAssetByObjectPath(FSoftObjectPath(GetAnimationAssetForRole(Role)));
+		
+				for (const FPoseSearchRoledSkeleton& RoledSkeleton : RoledSkeletons)
+				{
+					if (RoledSkeleton.Role == Role)
+					{
+						// Match skeleton
+						if (RoledSkeleton.Skeleton->IsCompatibleForEditor(AssetData))
+						{
+							return true;
+						}
+					}
+				}
+			}
+		}
+	}
+	
+	return false;
+}
 #endif // WITH_EDITOR
 
 UAnimationAsset* FPoseSearchDatabaseAnimationAssetBase::GetAnimationAssetForRole(const UE::PoseSearch::FRole& Role) const
diff --git a/Engine/Plugins/Animation/PoseSearch/Source/Runtime/Private/PoseSearchSchema.cpp b/Engine/Plugins/Animation/PoseSearch/Source/Runtime/Private/PoseSearchSchema.cpp
--- a/Engine/Plugins/Animation/PoseSearch/Source/Runtime/Private/PoseSearchSchema.cpp
+++ b/Engine/Plugins/Animation/PoseSearch/Source/Runtime/Private/PoseSearchSchema.cpp
@@ -130,6 +130,11 @@ const UE::PoseSearch::FRole UPoseSearchSchema::GetDefaultRole() const
 	}
 	return UE::PoseSearch::DefaultRole;
 }
+
+const TArray<FPoseSearchRoledSkeleton>& UPoseSearchSchema::GetRoledSkeletons() const
+{
+	return Skeletons;
+}
 #endif // WITH_EDITOR
 
 USkeleton* UPoseSearchSchema::GetSkeleton(const UE::PoseSearch::FRole& Role) const
diff --git a/Engine/Plugins/Animation/PoseSearch/Source/Runtime/Public/PoseSearch/PoseSearchDatabase.h b/Engine/Plugins/Animation/PoseSearch/Source/Runtime/Public/PoseSearch/PoseSearchDatabase.h
--- a/Engine/Plugins/Animation/PoseSearch/Source/Runtime/Public/PoseSearch/PoseSearchDatabase.h
+++ b/Engine/Plugins/Animation/PoseSearch/Source/Runtime/Public/PoseSearch/PoseSearchDatabase.h
@@ -61,6 +61,7 @@ struct POSESEARCH_API FPoseSearchDatabaseAnimationAssetBase
 
 #if WITH_EDITOR
 	virtual int32 GetFrameAtTime(float Time) const;
+	virtual bool IsSkeletonCompatible(TObjectPtr<const UPoseSearchSchema> InSchema) const;
 #endif // WITH_EDITOR
 
 #if WITH_EDITORONLY_DATA
diff --git a/Engine/Plugins/Animation/PoseSearch/Source/Runtime/Public/PoseSearch/PoseSearchSchema.h b/Engine/Plugins/Animation/PoseSearch/Source/Runtime/Public/PoseSearch/PoseSearchSchema.h
--- a/Engine/Plugins/Animation/PoseSearch/Source/Runtime/Public/PoseSearch/PoseSearchSchema.h
+++ b/Engine/Plugins/Animation/PoseSearch/Source/Runtime/Public/PoseSearch/PoseSearchSchema.h
@@ -144,6 +144,7 @@ public:
 #if WITH_EDITOR
 	virtual void PostEditChangeProperty(struct FPropertyChangedEvent& PropertyChangedEvent) override;
 	const UE::PoseSearch::FRole GetDefaultRole() const;
+	const TArray<FPoseSearchRoledSkeleton> & GetRoledSkeletons() const;
 #endif // WITH_EDITOR
 
 	TConstArrayView<float> BuildQuery(UE::PoseSearch::FSearchContext& SearchContext) const;
