Fixed issue where adding a blendpsace, montage or composite without a previously cached index fails.

[RN] minor

diff --git a/Engine/Plugins/Animation/PoseSearch/Source/Editor/Private/PoseSearchDatabaseEditorReflection.cpp b/Engine/Plugins/Animation/PoseSearch/Source/Editor/Private/PoseSearchDatabaseEditorReflection.cpp
--- a/Engine/Plugins/Animation/PoseSearch/Source/Editor/Private/PoseSearchDatabaseEditorReflection.cpp
+++ b/Engine/Plugins/Animation/PoseSearch/Source/Editor/Private/PoseSearchDatabaseEditorReflection.cpp
@@ -71,16 +71,13 @@ bool UPoseSearchDatabaseBlendSpaceReflection::ApplyChanges() const
 			UPoseSearchDatabase* Database = ViewModel->GetPoseSearchDatabase();
 			if (IsValid(Database))
 			{
-				if (UE::PoseSearch::EAsyncBuildIndexResult::Success == UE::PoseSearch::FAsyncPoseSearchDatabasesManagement::RequestAsyncBuildIndex(Database, UE::PoseSearch::ERequestAsyncBuildFlag::ContinueRequest))
+				FInstancedStruct& DatabaseAsset = Database->GetMutableAnimationAssetStruct(AssetTreeNode->SourceAssetIdx);
+				if (FPoseSearchDatabaseBlendSpace* DatabaseBlendSpace = DatabaseAsset.GetMutablePtr<FPoseSearchDatabaseBlendSpace>())
 				{
-					FInstancedStruct& DatabaseAsset = Database->GetMutableAnimationAssetStruct(AssetTreeNode->SourceAssetIdx);
-					if (FPoseSearchDatabaseBlendSpace* DatabaseBlendSpace = DatabaseAsset.GetMutablePtr<FPoseSearchDatabaseBlendSpace>())
-					{
-						*DatabaseBlendSpace = BlendSpace;
-						Database->MarkPackageDirty();
+					*DatabaseBlendSpace = BlendSpace;
+					Database->MarkPackageDirty();
 
-						return true;
-					}
+					return true;
 				}
 			}
 		}
@@ -98,16 +95,13 @@ bool UPoseSearchDatabaseAnimCompositeReflection::ApplyChanges() const
 			UPoseSearchDatabase* Database = ViewModel->GetPoseSearchDatabase();
 			if (IsValid(Database))
 			{
-				if (UE::PoseSearch::EAsyncBuildIndexResult::Success == UE::PoseSearch::FAsyncPoseSearchDatabasesManagement::RequestAsyncBuildIndex(Database, UE::PoseSearch::ERequestAsyncBuildFlag::ContinueRequest))
+				FInstancedStruct& DatabaseAsset = Database->GetMutableAnimationAssetStruct(AssetTreeNode->SourceAssetIdx);
+				if (FPoseSearchDatabaseAnimComposite* DatabaseAnimComposite = DatabaseAsset.GetMutablePtr<FPoseSearchDatabaseAnimComposite>())
 				{
-					FInstancedStruct& DatabaseAsset = Database->GetMutableAnimationAssetStruct(AssetTreeNode->SourceAssetIdx);
-					if (FPoseSearchDatabaseAnimComposite* DatabaseAnimComposite = DatabaseAsset.GetMutablePtr<FPoseSearchDatabaseAnimComposite>())
-					{
-						*DatabaseAnimComposite = AnimComposite;
-						Database->MarkPackageDirty();
+					*DatabaseAnimComposite = AnimComposite;
+					Database->MarkPackageDirty();
 
-						return true;
-					}
+					return true;
 				}
 			}
 		}
@@ -125,16 +119,13 @@ bool UPoseSearchDatabaseAnimMontageReflection::ApplyChanges() const
 			UPoseSearchDatabase* Database = ViewModel->GetPoseSearchDatabase();
 			if (IsValid(Database))
 			{
-				if (UE::PoseSearch::EAsyncBuildIndexResult::Success == UE::PoseSearch::FAsyncPoseSearchDatabasesManagement::RequestAsyncBuildIndex(Database, UE::PoseSearch::ERequestAsyncBuildFlag::ContinueRequest))
+				FInstancedStruct& DatabaseAsset = Database->GetMutableAnimationAssetStruct(AssetTreeNode->SourceAssetIdx);
+				if (FPoseSearchDatabaseAnimMontage* DatabaseAnimMontage = DatabaseAsset.GetMutablePtr<FPoseSearchDatabaseAnimMontage>())
 				{
-					FInstancedStruct& DatabaseAsset = Database->GetMutableAnimationAssetStruct(AssetTreeNode->SourceAssetIdx);
-					if (FPoseSearchDatabaseAnimMontage* DatabaseAnimMontage = DatabaseAsset.GetMutablePtr<FPoseSearchDatabaseAnimMontage>())
-					{
-						*DatabaseAnimMontage = AnimMontage;
-						Database->MarkPackageDirty();
+					*DatabaseAnimMontage = AnimMontage;
+					Database->MarkPackageDirty();
 
-						return true;
-					}
+					return true;
 				}
 			}
 		}
