Added a content browser to the pose search database editor

[RN] minor

diff --git a/Engine/Plugins/Animation/PoseSearch/Source/Editor/Private/PoseSearchDatabaseAssetBrowser.cpp b/Engine/Plugins/Animation/PoseSearch/Source/Editor/Private/PoseSearchDatabaseAssetBrowser.cpp
new file mode 100644
--- /dev/null
+++ b/Engine/Plugins/Animation/PoseSearch/Source/Editor/Private/PoseSearchDatabaseAssetBrowser.cpp
@@ -0,0 +1,140 @@
+﻿// Copyright Epic Games, Inc. All Rights Reserved.
+
+#include "PoseSearchDatabaseAssetBrowser.h"
+
+#include "ContentBrowserDataSource.h"
+#include "ContentBrowserModule.h"
+#include "IContentBrowserSingleton.h"
+#include "Animation/AnimComposite.h"
+#include "Animation/AnimSequence.h"
+#include "Animation/AnimMontage.h"
+#include "Animation/BlendSpace.h"
+#include "Widgets/SBoxPanel.h"
+#include "UObject/AssetRegistryTagsContext.h"
+#include "PoseSearchDatabaseViewModel.h"
+#include "PoseSearch/PoseSearchDatabase.h"
+#include "PoseSearch/PoseSearchSchema.h"
+#include "Subsystems/AssetEditorSubsystem.h"
+#include "Editor.h"
+
+namespace UE::PoseSearch
+{
+	
+void SPoseSearchDatabaseAssetBrowser::Construct(const FArguments& InArgs, TSharedPtr<FDatabaseViewModel> InViewModel)
+{
+	DatabaseViewModel = InViewModel;
+	
+	ChildSlot
+		[
+			SNew(SVerticalBox)
+			+ SVerticalBox::Slot()
+			[
+				SAssignNew(AssetBrowserBox, SBox)
+			]
+		];
+
+	RefreshView();
+}
+
+void SPoseSearchDatabaseAssetBrowser::RefreshView()
+{
+	// @TODO: Add support for MultiSequence.
+	
+	FAssetPickerConfig AssetPickerConfig;
+	
+	AssetPickerConfig.Filter.ClassPaths.Add(UAnimationAsset::StaticClass()->GetClassPathName());
+	AssetPickerConfig.Filter.bRecursiveClasses = true;
+	
+	AssetPickerConfig.bAddFilterUI = true;
+	AssetPickerConfig.bAllowNullSelection = false;
+	AssetPickerConfig.InitialAssetViewType = EAssetViewType::Column;
+	AssetPickerConfig.bShowPathInColumnView = true;
+	AssetPickerConfig.bShowTypeInColumnView = false;
+	AssetPickerConfig.InitialThumbnailSize = EThumbnailSize::Small;
+	
+	AssetPickerConfig.OnShouldFilterAsset = FOnShouldFilterAsset::CreateSP(this, &SPoseSearchDatabaseAssetBrowser::OnShouldFilterAsset);
+	AssetPickerConfig.OnAssetDoubleClicked = FOnAssetSelected::CreateSP(this, &SPoseSearchDatabaseAssetBrowser::OnAssetDoubleClicked);
+
+	// Hide all asset registry columns by default (we only really want the name and path)
+	const UObject* AnimSequenceDefaultObject = UAnimSequence::StaticClass()->GetDefaultObject();
+	FAssetRegistryTagsContextData TagsContext(AnimSequenceDefaultObject, EAssetRegistryTagsCaller::Uncategorized);
+	AnimSequenceDefaultObject->GetAssetRegistryTags(TagsContext);
+	for (const TPair<FName, UObject::FAssetRegistryTag>& TagPair : TagsContext.Tags)
+	{
+		AssetPickerConfig.HiddenColumnNames.Add(TagPair.Key.ToString());
+	}
+
+	// Also hide the type column by default (but allow users to enable it, so don't use bShowTypeInColumnView)
+	AssetPickerConfig.HiddenColumnNames.Add(TEXT("Class"));
+	AssetPickerConfig.HiddenColumnNames.Add(ContentBrowserItemAttributes::ItemDiskSize.ToString());
+	AssetPickerConfig.HiddenColumnNames.Add(ContentBrowserItemAttributes::VirtualizedData.ToString());
+
+	FContentBrowserModule& ContentBrowserModule = FModuleManager::Get().LoadModuleChecked<FContentBrowserModule>(TEXT("ContentBrowser"));
+	AssetBrowserBox->SetContent(ContentBrowserModule.Get().CreateAssetPicker(AssetPickerConfig));
+}
+
+void SPoseSearchDatabaseAssetBrowser::OnAssetDoubleClicked(const FAssetData& AssetData)
+{
+	if (!AssetData.GetAsset())
+	{
+		return;
+	}
+
+	UAnimationAsset* NewAnimationAsset = Cast<UAnimationAsset>(AssetData.GetAsset());
+	if (!NewAnimationAsset)
+	{
+		return;
+	}
+
+	// Just open asset in persona.
+	GEditor->GetEditorSubsystem<UAssetEditorSubsystem>()->OpenEditorForAsset(NewAnimationAsset);
+}
+
+bool SPoseSearchDatabaseAssetBrowser::OnShouldFilterAsset(const FAssetData& AssetData)
+{
+	bool bAssetHasCompatibleSkeleton = false;
+
+	if (DatabaseViewModel)
+	{
+		if (const UPoseSearchDatabase* Database = DatabaseViewModel->GetPoseSearchDatabase())
+		{
+			if (Database->Schema)
+			{
+				TArray<FPoseSearchRoledSkeleton> RoledSkeletons = Database->Schema->GetRoledSkeletons();
+
+				for (const FPoseSearchRoledSkeleton& RoledSkeleton : RoledSkeletons)
+				{
+					if (RoledSkeleton.Skeleton && RoledSkeleton.Skeleton->IsCompatibleForEditor(AssetData))
+					{
+						// We found a compatible skeleton in the schema.
+						bAssetHasCompatibleSkeleton = true;
+					}
+				}
+			}
+		}
+	}
+	
+	if (AssetData.GetClass()->IsChildOf(UAnimSequence::StaticClass()) && bAssetHasCompatibleSkeleton)
+	{
+		return false;
+	}
+
+	if (AssetData.GetClass()->IsChildOf(UAnimComposite::StaticClass()) && bAssetHasCompatibleSkeleton)
+	{
+		return false;
+	}
+
+	if (AssetData.GetClass()->IsChildOf(UAnimMontage::StaticClass()) && bAssetHasCompatibleSkeleton)
+	{
+		return false;
+	}
+	
+	if (AssetData.GetClass()->IsChildOf(UBlendSpace::StaticClass()) && bAssetHasCompatibleSkeleton)
+	{
+		return false;
+	}
+	
+	return true;
+}
+	
+}
diff --git a/Engine/Plugins/Animation/PoseSearch/Source/Editor/Private/PoseSearchDatabaseAssetBrowser.h b/Engine/Plugins/Animation/PoseSearch/Source/Editor/Private/PoseSearchDatabaseAssetBrowser.h
new file mode 100644
--- /dev/null
+++ b/Engine/Plugins/Animation/PoseSearch/Source/Editor/Private/PoseSearchDatabaseAssetBrowser.h
@@ -0,0 +1,33 @@
+﻿// Copyright Epic Games, Inc. All Rights Reserved.
+
+#pragma once
+
+#include "CoreMinimal.h"
+#include "Widgets/Layout/SBox.h"
+
+struct FAssetData;
+
+namespace UE::PoseSearch
+{
+	class FDatabaseViewModel;
+	
+	class SPoseSearchDatabaseAssetBrowser : public SBox
+	{
+	public:
+	
+		SLATE_BEGIN_ARGS(SPoseSearchDatabaseAssetBrowser) {}
+		SLATE_END_ARGS();
+
+		void Construct(const FArguments& InArgs, TSharedPtr<FDatabaseViewModel> InViewModel);
+		
+		void RefreshView();
+		
+	private:
+
+		TSharedPtr<SBox> AssetBrowserBox;
+		TSharedPtr<FDatabaseViewModel> DatabaseViewModel;
+		
+		void OnAssetDoubleClicked(const FAssetData& AssetData);
+		bool OnShouldFilterAsset(const FAssetData& AssetData);
+	};
+}
\ No newline at end of file
diff --git a/Engine/Plugins/Animation/PoseSearch/Source/Editor/Private/PoseSearchDatabaseEditor.cpp b/Engine/Plugins/Animation/PoseSearch/Source/Editor/Private/PoseSearchDatabaseEditor.cpp
--- a/Engine/Plugins/Animation/PoseSearch/Source/Editor/Private/PoseSearchDatabaseEditor.cpp
+++ b/Engine/Plugins/Animation/PoseSearch/Source/Editor/Private/PoseSearchDatabaseEditor.cpp
@@ -6,6 +6,7 @@
 #include "GameFramework/WorldSettings.h"
 #include "IStructureDetailsView.h"
 #include "InstancedStruct.h"
+#include "PoseSearchDatabaseAssetBrowser.h"
 #include "Modules/ModuleManager.h"
 #include "PoseSearch/PoseSearchDerivedData.h"
 #include "PoseSearchDatabaseAssetTree.h"
@@ -38,6 +39,7 @@ namespace UE::PoseSearch
 		static const FName SelectionDetailsID;
 		static const FName StatisticsOverview;
 		static const FName DataDetailsID;
+		static const FName AssetBrowserID;
 	};
 
 	const FName FDatabaseEditorTabs::AssetDetailsID(TEXT("PoseSearchDatabaseEditorAssetDetailsTabID"));
@@ -47,6 +49,7 @@ namespace UE::PoseSearch
 	const FName FDatabaseEditorTabs::SelectionDetailsID(TEXT("PoseSearchDatabaseEditorSelectionDetailsID"));
 	const FName FDatabaseEditorTabs::StatisticsOverview(TEXT("PoseSearchDatabaseEditorStatisticsOverviewID"));
 	const FName FDatabaseEditorTabs::DataDetailsID(TEXT("PoseSearchDatabaseEditorDataDetailsTabID"));
+	const FName FDatabaseEditorTabs::AssetBrowserID(TEXT("PoseSearchDatabaseEditorAssetBrowserID"));
 
 	const UPoseSearchDatabase* FDatabaseEditor::GetPoseSearchDatabase() const
 	{
@@ -226,51 +229,70 @@ namespace UE::PoseSearch
 				DatabaseAsset->RegisterOnSynchronizeWithExternalDependencies(UPoseSearchDatabase::FOnDerivedDataRebuild::CreateSP(this, &FDatabaseEditor::RefreshEditor));
 			}
 		}
+
+		// Create asset browser widget
+		AssetBrowserWidget = SNew(SPoseSearchDatabaseAssetBrowser, ViewModel);
 		
 		// Define Editor Layout
 		const TSharedRef<FTabManager::FLayout> StandaloneDefaultLayout =
-		FTabManager::NewLayout("Standalone_PoseSearchDatabaseEditor_Layout_v0.13")
+		FTabManager::NewLayout("Standalone_PoseSearchDatabaseEditor_Layout_v0.14")
 			->AddArea
 			(
-				FTabManager::NewPrimaryArea()->SetOrientation(Orient_Horizontal)
-				->Split
-				(
-					FTabManager::NewSplitter()->SetOrientation(Orient_Vertical)
-					->SetSizeCoefficient(0.4f)
-					->Split
-					(
-						FTabManager::NewStack()
-						->SetSizeCoefficient(0.6f)
-						->AddTab(FDatabaseEditorTabs::AssetTreeViewID, ETabState::OpenedTab)
-						->SetHideTabWell(false)
-					)
-					->Split
-					(
-						FTabManager::NewStack()
-						->SetSizeCoefficient(0.4f)
-						->AddTab(FDatabaseEditorTabs::AssetDetailsID, ETabState::OpenedTab)
-						->SetHideTabWell(false)
-					)
-				)
+				// Main application area
+				FTabManager::NewPrimaryArea()
+				->SetOrientation(Orient_Vertical)
 				->Split
 				(
-					FTabManager::NewSplitter()->SetOrientation(Orient_Vertical)
-					->SetSizeCoefficient(0.6f)
+					FTabManager::NewSplitter()
+					->SetOrientation(Orient_Horizontal)
 					->Split
 					(
-						FTabManager::NewStack()
-						->SetSizeCoefficient(0.6f)
-						->AddTab(FDatabaseEditorTabs::ViewportID, ETabState::OpenedTab)
-						->SetHideTabWell(true)
+						FTabManager::NewSplitter()->SetOrientation(Orient_Vertical)
+						->SetSizeCoefficient(0.2f)
+						->Split
+						(
+							FTabManager::NewStack()
+							->SetSizeCoefficient(0.6f)
+							->AddTab(FDatabaseEditorTabs::AssetTreeViewID, ETabState::OpenedTab)
+							->SetHideTabWell(false)
+						)
+						->Split
+						(
+							FTabManager::NewStack()
+							->SetSizeCoefficient(0.4f)
+							->AddTab(FDatabaseEditorTabs::AssetDetailsID, ETabState::OpenedTab)
+							->SetHideTabWell(false)
+						)
 					)
 					->Split
 					(
-						FTabManager::NewStack()
-						->SetSizeCoefficient(0.4f)
-						->AddTab(FDatabaseEditorTabs::StatisticsOverview, ETabState::OpenedTab)
-						->AddTab(FDatabaseEditorTabs::PreviewSettingsID, ETabState::OpenedTab)
-						->AddTab(FDatabaseEditorTabs::DataDetailsID, ETabState::OpenedTab)
-						->AddTab(FDatabaseEditorTabs::SelectionDetailsID, ETabState::OpenedTab)
+						FTabManager::NewSplitter()->SetOrientation(Orient_Vertical)
+						->SetSizeCoefficient(0.8f)
+						->Split
+						(
+							FTabManager::NewSplitter()->SetOrientation(Orient_Horizontal)
+							->Split(
+								FTabManager::NewStack()
+								->SetSizeCoefficient(0.75f)
+								->AddTab(FDatabaseEditorTabs::ViewportID, ETabState::OpenedTab)
+								->SetHideTabWell(true)
+									
+							)
+							->Split(
+								FTabManager::NewStack()
+								->SetSizeCoefficient(0.25f)
+								->AddTab(FDatabaseEditorTabs::AssetBrowserID, ETabState::OpenedTab)
+							)
+						)
+						->Split
+						(
+							FTabManager::NewStack()
+							->SetSizeCoefficient(0.4f)
+							->AddTab(FDatabaseEditorTabs::StatisticsOverview, ETabState::OpenedTab)
+							->AddTab(FDatabaseEditorTabs::PreviewSettingsID, ETabState::OpenedTab)
+							->AddTab(FDatabaseEditorTabs::DataDetailsID, ETabState::OpenedTab)
+							->AddTab(FDatabaseEditorTabs::SelectionDetailsID, ETabState::OpenedTab)
+						)
 					)
 				)
 			);
@@ -359,6 +381,13 @@ namespace UE::PoseSearch
 			.SetDisplayName(LOCTEXT("DataDetailsTab", "Data Details"))
 			.SetGroup(WorkspaceMenuCategoryRef)
 			.SetIcon(FSlateIcon(FAppStyle::GetAppStyleSetName(), "GraphEditor.EventGraph_16x"));
+		
+		InTabManager->RegisterTabSpawner(
+			FDatabaseEditorTabs::AssetBrowserID,
+			FOnSpawnTab::CreateSP(this, &FDatabaseEditor::SpawnTab_AssetBrowser))
+			.SetDisplayName(LOCTEXT("AssetBrowserTab", "Asset Browser"))
+			.SetGroup(WorkspaceMenuCategoryRef);
+			// .SetIcon(FSlateIcon(FAppStyle::GetAppStyleSetName(), "GraphEditor.EventGraph_16x"));
 	}
 
 	void FDatabaseEditor::UnregisterTabSpawners(const TSharedRef<FTabManager>& InTabManager)
@@ -372,6 +401,7 @@ namespace UE::PoseSearch
 		InTabManager->UnregisterTabSpawner(FDatabaseEditorTabs::SelectionDetailsID);
 		InTabManager->UnregisterTabSpawner(FDatabaseEditorTabs::StatisticsOverview);
 		InTabManager->UnregisterTabSpawner(FDatabaseEditorTabs::DataDetailsID);
+		InTabManager->UnregisterTabSpawner(FDatabaseEditorTabs::AssetBrowserID);
 	}
 
 	FName FDatabaseEditor::GetToolkitFName() const
@@ -497,7 +527,18 @@ namespace UE::PoseSearch
 				DataDetails.ToSharedRef()
 			];
 	}
-
+	
+	TSharedRef<SDockTab> FDatabaseEditor::SpawnTab_AssetBrowser(const FSpawnTabArgs& Args) const
+	{
+		check(Args.GetTabId() == FDatabaseEditorTabs::AssetBrowserID);
+		
+		return SNew(SDockTab)
+			.Label(LOCTEXT("AssetBrowser_Title", "Asset Browser"))
+			[
+				AssetBrowserWidget.ToSharedRef()
+			];
+	}
+	
 	void FDatabaseEditor::OnFinishedChangingSelectionProperties(const FPropertyChangedEvent& PropertyChangedEvent)
 	{
 	}
diff --git a/Engine/Plugins/Animation/PoseSearch/Source/Editor/Public/PoseSearchDatabaseEditor.h b/Engine/Plugins/Animation/PoseSearch/Source/Editor/Public/PoseSearchDatabaseEditor.h
--- a/Engine/Plugins/Animation/PoseSearch/Source/Editor/Public/PoseSearchDatabaseEditor.h
+++ b/Engine/Plugins/Animation/PoseSearch/Source/Editor/Public/PoseSearchDatabaseEditor.h
@@ -18,6 +18,7 @@ namespace UE::PoseSearch
 	class SDatabaseAssetTree;
 	class SDatabaseDataDetails;
 	class SDatabasePreview;
+	class SPoseSearchDatabaseAssetBrowser;
 
 	class FDatabaseEditor : public FAssetEditorToolkit, public FNotifyHook
 	{
@@ -62,6 +63,7 @@ namespace UE::PoseSearch
 		TSharedRef<SDockTab> SpawnTab_SelectionDetails(const FSpawnTabArgs& Args);
 		TSharedRef<SDockTab> SpawnTab_StatisticsOverview(const FSpawnTabArgs& Args) const;
 		TSharedRef<SDockTab> SpawnTab_DataDetails(const FSpawnTabArgs& Args) const;
+		TSharedRef<SDockTab> SpawnTab_AssetBrowser(const FSpawnTabArgs& Args) const;
 		
 		void OnFinishedChangingSelectionProperties(const FPropertyChangedEvent& PropertyChangedEvent);
 		void OnAssetTreeSelectionChanged(
@@ -83,6 +85,8 @@ namespace UE::PoseSearch
 		TSharedPtr<SDatabaseAssetTree> AssetTreeWidget;
 
 		TSharedPtr<SDatabaseDataDetails> DataDetails;
+
+		TSharedPtr<SPoseSearchDatabaseAssetBrowser> AssetBrowserWidget;
 		
 		TSharedPtr<IDetailsView> EditingAssetWidget;
 
