motion matching - FPoseHistory performance improvements: moving processing from FAnimNode_PoseSearchHistoryCollector_Base::PreUpdate to FAnimNode_PoseSearchHistoryCollector_Base::PreUpdate to Evaluate_AnyThread. minimized per frame allocations

diff --git a/Engine/Plugins/Animation/PoseSearch/Source/Runtime/Private/AnimNode_PoseSearchHistoryCollector.cpp b/Engine/Plugins/Animation/PoseSearch/Source/Runtime/Private/AnimNode_PoseSearchHistoryCollector.cpp
--- a/Engine/Plugins/Animation/PoseSearch/Source/Runtime/Private/AnimNode_PoseSearchHistoryCollector.cpp
+++ b/Engine/Plugins/Animation/PoseSearch/Source/Runtime/Private/AnimNode_PoseSearchHistoryCollector.cpp
@@ -42,30 +42,37 @@ void FAnimNode_PoseSearchHistoryCollector_Base::Initialize_AnyThread(const FAnim
 	Super::Initialize_AnyThread(Context);
 
 	PoseHistory.Initialize_AnyThread(PoseCount, SamplingInterval);
+
+	if (bInitializeWithRefPose)
+	{
+		// initializing PoseHistory with a ref pose at FAnimInstanceProxy location/facing
+		FMemMark Mark(FMemStack::Get());
+		FCSPose<FCompactPose> ComponentSpacePose;
+		ComponentSpacePose.InitPose(&Context.AnimInstanceProxy->GetRequiredBones());
+		PoseHistory.EvaluateComponentSpace_AnyThread(0.f, ComponentSpacePose, bStoreScales, RootBoneRecoveryTime, true, true, GetRequiredBones(Context.AnimInstanceProxy));
+	}
 }
 
-void FAnimNode_PoseSearchHistoryCollector_Base::CacheBones_AnyThread(const FAnimationCacheBonesContext& Context)
+TArray<FBoneIndexType> FAnimNode_PoseSearchHistoryCollector_Base::GetRequiredBones(const FAnimInstanceProxy* AnimInstanceProxy) const
 {
-	DECLARE_SCOPE_HIERARCHICAL_COUNTER_ANIMNODE(Update_AnyThread);
-	check(Context.AnimInstanceProxy);
-
-	Super::CacheBones_AnyThread(Context);
+	check(AnimInstanceProxy);
 
 	TArray<FBoneIndexType> RequiredBones;
 	if (!CollectedBones.IsEmpty())
 	{
-		if (const USkeletalMeshComponent* SkeletalMeshComponent = Context.AnimInstanceProxy->GetSkelMeshComponent())
+		if (const USkeletalMeshComponent* SkeletalMeshComponent = AnimInstanceProxy->GetSkelMeshComponent())
 		{
 			if (const USkinnedAsset* SkinnedAsset = SkeletalMeshComponent->GetSkinnedAsset())
 			{
 				if (const USkeleton* Skeleton = SkinnedAsset->GetSkeleton())
 				{
 					RequiredBones.Reserve(CollectedBones.Num());
-					for (FBoneReference& BoneReference : CollectedBones)
+					for (const FBoneReference& BoneReference : CollectedBones)
 					{
-						if (BoneReference.Initialize(Skeleton))
+						FBoneReference BoneReferenceCopy = BoneReference;
+						if (BoneReferenceCopy.Initialize(Skeleton))
 						{
-							RequiredBones.AddUnique(BoneReference.BoneIndex);
+							RequiredBones.AddUnique(BoneReferenceCopy.BoneIndex);
 						}
 					}
 				}
@@ -73,18 +80,17 @@ void FAnimNode_PoseSearchHistoryCollector_Base::CacheBones_AnyThread(const FAnim
 		}
 	}
 
-	PoseHistory.CacheBones_AnyThread(RequiredBones);
+	return RequiredBones;
+}
 
-	if (bInitializeWithRefPose)
-	{
-		// initializing PoseHistory with a ref pose at FAnimInstanceProxy location/facing
-		FMemMark Mark(FMemStack::Get());
-		FCompactPose Pose;
-		Pose.SetBoneContainer(&Context.AnimInstanceProxy->GetRequiredBones());
-		FCSPose<FCompactPose> ComponentSpacePose;
-		ComponentSpacePose.InitPose(Pose);
-		PoseHistory.EvaluateComponentSpace_AnyThread(0.f, ComponentSpacePose, bStoreScales, RootBoneRecoveryTime);
-	}
+void FAnimNode_PoseSearchHistoryCollector_Base::CacheBones_AnyThread(const FAnimationCacheBonesContext& Context)
+{
+	DECLARE_SCOPE_HIERARCHICAL_COUNTER_ANIMNODE(Update_AnyThread);
+	check(Context.AnimInstanceProxy);
+
+	Super::CacheBones_AnyThread(Context);
+
+	bCacheBones = true;
 }
 
 void FAnimNode_PoseSearchHistoryCollector_Base::Update_AnyThread(const FAnimationUpdateContext& Context)
@@ -100,16 +106,18 @@ void FAnimNode_PoseSearchHistoryCollector_Base::PreUpdate(const UAnimInstance* I
 {
 	Super::PreUpdate(InAnimInstance);
 	
-	const bool bNeedsReset = bResetOnBecomingRelevant && UpdateCounter.HasEverBeenUpdated() && !UpdateCounter.WasSynchronizedCounter(InAnimInstance->GetUpdateCounter());
-	const float DeltaTime = InAnimInstance->GetDeltaSeconds();
+	if (bGenerateTrajectory)
+	{
+		FPoseSearchTrajectoryData::FSampling TrajectoryDataSampling;
+		TrajectoryDataSampling.NumHistorySamples = FMath::Max(PoseCount, TrajectoryHistoryCount);
+		TrajectoryDataSampling.SecondsPerHistorySample = SamplingInterval;
+		TrajectoryDataSampling.NumPredictionSamples = TrajectoryPredictionCount;
+		TrajectoryDataSampling.SecondsPerPredictionSample = PredictionSamplingInterval;
 
-	FPoseSearchTrajectoryData::FSampling TrajectoryDataSampling;
-	TrajectoryDataSampling.NumHistorySamples = FMath::Max(PoseCount, TrajectoryHistoryCount);
-	TrajectoryDataSampling.SecondsPerHistorySample = SamplingInterval;
-	TrajectoryDataSampling.NumPredictionSamples = TrajectoryPredictionCount;
-	TrajectoryDataSampling.SecondsPerPredictionSample = PredictionSamplingInterval;
+		PoseHistory.GenerateTrajectory(InAnimInstance, InAnimInstance->GetDeltaSeconds(), TrajectoryData, TrajectoryDataSampling);
+	}
 
-	PoseHistory.PreUpdate(InAnimInstance, DeltaTime, bGenerateTrajectory, TrajectoryData, TrajectoryDataSampling, bNeedsReset);
+	PoseHistory.PreUpdate();
 }
 
 /////////////////////////////////////////////////////
@@ -139,9 +147,21 @@ void FAnimNode_PoseSearchHistoryCollector::Evaluate_AnyThread(FPoseContext& Outp
 	Super::Evaluate_AnyThread(Output);
 	Source.Evaluate(Output);
 
+	const bool bNeedsReset = bResetOnBecomingRelevant && UpdateCounter.HasEverBeenUpdated() && !UpdateCounter.WasSynchronizedCounter(Output.AnimInstanceProxy->GetUpdateCounter());
+
 	FCSPose<FCompactPose> ComponentSpacePose;
 	ComponentSpacePose.InitPose(Output.Pose);
-	PoseHistory.EvaluateComponentSpace_AnyThread(Output.AnimInstanceProxy->GetDeltaSeconds(), ComponentSpacePose, bStoreScales, RootBoneRecoveryTime);
+
+	TArray<FBoneIndexType> RequiredBones;
+	if (bCacheBones)
+	{
+		RequiredBones = GetRequiredBones(Output.AnimInstanceProxy);
+	}
+
+	PoseHistory.EvaluateComponentSpace_AnyThread(Output.AnimInstanceProxy->GetDeltaSeconds(), ComponentSpacePose,
+		bStoreScales, RootBoneRecoveryTime, bNeedsReset, bCacheBones, RequiredBones);
+
+	bCacheBones = false;
 
 #if ENABLE_DRAW_DEBUG && ENABLE_ANIM_DEBUG
 	FColor Color;
@@ -195,7 +215,18 @@ void FAnimNode_PoseSearchComponentSpaceHistoryCollector::EvaluateComponentSpace_
 	Super::EvaluateComponentSpace_AnyThread(Output);
 	Source.EvaluateComponentSpace(Output);
 
-	PoseHistory.EvaluateComponentSpace_AnyThread(Output.AnimInstanceProxy->GetDeltaSeconds(), Output.Pose, bStoreScales, RootBoneRecoveryTime);
+	const bool bNeedsReset = bResetOnBecomingRelevant && UpdateCounter.HasEverBeenUpdated() && !UpdateCounter.WasSynchronizedCounter(Output.AnimInstanceProxy->GetUpdateCounter());
+
+	TArray<FBoneIndexType> RequiredBones;
+	if (bCacheBones)
+	{
+		RequiredBones = GetRequiredBones(Output.AnimInstanceProxy);
+	}
+
+	PoseHistory.EvaluateComponentSpace_AnyThread(Output.AnimInstanceProxy->GetDeltaSeconds(), Output.Pose,
+		bStoreScales, RootBoneRecoveryTime, bNeedsReset, bCacheBones, RequiredBones);
+	
+	bCacheBones = false;
 
 #if ENABLE_DRAW_DEBUG && ENABLE_ANIM_DEBUG
 	FColor Color;
diff --git a/Engine/Plugins/Animation/PoseSearch/Source/Runtime/Private/PoseSearchHistory.cpp b/Engine/Plugins/Animation/PoseSearch/Source/Runtime/Private/PoseSearchHistory.cpp
--- a/Engine/Plugins/Animation/PoseSearch/Source/Runtime/Private/PoseSearchHistory.cpp
+++ b/Engine/Plugins/Animation/PoseSearch/Source/Runtime/Private/PoseSearchHistory.cpp
@@ -177,6 +177,67 @@ static bool LerpEntries(float Time, bool bExtrapolate, const FPoseHistoryEntry&
 	return bSuccess;
 }
 
+static uint32 GetTypeHash(const FBoneToTransformMap& BoneToTransformMap)
+{
+	const int32 Num = BoneToTransformMap.Num();
+
+	if (Num == 0)
+	{
+		return 0;
+	}
+
+	TArrayView<FBoneToTransformPair> Pairs((FBoneToTransformPair*)FMemory_Alloca(Num * sizeof(FBoneToTransformPair)), Num);
+	
+	int32 Index = 0;
+	for (const FBoneToTransformPair& BoneToTransformPair : BoneToTransformMap)
+	{
+		Pairs[Index] = BoneToTransformPair;
+		++Index;
+	}
+
+	Pairs.StableSort();
+
+	uint32 TypeHash = ::GetTypeHash(Pairs[0]);
+	for (Index = 1; Index < Num; ++Index)
+	{
+		TypeHash = HashCombineFast(TypeHash, ::GetTypeHash(Pairs[Index]));
+	}
+
+	return TypeHash;
+}
+
+// optimized copy for TRingBuffer<FPoseHistoryEntry> implementing "To = From" to avoid allocations as much as possible
+static void CopyEntries(const TRingBuffer<FPoseHistoryEntry>& From, TRingBuffer<FPoseHistoryEntry>& To)
+{
+	const int32 FromNum = From.Num();
+	const int32 PopCount = To.Num() - FromNum;
+	if (PopCount != 0)
+	{
+		if (PopCount > 0)
+		{
+			To.Pop(PopCount);
+		}
+		else
+		{
+			To.Reserve(FromNum);
+			const int32 AddCount = -PopCount;
+			for (int32 Index = 0; Index < AddCount; ++Index)
+			{
+				To.Add(FPoseHistoryEntry());
+			}
+		}
+	}
+
+	check(To.Num() == FromNum);
+
+	TRingBuffer<FPoseHistoryEntry>::TIterator ToIt = To.begin();
+	const TRingBuffer<FPoseHistoryEntry>::TConstIterator FromEnd = From.end();
+	for (TRingBuffer<FPoseHistoryEntry>::TConstIterator FromIt = From.begin(); FromIt != FromEnd; ++FromIt, ++ToIt)
+	{
+		*ToIt = *FromIt;
+	}
+}
+
 //////////////////////////////////////////////////////////////////////////
 // FPoseHistoryEntry
 void FPoseHistoryEntry::Update(float Time, FCSPose<FCompactPose>& ComponentSpacePose, const FBoneToTransformMap& BoneToTransformMap, bool bStoreScales)
@@ -441,13 +502,13 @@ FPoseHistory::FPoseHistory(FPoseHistory&& Other)
 FPoseHistory& FPoseHistory::operator=(const FPoseHistory& Other)
 {
 #if ENABLE_ANIM_DEBUG
-	CheckThreadSafetyWrite(ReadDataThreadSafeCounter);
-	CheckThreadSafetyWrite(WriteDataThreadSafeCounter);
+	CheckThreadSafetyWrite(ReadPoseDataThreadSafeCounter);
+	CheckThreadSafetyWrite(WritePoseDataThreadSafeCounter);
 
-	FThreadSafeCounter& OtherReadDataThreadSafeCounter = Other.ReadDataThreadSafeCounter;
-	FThreadSafeCounter& OtherWriteDataThreadSafeCounter = Other.WriteDataThreadSafeCounter;
-	CheckThreadSafetyWrite(OtherReadDataThreadSafeCounter);
-	CheckThreadSafetyWrite(OtherWriteDataThreadSafeCounter);
+	FThreadSafeCounter& OtherReadPoseDataThreadSafeCounter = Other.ReadPoseDataThreadSafeCounter;
+	FThreadSafeCounter& OtherWritePoseDataThreadSafeCounter = Other.WritePoseDataThreadSafeCounter;
+	CheckThreadSafetyWrite(OtherReadPoseDataThreadSafeCounter);
+	CheckThreadSafetyWrite(OtherWritePoseDataThreadSafeCounter);
 #endif // ENABLE_ANIM_DEBUG
 
 	MaxNumPoses = Other.MaxNumPoses;
@@ -457,21 +518,21 @@ FPoseHistory& FPoseHistory::operator=(const FPoseHistory& Other)
 	TrajectoryDataState = Other.TrajectoryDataState;
 	TrajectorySpeedMultiplier = Other.TrajectorySpeedMultiplier;
 
-	ReadData = Other.ReadData;
-	WriteData = Other.WriteData;
+	ReadPoseDataIndex = Other.ReadPoseDataIndex;
+	DoubleBufferedPoseData = Other.DoubleBufferedPoseData;
 	return *this;
 }
 
 FPoseHistory& FPoseHistory::operator=(FPoseHistory&& Other)
 {
 #if ENABLE_ANIM_DEBUG
-	CheckThreadSafetyWrite(ReadDataThreadSafeCounter);
-	CheckThreadSafetyWrite(WriteDataThreadSafeCounter);
+	CheckThreadSafetyWrite(ReadPoseDataThreadSafeCounter);
+	CheckThreadSafetyWrite(WritePoseDataThreadSafeCounter);
 
-	FThreadSafeCounter& OtherReadDataThreadSafeCounter = Other.ReadDataThreadSafeCounter;
-	FThreadSafeCounter& OtherWriteDataThreadSafeCounter = Other.WriteDataThreadSafeCounter;
-	CheckThreadSafetyWrite(OtherReadDataThreadSafeCounter);
-	CheckThreadSafetyWrite(OtherWriteDataThreadSafeCounter);
+	FThreadSafeCounter& OtherReadPoseDataThreadSafeCounter = Other.ReadPoseDataThreadSafeCounter;
+	FThreadSafeCounter& OtherWritePoseDataThreadSafeCounter = Other.WritePoseDataThreadSafeCounter;
+	CheckThreadSafetyWrite(OtherReadPoseDataThreadSafeCounter);
+	CheckThreadSafetyWrite(OtherWritePoseDataThreadSafeCounter);
 #endif // ENABLE_ANIM_DEBUG
 
 	MaxNumPoses = MoveTemp(Other.MaxNumPoses);
@@ -481,14 +542,14 @@ FPoseHistory& FPoseHistory::operator=(FPoseHistory&& Other)
 	TrajectoryDataState = MoveTemp(Other.TrajectoryDataState);
 	TrajectorySpeedMultiplier = MoveTemp(Other.TrajectorySpeedMultiplier);
 
-	ReadData = MoveTemp(Other.ReadData);
-	WriteData = MoveTemp(Other.WriteData);
+	ReadPoseDataIndex = MoveTemp(Other.ReadPoseDataIndex);
+	DoubleBufferedPoseData = MoveTemp(Other.DoubleBufferedPoseData);
 	return *this;
 }
 
 void FPoseHistory::Initialize_AnyThread(int32 InNumPoses, float InSamplingInterval)
 {
-	CheckThreadSafetyWrite(WriteDataThreadSafeCounter);
+	CheckThreadSafetyWrite(WritePoseDataThreadSafeCounter);
 	check(InNumPoses >= 2);
 
 	MaxNumPoses = InNumPoses;
@@ -498,41 +559,13 @@ void FPoseHistory::Initialize_AnyThread(int32 InNumPoses, float InSamplingInterv
 	TrajectoryDataState = FPoseSearchTrajectoryData::FState();
 	TrajectorySpeedMultiplier = 1.f;
 
-	ReadData = FData();
-	WriteData = FData();
-}
-
-void FPoseHistory::CacheBones_AnyThread(const TArray<FBoneIndexType>& RequiredBones)
-{
-	CheckThreadSafetyWrite(WriteDataThreadSafeCounter);
-
-	check(MaxNumPoses >= 2);
-
-	WriteData.BoneToTransformMap.Reset();
-	if (!RequiredBones.IsEmpty())
-	{
-		// making sure we always collect the root bone transform (by construction BoneToTransformMap[0] = 0)
-		const FComponentSpaceTransformIndex ComponentSpaceTransformRootBoneIndex = 0;
-		WriteData.BoneToTransformMap.Add(RootBoneIndexType) = ComponentSpaceTransformRootBoneIndex;
-
-		for (int32 i = 0; i < RequiredBones.Num(); ++i)
-		{
-			// adding only unique RequiredBones to avoid oversizing Entries::ComponentSpaceTransforms
-			if (!WriteData.BoneToTransformMap.Find(RequiredBones[i]))
-			{
-				const FComponentSpaceTransformIndex ComponentSpaceTransformIndex = WriteData.BoneToTransformMap.Num();
-				WriteData.BoneToTransformMap.Add(RequiredBones[i]) = ComponentSpaceTransformIndex;
-			}
-		}
-	}
-
-	WriteData.Entries.Reset();
-	WriteData.Entries.Reserve(MaxNumPoses);
+	ReadPoseDataIndex = 0;
+	DoubleBufferedPoseData = FDoubleBufferedPoseData();
 }
 
 bool FPoseHistory::GetTransformAtTime(float Time, FTransform& OutBoneTransform, const USkeleton* BoneIndexSkeleton, FBoneIndexType BoneIndexType, FBoneIndexType ReferenceBoneIndexType, bool bExtrapolate) const
 {
-	CheckThreadSafetyRead(ReadDataThreadSafeCounter);
+	CheckThreadSafetyRead(ReadPoseDataThreadSafeCounter);
 
 	static_assert(RootBoneIndexType == 0 && ComponentSpaceIndexType == FBoneIndexType(-1) && WorldSpaceIndexType == FBoneIndexType(-2)); // some assumptions
 	check(BoneIndexType != ComponentSpaceIndexType && BoneIndexType != WorldSpaceIndexType);
@@ -547,7 +580,8 @@ bool FPoseHistory::GetTransformAtTime(float Time, FTransform& OutBoneTransform,
 		ReferenceBoneIndexType = ComponentSpaceIndexType;
 	}
 
-	const int32 NumEntries = ReadData.Entries.Num();
+	const FPoseData& ReadPoseData = GetReadPoseData();
+	const int32 NumEntries = ReadPoseData.Entries.Num();
 	if (NumEntries > 0)
 	{
 		int32 NextIdx = 0;
@@ -555,15 +589,15 @@ bool FPoseHistory::GetTransformAtTime(float Time, FTransform& OutBoneTransform,
 
 		if (NumEntries > 1)
 		{
-			const int32 LowerBoundIdx = LowerBound(ReadData.Entries.begin(), ReadData.Entries.end(), Time, [](const FPoseHistoryEntry& Entry, float Value) { return Value > Entry.AccumulatedSeconds; });
+			const int32 LowerBoundIdx = LowerBound(ReadPoseData.Entries.begin(), ReadPoseData.Entries.end(), Time, [](const FPoseHistoryEntry& Entry, float Value) { return Value > Entry.AccumulatedSeconds; });
 			NextIdx = FMath::Clamp(LowerBoundIdx, 1, NumEntries - 1);
 			PrevIdx = NextIdx - 1;
 		}
 	
-		const FPoseHistoryEntry& PrevEntry = ReadData.Entries[PrevIdx];
-		const FPoseHistoryEntry& NextEntry = ReadData.Entries[NextIdx];
+		const FPoseHistoryEntry& PrevEntry = ReadPoseData.Entries[PrevIdx];
+		const FPoseHistoryEntry& NextEntry = ReadPoseData.Entries[NextIdx];
 
-		bSuccess = LerpEntries(Time, bExtrapolate, PrevEntry, NextEntry, BoneIndexSkeleton, ReadData.LastUpdateSkeleton.Get(), ReadData.BoneToTransformMap, BoneIndexType, ReferenceBoneIndexType, OutBoneTransform);
+		bSuccess = LerpEntries(Time, bExtrapolate, PrevEntry, NextEntry, BoneIndexSkeleton, ReadPoseData.LastUpdateSkeleton.Get(), ReadPoseData.BoneToTransformMap, BoneIndexType, ReferenceBoneIndexType, OutBoneTransform);
 		if (bApplyComponentToWorld)
 		{
 			OutBoneTransform *= ComponentToWorld;
@@ -579,74 +613,69 @@ bool FPoseHistory::GetTransformAtTime(float Time, FTransform& OutBoneTransform,
 
 const FPoseSearchQueryTrajectory& FPoseHistory::GetTrajectory() const
 {
-	CheckThreadSafetyRead(ReadDataThreadSafeCounter);
+	CheckThreadSafetyRead(ReadPoseDataThreadSafeCounter);
 	return Trajectory;
 }
 
 float FPoseHistory::GetTrajectorySpeedMultiplier() const
 {
-	CheckThreadSafetyRead(ReadDataThreadSafeCounter);
+	CheckThreadSafetyRead(ReadPoseDataThreadSafeCounter);
 	return TrajectorySpeedMultiplier;
 }
 
 bool FPoseHistory::IsEmpty() const
 {
-	CheckThreadSafetyRead(ReadDataThreadSafeCounter);
-	return ReadData.Entries.IsEmpty();
+	CheckThreadSafetyRead(ReadPoseDataThreadSafeCounter);
+	return GetReadPoseData().Entries.IsEmpty();
 }
 
 const FBoneToTransformMap& FPoseHistory::GetBoneToTransformMap() const
 {
-	CheckThreadSafetyRead(ReadDataThreadSafeCounter);
-	return ReadData.BoneToTransformMap;
+	CheckThreadSafetyRead(ReadPoseDataThreadSafeCounter);
+	return GetReadPoseData().BoneToTransformMap;
 }
 
 int32 FPoseHistory::GetNumEntries() const
 {
-	CheckThreadSafetyRead(ReadDataThreadSafeCounter);
-	return ReadData.Entries.Num();
+	CheckThreadSafetyRead(ReadPoseDataThreadSafeCounter);
+	return GetReadPoseData().Entries.Num();
 }
 
 const FPoseHistoryEntry& FPoseHistory::GetEntry(int32 EntryIndex) const
 {
-	CheckThreadSafetyRead(ReadDataThreadSafeCounter);
-	return ReadData.Entries[EntryIndex];
+	CheckThreadSafetyRead(ReadPoseDataThreadSafeCounter);
+	return GetReadPoseData().Entries[EntryIndex];
 }
 
-void FPoseHistory::PreUpdate(const UAnimInstance* AnimInstance, float DeltaTime, bool bGenerateTrajectory, const FPoseSearchTrajectoryData& TrajectoryData, const FPoseSearchTrajectoryData::FSampling& TrajectoryDataSampling, bool bNeedsReset)
+void FPoseHistory::GenerateTrajectory(const UAnimInstance* AnimInstance, float DeltaTime, const FPoseSearchTrajectoryData& TrajectoryData, const FPoseSearchTrajectoryData::FSampling& TrajectoryDataSampling)
 {
-	// we're writing both ReadData, WriteData. checking for thread safety
-	CheckThreadSafetyWrite(ReadDataThreadSafeCounter);
-	CheckThreadSafetyWrite(WriteDataThreadSafeCounter);
+	// @todo: Synchronize the FPoseSearchQueryTrajectorySample::AccumulatedSeconds of the generated trajectory with the FPoseHistoryEntry::AccumulatedSeconds of the captured poses
+	FPoseSearchTrajectoryData::FDerived TrajectoryDataDerived;
+	TrajectoryData.UpdateData(DeltaTime, AnimInstance, TrajectoryDataDerived, TrajectoryDataState);
+	UPoseSearchTrajectoryLibrary::InitTrajectorySamples(Trajectory, TrajectoryData, TrajectoryDataDerived, TrajectoryDataSampling, DeltaTime);
+	UPoseSearchTrajectoryLibrary::UpdateHistory_TransformHistory(Trajectory, TrajectoryData, TrajectoryDataDerived, TrajectoryDataSampling, DeltaTime);
+	UPoseSearchTrajectoryLibrary::UpdatePrediction_SimulateCharacterMovement(Trajectory, TrajectoryData, TrajectoryDataDerived, TrajectoryDataSampling, DeltaTime);
 
-	check(IsInGameThread());
+	// @todo: support TrajectorySpeedMultiplier
+	//TrajectorySpeedMultiplier = 1.f;
+}
 
-	if (bNeedsReset)
-	{
-		WriteData.Entries.Reset();
-	}
+void FPoseHistory::PreUpdate()
+{
+	// checking for thread safety
+	CheckThreadSafetyWrite(ReadPoseDataThreadSafeCounter);
+	CheckThreadSafetyWrite(WritePoseDataThreadSafeCounter);
 
-	if (bGenerateTrajectory)
-	{
-		// @todo: Synchronize the FPoseSearchQueryTrajectorySample::AccumulatedSeconds of the generated trajectory with the FPoseHistoryEntry::AccumulatedSeconds of the captured poses
-		FPoseSearchTrajectoryData::FDerived TrajectoryDataDerived;
-		TrajectoryData.UpdateData(DeltaTime, AnimInstance, TrajectoryDataDerived, TrajectoryDataState);
-		UPoseSearchTrajectoryLibrary::InitTrajectorySamples(Trajectory, TrajectoryData, TrajectoryDataDerived, TrajectoryDataSampling, DeltaTime);
-		UPoseSearchTrajectoryLibrary::UpdateHistory_TransformHistory(Trajectory, TrajectoryData, TrajectoryDataDerived, TrajectoryDataSampling, DeltaTime);
-		UPoseSearchTrajectoryLibrary::UpdatePrediction_SimulateCharacterMovement(Trajectory, TrajectoryData, TrajectoryDataDerived, TrajectoryDataSampling, DeltaTime);
+	check(IsInGameThread());
 
-		// @todo: support TrajectorySpeedMultiplier
-		//TrajectorySpeedMultiplier = 1.f;
-	}
-	
-	ReadData = WriteData;
+	ReadPoseDataIndex = GetWritePoseDataIndex();
 }
 
 void FPoseHistory::SetTrajectory(const FPoseSearchQueryTrajectory& InTrajectory, float InTrajectorySpeedMultiplier)
 {
 	if (!InTrajectory.Samples.IsEmpty())
 	{
-		CheckThreadSafetyWrite(ReadDataThreadSafeCounter);
+		CheckThreadSafetyWrite(ReadPoseDataThreadSafeCounter);
 
 		// @todo: THIS IS NOT THREAD SAFE! in the contex of multi character motion matching (CheckThreadSafetyWrite will assert in case of improper usage)
 		Trajectory = InTrajectory;
@@ -664,50 +693,96 @@ void FPoseHistory::SetTrajectory(const FPoseSearchQueryTrajectory& InTrajectory,
 	}
 }
 
-void FPoseHistory::EvaluateComponentSpace_AnyThread(float DeltaTime, FCSPose<FCompactPose>& ComponentSpacePose, bool bStoreScales, float RootBoneRecoveryTime)
+void FPoseHistory::EvaluateComponentSpace_AnyThread(float DeltaTime, FCSPose<FCompactPose>& ComponentSpacePose, bool bStoreScales,
+	float RootBoneRecoveryTime, bool bNeedsReset, bool bCacheBones, const TArray<FBoneIndexType>& RequiredBones)
 {
-	CheckThreadSafetyWrite(WriteDataThreadSafeCounter);
+	CheckThreadSafetyWrite(WritePoseDataThreadSafeCounter);
+	CheckThreadSafetyRead(ReadPoseDataThreadSafeCounter);
 
 	check(MaxNumPoses >= 2);
 
 	const USkeleton* Skeleton = ComponentSpacePose.GetPose().GetBoneContainer().GetSkeletonAsset();
 	check(Skeleton);
 
-	if (WriteData.LastUpdateSkeleton != Skeleton)
+	const FPoseData& ReadPoseData = GetReadPoseData();
+	FPoseData& WritePoseData = EditWritePoseData();
+
+	WritePoseData.LastUpdateSkeleton = ReadPoseData.LastUpdateSkeleton;
+	
+	if (bCacheBones)
+	{
+		WritePoseData.BoneToTransformMap.Reset();
+		if (!RequiredBones.IsEmpty())
+		{
+			// making sure we always collect the root bone transform (by construction BoneToTransformMap[0] = 0)
+			const FComponentSpaceTransformIndex ComponentSpaceTransformRootBoneIndex = 0;
+			WritePoseData.BoneToTransformMap.Add(RootBoneIndexType) = ComponentSpaceTransformRootBoneIndex;
+
+			for (int32 i = 0; i < RequiredBones.Num(); ++i)
+			{
+				// adding only unique RequiredBones to avoid oversizing Entries::ComponentSpaceTransforms
+				if (!WritePoseData.BoneToTransformMap.Find(RequiredBones[i]))
+				{
+					const FComponentSpaceTransformIndex ComponentSpaceTransformIndex = WritePoseData.BoneToTransformMap.Num();
+					WritePoseData.BoneToTransformMap.Add(RequiredBones[i]) = ComponentSpaceTransformIndex;
+				}
+			}
+		}
+
+		WritePoseData.BoneToTransformMapTypeHash = GetTypeHash(WritePoseData.BoneToTransformMap);
+		bNeedsReset |= (WritePoseData.BoneToTransformMapTypeHash != ReadPoseData.BoneToTransformMapTypeHash);
+	}
+	else if (WritePoseData.BoneToTransformMapTypeHash != ReadPoseData.BoneToTransformMapTypeHash)
+	{
+		WritePoseData.BoneToTransformMap = ReadPoseData.BoneToTransformMap;
+		WritePoseData.BoneToTransformMapTypeHash = ReadPoseData.BoneToTransformMapTypeHash;
+		bNeedsReset = true;
+	}
+
+	if (WritePoseData.LastUpdateSkeleton != Skeleton)
+	{
+		bNeedsReset = true;
+		WritePoseData.LastUpdateSkeleton = Skeleton;
+	}
+
+	if (bNeedsReset)
+	{
+		WritePoseData.Entries.Reset();
+		WritePoseData.Entries.Reserve(MaxNumPoses);
+	}
+	else
 	{
-		// @todo: support a different USkeleton per FPoseHistoryEntry if required
-		WriteData.Entries.Reset();
-		WriteData.LastUpdateSkeleton = Skeleton;
+		CopyEntries(ReadPoseData.Entries, WritePoseData.Entries);
 	}
 
 	FPoseHistoryEntry FutureEntryTemp;
-	if (!WriteData.Entries.IsEmpty() && WriteData.Entries.Last().AccumulatedSeconds > 0.f)
+	if (!WritePoseData.Entries.IsEmpty() && WritePoseData.Entries.Last().AccumulatedSeconds > 0.f)
 	{
 		// removing the "future" root bone Entry
-		FutureEntryTemp = MoveTemp(WriteData.Entries.Last());
-		WriteData.Entries.Pop();
+		FutureEntryTemp = MoveTemp(WritePoseData.Entries.Last());
+		WritePoseData.Entries.Pop();
 	}
 
 	// Age our elapsed times
-	for (FPoseHistoryEntry& Entry : WriteData.Entries)
+	for (FPoseHistoryEntry& Entry : WritePoseData.Entries)
 	{
 		Entry.AccumulatedSeconds -= DeltaTime;
 	}
 
-	if (WriteData.Entries.Num() != MaxNumPoses)
+	if (WritePoseData.Entries.Num() != MaxNumPoses)
 	{
 		// Consume every pose until the queue is full
-		WriteData.Entries.Emplace();
+		WritePoseData.Entries.Emplace();
 	}
-	else if (SamplingInterval <= 0.f || WriteData.Entries[WriteData.Entries.Num() - 2].AccumulatedSeconds <= -SamplingInterval)
+	else if (SamplingInterval <= 0.f || WritePoseData.Entries[WritePoseData.Entries.Num() - 2].AccumulatedSeconds <= -SamplingInterval)
 	{
-		FPoseHistoryEntry EntryTemp = MoveTemp(WriteData.Entries.First());
-		WriteData.Entries.PopFront();
-		WriteData.Entries.Emplace(MoveTemp(EntryTemp));
+		FPoseHistoryEntry EntryTemp = MoveTemp(WritePoseData.Entries.First());
+		WritePoseData.Entries.PopFront();
+		WritePoseData.Entries.Emplace(MoveTemp(EntryTemp));
 	}
 
 	// Regardless of the retention policy, we always update the most recent Entry
-	WriteData.Entries.Last().Update(0.f, ComponentSpacePose, WriteData.BoneToTransformMap, bStoreScales);
+	WritePoseData.Entries.Last().Update(0.f, ComponentSpacePose, WritePoseData.BoneToTransformMap, bStoreScales);
 
 	if (RootBoneRecoveryTime > 0.f)
 	{
@@ -719,14 +794,14 @@ void FPoseHistory::EvaluateComponentSpace_AnyThread(float DeltaTime, FCSPose<FCo
 		FutureEntryTemp.SetNum(1, bStoreScales);
 		FutureEntryTemp.SetComponentSpaceTransform(RootBoneIndexType, RefRootBone);
 		FutureEntryTemp.AccumulatedSeconds = RootBoneRecoveryTime;
-		WriteData.Entries.Emplace(MoveTemp(FutureEntryTemp));
+		WritePoseData.Entries.Emplace(MoveTemp(FutureEntryTemp));
 	}
 }
 
 #if ENABLE_DRAW_DEBUG && ENABLE_ANIM_DEBUG
 void FPoseHistory::DebugDraw(FAnimInstanceProxy& AnimInstanceProxy, FColor Color) const
 {
-	CheckThreadSafetyRead(ReadDataThreadSafeCounter);
+	CheckThreadSafetyRead(ReadPoseDataThreadSafeCounter);
 
 	if (CVarAnimPoseHistoryDebugDrawTrajectory.GetValueOnAnyThread())
 	{
@@ -741,9 +816,10 @@ void FPoseHistory::DebugDraw(FAnimInstanceProxy& AnimInstanceProxy, FColor Color
 		const bool bValidTrajectory = !Trajectory.Samples.IsEmpty();
 		TArray<FTransform, TInlineAllocator<128>> PrevGlobalTransforms;
 
-		for (int32 EntryIndex = 0; EntryIndex < ReadData.Entries.Num(); ++EntryIndex)
+		const FPoseData& ReadPoseData = GetReadPoseData();
+		for (int32 EntryIndex = 0; EntryIndex < ReadPoseData.Entries.Num(); ++EntryIndex)
 		{
-			const FPoseHistoryEntry& Entry = ReadData.Entries[EntryIndex];
+			const FPoseHistoryEntry& Entry = ReadPoseData.Entries[EntryIndex];
 
 			const int32 PrevGlobalTransformsNum = PrevGlobalTransforms.Num();
 			const int32 Max = FMath::Max(PrevGlobalTransformsNum, Entry.Num());
diff --git a/Engine/Plugins/Animation/PoseSearch/Source/Runtime/Public/PoseSearch/AnimNode_PoseSearchHistoryCollector.h b/Engine/Plugins/Animation/PoseSearch/Source/Runtime/Public/PoseSearch/AnimNode_PoseSearchHistoryCollector.h
--- a/Engine/Plugins/Animation/PoseSearch/Source/Runtime/Public/PoseSearch/AnimNode_PoseSearchHistoryCollector.h
+++ b/Engine/Plugins/Animation/PoseSearch/Source/Runtime/Public/PoseSearch/AnimNode_PoseSearchHistoryCollector.h
@@ -78,6 +78,8 @@ public:
 	UPROPERTY(EditAnywhere, Category = Experimental, meta=(EditCondition="bGenerateTrajectory", EditConditionHides))
 	FPoseSearchTrajectoryData TrajectoryData;
 
+	bool bCacheBones = false;
+
 	// FAnimNode_Base interface
 	virtual void Initialize_AnyThread(const FAnimationInitializeContext& Context) override;
 	virtual void CacheBones_AnyThread(const FAnimationCacheBonesContext& Context) override;
@@ -90,6 +92,8 @@ public:
 	UE::PoseSearch::FPoseHistory& GetPoseHistory() { return PoseHistory; }
 
 protected:
+	TArray<FBoneIndexType> GetRequiredBones(const FAnimInstanceProxy* AnimInstanceProxy) const;
+
 	UE::PoseSearch::FPoseHistory PoseHistory;
 };
 
diff --git a/Engine/Plugins/Animation/PoseSearch/Source/Runtime/Public/PoseSearch/PoseSearchHistory.h b/Engine/Plugins/Animation/PoseSearch/Source/Runtime/Public/PoseSearch/PoseSearchHistory.h
--- a/Engine/Plugins/Animation/PoseSearch/Source/Runtime/Public/PoseSearch/PoseSearchHistory.h
+++ b/Engine/Plugins/Animation/PoseSearch/Source/Runtime/Public/PoseSearch/PoseSearchHistory.h
@@ -106,11 +106,11 @@ struct FPoseHistory : public IPoseHistory
 	FPoseHistory& operator=(const FPoseHistory& Other);
 	FPoseHistory& operator=(FPoseHistory&& Other);
 
-	void PreUpdate(const UAnimInstance* AnimInstance, float DeltaTime, bool bGenerateTrajectory, const FPoseSearchTrajectoryData& TrajectoryData, const FPoseSearchTrajectoryData::FSampling& TrajectoryDataSampling, bool bNeedsReset);
+	void GenerateTrajectory(const UAnimInstance* AnimInstance, float DeltaTime, const FPoseSearchTrajectoryData& TrajectoryData, const FPoseSearchTrajectoryData::FSampling& TrajectoryDataSampling);
+	void PreUpdate();
 
 	void Initialize_AnyThread(int32 InNumPoses, float InSamplingInterval);
-	void CacheBones_AnyThread(const TArray<FBoneIndexType>& RequiredBones);
-	void EvaluateComponentSpace_AnyThread(float DeltaTime, FCSPose<FCompactPose>& ComponentSpacePose, bool bStoreScales, float RootBoneRecoveryTime);
+	void EvaluateComponentSpace_AnyThread(float DeltaTime, FCSPose<FCompactPose>& ComponentSpacePose, bool bStoreScales, float RootBoneRecoveryTime, bool bNeedsReset, bool bCacheBones, const TArray<FBoneIndexType>& RequiredBones);
 
 	// IPoseHistory interface
 	virtual bool GetTransformAtTime(float Time, FTransform& OutBoneTransform, const USkeleton* BoneIndexSkeleton = nullptr, FBoneIndexType BoneIndexType = RootBoneIndexType, FBoneIndexType ReferenceBoneIndexType = ComponentSpaceIndexType, bool bExtrapolate = false) const override;
@@ -141,26 +141,35 @@ private:
 	// @todo: deprecate this member and expose it via blue print logic or as global query scaling multiplier
 	float TrajectorySpeedMultiplier = 1.f;
 
-	struct FData
+	struct FPoseData
 	{
 		// skeleton from the last Update, to keep tracking skeleton changes, and support compatible skeletons
 		TWeakObjectPtr<const USkeleton> LastUpdateSkeleton;
 
 		// map of FBoneIndexType(s) to collect. If Empty all the bones get collected
 		FBoneToTransformMap BoneToTransformMap;
+		
+		// GetTypeHash for BoneToTransformMap
+		uint32 BoneToTransformMapTypeHash = 0;
 
 		// ring buffer of collected bones
 		TRingBuffer<FPoseHistoryEntry> Entries;
 	};
 
-	FData ReadData;
-	FData WriteData;
+	
+	typedef TStaticArray<FPoseData, 2> FDoubleBufferedPoseData;
+	FDoubleBufferedPoseData DoubleBufferedPoseData;
+	int32 ReadPoseDataIndex = 0;
+
+	int32 GetWritePoseDataIndex() const { return (ReadPoseDataIndex + 1) % 2; }
+	const FPoseData& GetReadPoseData() const { return DoubleBufferedPoseData[ReadPoseDataIndex]; }
+	FPoseData& EditWritePoseData() { return DoubleBufferedPoseData[GetWritePoseDataIndex()]; }
 
 #if ENABLE_ANIM_DEBUG
 	
 	// used to analyze thread safety
-	mutable FThreadSafeCounter ReadDataThreadSafeCounter = 0;
-	mutable FThreadSafeCounter WriteDataThreadSafeCounter = 0;
+	mutable FThreadSafeCounter ReadPoseDataThreadSafeCounter = 0;
+	mutable FThreadSafeCounter WritePoseDataThreadSafeCounter = 0;
 
 #endif // ENABLE_ANIM_DEBUG
 };
