motion matching
- stitching debug lines improvements
- invalid MultiSequence indexing fixes
- MultiSequence support for ValidRanges
- fix for database mirrored asset indexing
- fix for permutation channel debug lines

diff --git a/Engine/Plugins/Animation/PoseSearch/Source/Editor/Private/PoseSearchDatabaseAssetTree.cpp b/Engine/Plugins/Animation/PoseSearch/Source/Editor/Private/PoseSearchDatabaseAssetTree.cpp
--- a/Engine/Plugins/Animation/PoseSearch/Source/Editor/Private/PoseSearchDatabaseAssetTree.cpp
+++ b/Engine/Plugins/Animation/PoseSearch/Source/Editor/Private/PoseSearchDatabaseAssetTree.cpp
@@ -912,7 +912,7 @@ namespace UE::PoseSearch
 								bModified = true;
 							}
 							
-							const FFloatInterval SamplingRange = FPoseSearchDatabaseAnimationAssetBase::GetEffectiveSamplingRange(AnimSequenceBase, DatabaseAnimationAssetBase->GetSamplingRange());
+							const FFloatInterval SamplingRange = DatabaseAnimationAssetBase->GetEffectiveSamplingRange();
 							const float StartTime = SamplingRange.Min;
 							const float Duration = SamplingRange.Max - SamplingRange.Min;
 							const FName TrackName = "PoseSearch";
diff --git a/Engine/Plugins/Animation/PoseSearch/Source/Editor/Private/PoseSearchDatabaseViewModel.cpp b/Engine/Plugins/Animation/PoseSearch/Source/Editor/Private/PoseSearchDatabaseViewModel.cpp
--- a/Engine/Plugins/Animation/PoseSearch/Source/Editor/Private/PoseSearchDatabaseViewModel.cpp
+++ b/Engine/Plugins/Animation/PoseSearch/Source/Editor/Private/PoseSearchDatabaseViewModel.cpp
@@ -33,12 +33,13 @@ namespace UE::PoseSearch
 #if ENABLE_ANIM_DEBUG
 static TAutoConsoleVariable<float> CVarDatabasePreviewDebugDrawSamplerSize(TEXT("a.DatabasePreview.DebugDrawSamplerSize"), 0.f, TEXT("Debug Draw Sampler Positions Size"));
 static TAutoConsoleVariable<float> CVarDatabasePreviewDebugDrawSamplerTimeOffset(TEXT("a.DatabasePreview.DebugDrawSamplerTimeOffset"), 0.f, TEXT("Debug Draw Sampler Positions At Time Offset"));
+static TAutoConsoleVariable<float> CVarDatabasePreviewDebugDrawSamplerRootAxisLength(TEXT("a.DatabasePreview.DebugDrawSamplerRootAxisLength"), 0.f, TEXT("Debug Draw Sampler Root Axis Length"));
 #endif
 
-constexpr float StepDeltaTime = 1.0f / 30.0f;
+constexpr float StepDeltaTime = 1.f / 30.f;
 
 // FDatabasePreviewActor
-bool FDatabasePreviewActor::SpawnPreviewActor(UWorld* World, const UPoseSearchDatabase* PoseSearchDatabase, int32 IndexAssetIdx, const FRole& Role, const FTransform& SamplerRootTransformOrigin, const FTransform* PrecalculatedRootTransformOrigin, int32 PoseIdxForTimeOffset)
+bool FDatabasePreviewActor::SpawnPreviewActor(UWorld* World, const UPoseSearchDatabase* PoseSearchDatabase, int32 IndexAssetIdx, const FRole& Role, const FTransform& SamplerRootTransformOrigin, int32 PoseIdxForTimeOffset)
 {
 	check(PoseSearchDatabase && PoseSearchDatabase->Schema);
 	const FSearchIndex& SearchIndex = PoseSearchDatabase->GetSearchIndex();
@@ -60,18 +61,9 @@ bool FDatabasePreviewActor::SpawnPreviewActor(UWorld* World, const UPoseSearchDa
 	}
 
 	ActorRole = Role;
-	Sampler.Init(PreviewAsset, SamplerRootTransformOrigin, IndexAsset.GetBlendParameters());
-	Sampler.Process();
-
 	IndexAssetIndex = IndexAssetIdx;
 	CurrentPoseIndex = INDEX_NONE;
 
-	PlayTimeOffset = 0.f;
-	if (PoseIdxForTimeOffset >= 0)
-	{
-		PlayTimeOffset = PoseSearchDatabase->GetRealAssetTime(PoseIdxForTimeOffset) - IndexAsset.GetFirstSampleTime(PoseSearchDatabase->Schema->SampleRate);
-	}
-
 	FActorSpawnParameters Params;
 	Params.SpawnCollisionHandlingOverride = ESpawnActorCollisionHandlingMethod::AdjustIfPossibleButAlwaysSpawn;
 	ActorPtr = World->SpawnActor<AActor>(AActor::StaticClass(), FTransform::Identity, Params);
@@ -85,13 +77,13 @@ bool FDatabasePreviewActor::SpawnPreviewActor(UWorld* World, const UPoseSearchDa
 	AnimInstance->InitializeAnimation();
 
 	USkeletalMesh* DatabasePreviewMesh = PoseSearchDatabase->PreviewMesh;
-	Mesh->SetSkeletalMesh(DatabasePreviewMesh ? DatabasePreviewMesh : PoseSearchDatabase->Schema->GetSkeleton(Role)->GetPreviewMesh(true));
+	Mesh->SetSkeletalMesh(DatabasePreviewMesh ? DatabasePreviewMesh : Skeleton->GetPreviewMesh(true));
 	Mesh->EnablePreview(true, PreviewAsset);
 		
-	AnimInstance->SetAnimationAsset(PreviewAsset, IndexAsset.IsLooping(), 0.0f);
+	AnimInstance->SetAnimationAsset(PreviewAsset, IndexAsset.IsLooping(), 0.f);
 	AnimInstance->SetBlendSpacePosition(IndexAsset.GetBlendParameters());
 		
-	if (IndexAsset.IsMirrored() && PoseSearchDatabase->Schema)
+	if (IndexAsset.IsMirrored())
 	{
 		const UMirrorDataTable* MirrorDataTable = PoseSearchDatabase->Schema->GetMirrorDataTable(Role);
 		AnimInstance->SetMirrorDataTable(MirrorDataTable);
@@ -99,17 +91,21 @@ bool FDatabasePreviewActor::SpawnPreviewActor(UWorld* World, const UPoseSearchDa
 
 	const FMirrorDataCache MirrorDataCache(AnimInstance->GetMirrorDataTable(), AnimInstance->GetRequiredBonesOnAnyThread());
 	
-	if (PrecalculatedRootTransformOrigin)
-	{
-		// using the PrecalculatedRootTransformOrigin if provided. useful to synchronize multiple roles preview actors
-		RootTransformOrigin = *PrecalculatedRootTransformOrigin;
-	}
-	else
+	Sampler.Init(PreviewAsset, SamplerRootTransformOrigin, IndexAsset.GetBlendParameters());
+	Sampler.Process();
+
+	PlayTimeOffset = 0.f;
+	if (PoseIdxForTimeOffset >= 0)
 	{
-		RootTransformOrigin = MirrorDataCache.MirrorTransform(Sampler.ExtractRootTransform(PlayTimeOffset));
+		PlayTimeOffset = PoseSearchDatabase->GetRealAssetTime(PoseIdxForTimeOffset) - IndexAsset.GetFirstSampleTime(PoseSearchDatabase->Schema->SampleRate);
+		
+		// centering the Sampler RootTransformOrigin at PlayTimeOffset time, to be able to "align" multiple actors from different animation frames when selected by the pose search debugger
+		FTransform NewSamplerRootTransformOrigin = MirrorDataCache.MirrorTransform(Sampler.ExtractRootTransform(0.f));
+		NewSamplerRootTransformOrigin.SetToRelativeTransform(MirrorDataCache.MirrorTransform(Sampler.ExtractRootTransform(PlayTimeOffset)));
+		Sampler.SetRootTransformOrigin(NewSamplerRootTransformOrigin);
 	}
 
-	AnimInstance->PlayAnim(IndexAsset.IsLooping(), 0.0f);
+	AnimInstance->PlayAnim(IndexAsset.IsLooping(), 0.f);
 	if (!ActorPtr->GetRootComponent())
 	{
 		ActorPtr->SetRootComponent(Mesh);
@@ -117,6 +113,32 @@ bool FDatabasePreviewActor::SpawnPreviewActor(UWorld* World, const UPoseSearchDa
 
 	AnimInstance->SetPlayRate(0.f);
 
+	// initializing Trajectory and TrajectorySpeed
+	const int NumPoses = IndexAsset.GetNumPoses();
+	Trajectory.Samples.SetNumUninitialized(NumPoses);
+	TrajectorySpeed.SetNumUninitialized(NumPoses);
+
+	if (NumPoses > 1)
+	{
+		for (int32 Index = 0; Index < NumPoses; ++Index)
+		{
+			const int32 IndexAssetPoseIdx = Index + IndexAsset.GetFirstPoseIdx();
+			const float IndexAssetPoseTime = IndexAsset.GetTimeFromPoseIndex(IndexAssetPoseIdx, PoseSearchDatabase->Schema->SampleRate);
+			const FTransform IndexAssetPoseTransform = MirrorDataCache.MirrorTransform(Sampler.ExtractRootTransform(IndexAssetPoseTime));
+
+			Trajectory.Samples[Index].SetTransform(IndexAssetPoseTransform);
+			Trajectory.Samples[Index].AccumulatedSeconds = IndexAssetPoseTime;
+		}
+
+		for (int32 Index = 1; Index < NumPoses; ++Index)
+		{
+			const FVector& Start = Trajectory.Samples[Index - 1].Position;
+			const FVector& End = Trajectory.Samples[Index].Position;
+			TrajectorySpeed[Index] = (Start - End).Length() * PoseSearchDatabase->Schema->SampleRate;
+		}
+		TrajectorySpeed[0] = TrajectorySpeed[1];
+	}
+
 	UE_LOG(LogPoseSearchEditor, Log, TEXT("Spawned preview Actor: %s"), *GetNameSafe(ActorPtr.Get()));
 	return true;
 }
@@ -147,7 +169,7 @@ void FDatabasePreviewActor::UpdatePreviewActor(const UPoseSearchDatabase* PoseSe
 	// time to pose index
 	CurrentPoseIndex = PoseSearchDatabase->GetPoseIndexFromTime(CurrentTime, IndexAsset);
 
-	const float QuantizedTime = CurrentPoseIndex >= 0 ? PoseSearchDatabase->GetRealAssetTime(CurrentPoseIndex) : CurrentTime;
+	QuantizedTime = CurrentPoseIndex >= 0 ? PoseSearchDatabase->GetRealAssetTime(CurrentPoseIndex) : CurrentTime;
 	if (bQuantizeAnimationToPoseData)
 	{
 		CurrentTime = QuantizedTime;
@@ -158,38 +180,8 @@ void FDatabasePreviewActor::UpdatePreviewActor(const UPoseSearchDatabase* PoseSe
 	AnimInstance->SetPlayRate(0.f);
 	AnimInstance->SetBlendSpacePosition(IndexAsset.GetBlendParameters());
 
-	const FMirrorDataCache MirrorDataCache(AnimInstance->GetMirrorDataTable(), AnimInstance->GetRequiredBonesOnAnyThread());
-
-	// updating root transforms
-	RootTransformCurrentQuantizedTime = MirrorDataCache.MirrorTransform(Sampler.ExtractRootTransform(QuantizedTime));
-	RootTransformCurrentQuantizedTime.SetToRelativeTransform(RootTransformOrigin);
-
-	if (CurrentTime == QuantizedTime)
-	{
-		RootTransformCurrent = RootTransformCurrentQuantizedTime;
-	}
-	else
-	{
-		RootTransformCurrent = MirrorDataCache.MirrorTransform(Sampler.ExtractRootTransform(CurrentTime));
-		RootTransformCurrent.SetToRelativeTransform(RootTransformOrigin);
-	}
-
 	check(ActorPtr != nullptr);
-	ActorPtr->SetActorTransform(RootTransformCurrent);
-
-	// @todo: optimize this bone container, since we only need the root bone here...
-	const FBoneContainer& BoneContainer = AnimInstance->GetRequiredBonesOnAnyThread();
-	if (BoneContainer.GetNumBones() > 0)
-	{
-		FMemMark Mark(FMemStack::Get());
-		FCompactPose Pose;
-		Pose.SetBoneContainer(&BoneContainer);
-
-		Sampler.ExtractPose(QuantizedTime, Pose);
-		MirrorDataCache.MirrorPose(Pose);
-
-		RootBoneTransformCurrentQuantizedTime = Pose[FCompactPoseBoneIndex(RootBoneIndexType)];
-	}
+	ActorPtr->SetActorTransform(Trajectory.GetSampleAtTime(CurrentTime).GetTransform());
 }
 
 void FDatabasePreviewActor::Destroy()
@@ -200,7 +192,7 @@ void FDatabasePreviewActor::Destroy()
 	}
 }
 
-bool FDatabasePreviewActor::DrawPreviewActors(TArrayView<FDatabasePreviewActor> PreviewActors, const UPoseSearchDatabase* PoseSearchDatabase, bool bDisplayRootMotionSpeed, bool bDisplayBlockTransition, TConstArrayView<float> QueryVector)
+bool FDatabasePreviewActor::DrawPreviewActors(TConstArrayView<FDatabasePreviewActor> PreviewActors, const UPoseSearchDatabase* PoseSearchDatabase, bool bDisplayRootMotionSpeed, bool bDisplayBlockTransition, TConstArrayView<float> QueryVector)
 {
 	using namespace UE::PoseSearch;
 
@@ -222,7 +214,7 @@ bool FDatabasePreviewActor::DrawPreviewActors(TArrayView<FDatabasePreviewActor>
 	ArchivedPoseHistories.Reserve(NumPreviewActors);
 	PoseHistories.Reserve(NumPreviewActors);
 
-	for (FDatabasePreviewActor& PreviewActor : PreviewActors)
+	for (const FDatabasePreviewActor& PreviewActor : PreviewActors)
 	{
 		if (!PoseSearchDatabase->GetSearchIndex().IsValidPoseIndex(PreviewActor.GetCurrentPoseIndex()))
 		{
@@ -270,15 +262,14 @@ bool FDatabasePreviewActor::DrawPreviewActors(TArrayView<FDatabasePreviewActor>
 		RoleToIndex.Add(PreviewActor.ActorRole) = Meshes.Num();
 		Meshes.Add(Mesh);
 
-		const FTransform RootMotionTransform = PreviewActor.RootBoneTransformCurrentQuantizedTime * PreviewActor.RootTransformCurrentQuantizedTime;
-
-		// @todo: reconstruct ArchivedPoseHistory::BoneToTransformMap and ArchivedPoseHistory::Entries if needed
+		// @todo: reconstruct ArchivedPoseHistory::BoneToTransformMap and ArchivedPoseHistory::Entries if needed.
+		//        So far has been unnecessary, becasue all the animations have an identity root bone (root bone lock)
 		FArchivedPoseHistory& ArchivedPoseHistory = ArchivedPoseHistories.AddDefaulted_GetRef();
-		// @todo: reconstruct multiple trajectory samples if needed
-		// reconstructing the FPoseSearchQueryTrajectory with only one sample with AccumulatedSeconds at zero
-		FPoseSearchQueryTrajectorySample& TrajectorySample = ArchivedPoseHistory.Trajectory.Samples.AddDefaulted_GetRef();
-		TrajectorySample.SetTransform(RootMotionTransform);
-		TrajectorySample.AccumulatedSeconds = 0.f;
+		ArchivedPoseHistory.Trajectory = PreviewActor.Trajectory;
+		for (FPoseSearchQueryTrajectorySample& TrajectorySample : ArchivedPoseHistory.Trajectory.Samples)
+		{
+			TrajectorySample.AccumulatedSeconds -= PreviewActor.QuantizedTime;
+		}
 
 		PoseHistories.Add(&ArchivedPoseHistory);
 	}
@@ -291,62 +282,31 @@ bool FDatabasePreviewActor::DrawPreviewActors(TArrayView<FDatabasePreviewActor>
 		DrawParams.DrawFeatureVector(QueryVector);
 	}
 
-	for (FDatabasePreviewActor& PreviewActor : PreviewActors)
+	for (const FDatabasePreviewActor& PreviewActor : PreviewActors)
 	{
 		const UDebugSkelMeshComponent* Mesh = PreviewActor.GetDebugSkelMeshComponent();
 		const FSearchIndex& SearchIndex = PoseSearchDatabase->GetSearchIndex();
 		const FSearchIndexAsset& IndexAsset = SearchIndex.Assets[PreviewActor.IndexAssetIndex];
-
-		const FMirrorDataCache MirrorDataCache(Mesh->PreviewInstance->GetMirrorDataTable(), Mesh->PreviewInstance->GetRequiredBonesOnAnyThread());
-		if (bDisplayRootMotionSpeed || bDisplayBlockTransition)
-		{
-			// initializing SampledRootMotion if required
-			if (PreviewActor.SampledRootMotion.IsEmpty())
-			{
-				const int NumPoses = IndexAsset.GetNumPoses();
-				if (NumPoses > 1)
-				{
-					PreviewActor.SampledRootMotion.SetNumUninitialized(NumPoses);
-					PreviewActor.SampledRootMotionSpeed.SetNumUninitialized(NumPoses);
-
-					for (int32 Index = 0; Index < NumPoses; ++Index)
-					{
-						const int32 IndexAssetPoseIdx = Index + IndexAsset.GetFirstPoseIdx();
-						const float IndexAssetPoseTime = IndexAsset.GetTimeFromPoseIndex(IndexAssetPoseIdx, PoseSearchDatabase->Schema->SampleRate);
-						FTransform IndexAssetPoseTransform = MirrorDataCache.MirrorTransform(PreviewActor.Sampler.ExtractRootTransform(IndexAssetPoseTime));
-						IndexAssetPoseTransform.SetToRelativeTransform(PreviewActor.RootTransformOrigin);
-
-						PreviewActor.SampledRootMotion[Index] = IndexAssetPoseTransform.GetTranslation();
-					}
-
-					for (int32 Index = 1; Index < NumPoses; ++Index)
-					{
-						const FVector& Start = PreviewActor.SampledRootMotion[Index - 1];
-						const FVector& End = PreviewActor.SampledRootMotion[Index];
-						PreviewActor.SampledRootMotionSpeed[Index] = (Start - End).Length() * PoseSearchDatabase->Schema->PermutationsSampleRate;
-					}
-					PreviewActor.SampledRootMotionSpeed[0] = PreviewActor.SampledRootMotionSpeed[1];
-				}
-			}
-		}
+		const int32 SamplesNum = PreviewActor.Trajectory.Samples.Num();
 
 		if (bDisplayRootMotionSpeed)
 		{
-			// drawing PreviewActor.SampledRootMotion
-			const int32 SampledRootMotionNum = PreviewActor.SampledRootMotion.Num();
-			if (SampledRootMotionNum > 1)
+			// @todo: should we be using PreviewActor.Trajectory.DebugDrawTrajectory instead?
+			// drawing PreviewActor.Trajectory
+			
+			if (SamplesNum > 1)
 			{
-				for (int32 Index = 0; Index < PreviewActor.SampledRootMotion.Num(); ++Index)
+				for (int32 Index = 0; Index < SamplesNum; ++Index)
 				{
-					const FVector& EndDown = PreviewActor.SampledRootMotion[Index];
-					const FVector EndUp = EndDown + (PreviewActor.SampledRootMotionSpeed[Index] * FVector::UpVector);
+					const FVector& EndDown = PreviewActor.Trajectory.Samples[Index].Position;
+					const FVector EndUp = EndDown + (PreviewActor.TrajectorySpeed[Index] * FVector::UpVector);
 
 					DrawParams.DrawLine(EndDown, EndUp, FColor::Black);
 					if (Index > 0)
 					{
 						const FColor RootMotionColor = Index % 2 == 0 ? FColor::Purple : FColor::Orange;
-						const FVector& StartDown = PreviewActor.SampledRootMotion[Index - 1];
-						const FVector StartUp = StartDown + (PreviewActor.SampledRootMotionSpeed[Index - 1] * FVector::UpVector);
+						const FVector& StartDown = PreviewActor.Trajectory.Samples[Index - 1].Position;
+						const FVector StartUp = StartDown + (PreviewActor.TrajectorySpeed[Index - 1] * FVector::UpVector);
 						DrawParams.DrawLine(StartDown, EndDown, RootMotionColor);
 						DrawParams.DrawLine(StartUp, EndUp, RootMotionColor);
 					}
@@ -357,24 +317,23 @@ bool FDatabasePreviewActor::DrawPreviewActors(TArrayView<FDatabasePreviewActor>
 		if (bDisplayBlockTransition)
 		{
 			const int NumPoses = IndexAsset.GetNumPoses();
-			if (NumPoses == PreviewActor.SampledRootMotion.Num())
+			if (NumPoses == SamplesNum)
 			{
-				for (int32 Index = 0; Index < NumPoses; ++Index)
+				for (int32 Index = 0; Index < SamplesNum; ++Index)
 				{
 					const int32 IndexAssetPoseIdx = Index + IndexAsset.GetFirstPoseIdx();
 					if (SearchIndex.PoseMetadata[IndexAssetPoseIdx].IsBlockTransition())
 					{
-						DrawParams.DrawPoint(PreviewActor.SampledRootMotion[Index], FColor::Red);
+						DrawParams.DrawPoint(PreviewActor.Trajectory.Samples[Index].Position, FColor::Red);
 					}
 					else
 					{
-						DrawParams.DrawPoint(PreviewActor.SampledRootMotion[Index], FColor::Green);
+						DrawParams.DrawPoint(PreviewActor.Trajectory.Samples[Index].Position, FColor::Green);
 					}
 				}
 			}
 		}
 
-
 #if ENABLE_ANIM_DEBUG
 		const float DebugDrawSamplerSize = CVarDatabasePreviewDebugDrawSamplerSize.GetValueOnAnyThread();
 		if (DebugDrawSamplerSize > UE_KINDA_SMALL_NUMBER)
@@ -387,6 +346,8 @@ bool FDatabasePreviewActor::DrawPreviewActors(TArrayView<FDatabasePreviewActor>
 			FCompactPose Pose;
 			FCSPose<FCompactPose> ComponentSpacePose;
 
+			const FMirrorDataCache MirrorDataCache(Mesh->PreviewInstance->GetMirrorDataTable(), Mesh->PreviewInstance->GetRequiredBonesOnAnyThread());
+
 			for (int32 DrawPass = 0; DrawPass < NumDrawPasses; ++DrawPass)
 			{
 				// drawing the pose extracted from the Sampler to visually compare with the pose features and the mesh drawing
@@ -400,6 +361,14 @@ bool FDatabasePreviewActor::DrawPreviewActors(TArrayView<FDatabasePreviewActor>
 				ComponentSpacePose.InitPose(MoveTemp(Pose));
 
 				const FTransform RootTransform = MirrorDataCache.MirrorTransform(PreviewActor.Sampler.ExtractRootTransform(SamplerTime));
+				const float DebugDrawSamplerRootAxisLength = CVarDatabasePreviewDebugDrawSamplerRootAxisLength.GetValueOnAnyThread();
+				if (DebugDrawSamplerRootAxisLength > 0.f)
+				{
+					DrawParams.DrawLine(RootTransform.GetTranslation(), RootTransform.GetTranslation() + RootTransform.GetScaledAxis(EAxis::X) * DebugDrawSamplerRootAxisLength, FColor::Red);
+					DrawParams.DrawLine(RootTransform.GetTranslation(), RootTransform.GetTranslation() + RootTransform.GetScaledAxis(EAxis::Y) * DebugDrawSamplerRootAxisLength, FColor::Green);
+					DrawParams.DrawLine(RootTransform.GetTranslation(), RootTransform.GetTranslation() + RootTransform.GetScaledAxis(EAxis::Z) * DebugDrawSamplerRootAxisLength, FColor::Blue);
+				}
+
 				for (int32 BoneIndex = 0; BoneIndex < ComponentSpacePose.GetPose().GetNumBones(); ++BoneIndex)
 				{
 					const FTransform BoneWorldTransforms = ComponentSpacePose.GetComponentSpaceTransform(FCompactPoseBoneIndex(BoneIndex)) * RootTransform;
@@ -735,26 +704,18 @@ int32 FDatabaseViewModel::SetSelectedNode(int32 PoseIdx, bool bClearSelection, b
 					const FPoseSearchDatabaseAnimationAssetBase* DatabaseAnimationAsset = Database->GetAnimationAssetBase(IndexAsset.GetSourceAssetIdx());
 					check(DatabaseAnimationAsset);
 					int32 PreviewActorGroupIndex = INDEX_NONE;
-					bool bIsSynchronizedRootTransformOriginValid = false;
-					FTransform SynchronizedRootTransformOrigin = FTransform::Identity;
 					for (int32 RoleIndex = 0; RoleIndex < DatabaseAnimationAsset->GetNumRoles(); ++RoleIndex)
 					{
 						FDatabasePreviewActor PreviewActor;
 						const UE::PoseSearch::FRole Role = DatabaseAnimationAsset->GetRole(RoleIndex);
 						const FTransform& RootTransformOrigin = DatabaseAnimationAsset->GetRootTransformOriginForRole(Role);
-						if (PreviewActor.SpawnPreviewActor(GetWorld(), Database, IndexAssetIndex, Role, RootTransformOrigin, bIsSynchronizedRootTransformOriginValid ? &SynchronizedRootTransformOrigin : nullptr, PoseIdx))
+						if (PreviewActor.SpawnPreviewActor(GetWorld(), Database, IndexAssetIndex, Role, RootTransformOrigin, PoseIdx))
 						{
 							if (PreviewActorGroupIndex == INDEX_NONE)
 							{
 								PreviewActorGroupIndex = PreviewActors.AddDefaulted();
 							}
 
-							if (!bIsSynchronizedRootTransformOriginValid)
-							{
-								bIsSynchronizedRootTransformOriginValid = true;
-								SynchronizedRootTransformOrigin = PreviewActor.GetRootTransformOrigin();
-							}
-
 							MaxPreviewPlayLength = FMath::Max(MaxPreviewPlayLength, IndexAsset.GetLastSampleTime(Database->Schema->SampleRate) - PreviewActor.GetPlayTimeOffset());
 							MinPreviewPlayLength = FMath::Min(MinPreviewPlayLength, IndexAsset.GetFirstSampleTime(Database->Schema->SampleRate) - PreviewActor.GetPlayTimeOffset());
 							PreviewActors[PreviewActorGroupIndex].Add(PreviewActor);
@@ -806,26 +767,18 @@ void FDatabaseViewModel::SetSelectedNodes(const TArrayView<TSharedPtr<FDatabaseA
 					const FPoseSearchDatabaseAnimationAssetBase* DatabaseAnimationAsset = Database->GetAnimationAssetBase(IndexAsset.GetSourceAssetIdx());
 					check(DatabaseAnimationAsset);
 					int32 PreviewActorGroupIndex = INDEX_NONE;
-					bool bIsSynchronizedRootTransformOriginValid = false;
-					FTransform SynchronizedRootTransformOrigin = FTransform::Identity;
 					for (int32 RoleIndex = 0; RoleIndex < DatabaseAnimationAsset->GetNumRoles(); ++RoleIndex)
 					{
 						FDatabasePreviewActor PreviewActor;
 						const UE::PoseSearch::FRole Role = DatabaseAnimationAsset->GetRole(RoleIndex);
 						const FTransform& RootTransformOrigin = DatabaseAnimationAsset->GetRootTransformOriginForRole(Role);
-						if (PreviewActor.SpawnPreviewActor(GetWorld(), Database, IndexAssetIndex, Role, RootTransformOrigin, bIsSynchronizedRootTransformOriginValid ? &SynchronizedRootTransformOrigin : nullptr))
+						if (PreviewActor.SpawnPreviewActor(GetWorld(), Database, IndexAssetIndex, Role, RootTransformOrigin))
 						{
 							if (PreviewActorGroupIndex == INDEX_NONE)
 							{
 								PreviewActorGroupIndex = PreviewActors.AddDefaulted();
 							}
 
-							if (!bIsSynchronizedRootTransformOriginValid)
-							{
-								bIsSynchronizedRootTransformOriginValid = true;
-								SynchronizedRootTransformOrigin = PreviewActor.GetRootTransformOrigin();
-							}
-
 							MaxPreviewPlayLength = FMath::Max(MaxPreviewPlayLength, IndexAsset.GetLastSampleTime(Database->Schema->SampleRate) - IndexAsset.GetFirstSampleTime(Database->Schema->SampleRate));
 							PreviewActors[PreviewActorGroupIndex].Add(PreviewActor);
 						}
@@ -903,7 +856,7 @@ float FDatabaseViewModel::GetPlayTime() const
 void FDatabaseViewModel::SetPlayTime(float NewPlayTime, bool bInTickPlayTime)
 {
 	NewPlayTime = FMath::Clamp(NewPlayTime, MinPreviewPlayLength, MaxPreviewPlayLength);
-	DeltaTimeMultiplier = bInTickPlayTime ? DeltaTimeMultiplier : 0.0f;
+	DeltaTimeMultiplier = bInTickPlayTime ? DeltaTimeMultiplier : 0.f;
 
 	if (!FMath::IsNearlyEqual(PlayTime, NewPlayTime))
 	{
diff --git a/Engine/Plugins/Animation/PoseSearch/Source/Editor/Public/PoseSearchDatabaseViewModel.h b/Engine/Plugins/Animation/PoseSearch/Source/Editor/Public/PoseSearchDatabaseViewModel.h
--- a/Engine/Plugins/Animation/PoseSearch/Source/Editor/Public/PoseSearchDatabaseViewModel.h
+++ b/Engine/Plugins/Animation/PoseSearch/Source/Editor/Public/PoseSearchDatabaseViewModel.h
@@ -8,6 +8,7 @@
 #include "PoseSearch/PoseSearchAssetSampler.h"
 #include "PoseSearch/PoseSearchMirrorDataCache.h"
 #include "PoseSearch/PoseSearchRole.h"
+#include "PoseSearch/PoseSearchTrajectoryTypes.h"
 #include "PoseSearchDatabasePreviewScene.h"
 #include "UObject/GCObject.h"
 
@@ -29,9 +30,9 @@ namespace UE::PoseSearch
 	struct FDatabasePreviewActor
 	{
 	public:
-		bool SpawnPreviewActor(UWorld* World, const UPoseSearchDatabase* PoseSearchDatabase, int32 IndexAssetIdx, const FRole& Role, const FTransform& SamplerRootTransformOrigin, const FTransform* PrecalculatedRootTransformOrigin, int32 PoseIdxForTimeOffset = INDEX_NONE);
+		bool SpawnPreviewActor(UWorld* World, const UPoseSearchDatabase* PoseSearchDatabase, int32 IndexAssetIdx, const FRole& Role, const FTransform& SamplerRootTransformOrigin, int32 PoseIdxForTimeOffset = INDEX_NONE);
 		void UpdatePreviewActor(const UPoseSearchDatabase* PoseSearchDatabase, float PlayTime, bool bQuantizeAnimationToPoseData);
-		static bool DrawPreviewActors(TArrayView<FDatabasePreviewActor> PreviewActors, const UPoseSearchDatabase* PoseSearchDatabase, bool bDisplayRootMotionSpeed, bool bDisplayBlockTransition, TConstArrayView<float> QueryVector);
+		static bool DrawPreviewActors(TConstArrayView<FDatabasePreviewActor> PreviewActors, const UPoseSearchDatabase* PoseSearchDatabase, bool bDisplayRootMotionSpeed, bool bDisplayBlockTransition, TConstArrayView<float> QueryVector);
 
 		void Destroy();
 
@@ -43,7 +44,6 @@ namespace UE::PoseSearch
 		int32 GetIndexAssetIndex() const { return IndexAssetIndex; }
 		int32 GetCurrentPoseIndex() const { return CurrentPoseIndex; }
 		float GetPlayTimeOffset() const { return PlayTimeOffset; }
-		const FTransform& GetRootTransformOrigin() const { return RootTransformOrigin;  }
 		
 	private:
 		UAnimPreviewInstance* GetAnimPreviewInstanceInternal();
@@ -53,19 +53,11 @@ namespace UE::PoseSearch
 		int32 CurrentPoseIndex = INDEX_NONE;
 		float PlayTimeOffset = 0.f;
 		float CurrentTime = 0.f;
-
-		// world space root transforms
-		FTransform RootTransformCurrentQuantizedTime = FTransform::Identity;
-		FTransform RootTransformCurrent = FTransform::Identity;
-		FTransform RootTransformOrigin = FTransform::Identity;
-
-		// local space root BONE transform
-		FTransform RootBoneTransformCurrentQuantizedTime = FTransform::Identity;
+		float QuantizedTime = 0.f;
 
 		FAnimationAssetSampler Sampler;
-
-		TArray<FVector> SampledRootMotion;
-		TArray<float> SampledRootMotionSpeed;
+		FPoseSearchQueryTrajectory Trajectory;
+		TArray<float> TrajectorySpeed;
 
 		FRole ActorRole = DefaultRole;
 	};
diff --git a/Engine/Plugins/Animation/PoseSearch/Source/Runtime/Private/PoseSearchAssetIndexer.cpp b/Engine/Plugins/Animation/PoseSearch/Source/Runtime/Private/PoseSearchAssetIndexer.cpp
--- a/Engine/Plugins/Animation/PoseSearch/Source/Runtime/Private/PoseSearchAssetIndexer.cpp
+++ b/Engine/Plugins/Animation/PoseSearch/Source/Runtime/Private/PoseSearchAssetIndexer.cpp
@@ -483,7 +483,7 @@ FAssetIndexer::FCachedEntry& FAssetIndexer::GetEntry(float SampleTime)
 			StackComponentSpacePose.InitPose(MoveTemp(Pose));
 			Entry->ComponentSpacePose[RoleIndex].CopyPose(StackComponentSpacePose);
 
-			Entry->RootTransform[RoleIndex] = SampleRootTransform;
+			Entry->RootTransform[RoleIndex] = MirrorTransform(SampleRootTransform, RoleIndex);
 			Entry->bClamped |= bSampleClamped;
 		}
 	}
@@ -514,7 +514,7 @@ FTransform FAssetIndexer::GetTransform(float SampleTime, const FRole& Role, bool
 
 	const int32 RoleIndex = RoleToIndex[Role];
 	const FBoneReference& BoneReference = Schema.GetBoneReferences(Role)[SchemaBoneIdx];
-	return CalculateComponentSpaceTransform(Entry, BoneReference, RoleIndex) * MirrorTransform(Entry.RootTransform[RoleIndex], RoleIndex);
+	return CalculateComponentSpaceTransform(Entry, BoneReference, RoleIndex) * Entry.RootTransform[RoleIndex];
 }
 
 // returns the transform in animation space for the BoneReference at SampleTime seconds
@@ -522,7 +522,7 @@ FTransform FAssetIndexer::GetTransform(float SampleTime, int32 RoleIndex, bool&
 {
 	FCachedEntry& Entry = GetEntry(SampleTime);
 	bClamped = Entry.bClamped;
-	return CalculateComponentSpaceTransform(Entry, BoneReference, RoleIndex) * MirrorTransform(Entry.RootTransform[RoleIndex], RoleIndex);
+	return CalculateComponentSpaceTransform(Entry, BoneReference, RoleIndex) * Entry.RootTransform[RoleIndex];
 }
 
 FTransform FAssetIndexer::CalculateComponentSpaceTransform(FAssetIndexer::FCachedEntry& Entry, const FBoneReference& BoneReference, int32 RoleIndex)
diff --git a/Engine/Plugins/Animation/PoseSearch/Source/Runtime/Private/PoseSearchContext.cpp b/Engine/Plugins/Animation/PoseSearch/Source/Runtime/Private/PoseSearchContext.cpp
--- a/Engine/Plugins/Animation/PoseSearch/Source/Runtime/Private/PoseSearchContext.cpp
+++ b/Engine/Plugins/Animation/PoseSearch/Source/Runtime/Private/PoseSearchContext.cpp
@@ -7,8 +7,9 @@
 #include "PoseSearch/PoseSearchHistory.h"
 #include "PoseSearch/PoseSearchSchema.h"
 #include "PoseSearch/PoseSearchTrajectoryTypes.h"
-#include "PoseSearchFeatureChannel_Position.h"
 #include "PoseSearchFeatureChannel_Heading.h"
+#include "PoseSearchFeatureChannel_PermutationTime.h"
+#include "PoseSearchFeatureChannel_Position.h"
 
 namespace UE::PoseSearch
 {
@@ -55,7 +56,28 @@ const UPoseSearchSchema* FDebugDrawParams::GetSchema() const
 	return Database ? Database->Schema : nullptr;
 }
 
-FVector FDebugDrawParams::ExtractPosition(TConstArrayView<float> PoseVector, float SampleTimeOffset, int8 SchemaBoneIdx, const FRole& Role, EPermutationTimeType PermutationTimeType, int32 SamplingAttributeId) const
+float FDebugDrawParams::ExtractPermutationTime(TConstArrayView<float> PoseVector) const
+{
+	if (const UPoseSearchSchema* Schema = GetSchema())
+	{
+		if (const UPoseSearchFeatureChannel_PermutationTime* FoundPermutationTime = static_cast<const UPoseSearchFeatureChannel_PermutationTime*>(
+			Schema->FindChannel([](const UPoseSearchFeatureChannel* Channel) -> const UPoseSearchFeatureChannel_PermutationTime*
+				{
+					if (const UPoseSearchFeatureChannel_PermutationTime* PermutationTime = Cast<UPoseSearchFeatureChannel_PermutationTime>(Channel))
+					{
+						return PermutationTime;
+					}
+					return nullptr;
+				})))
+		{
+			check(FoundPermutationTime->GetChannelCardinality() == 1);
+			return FFeatureVectorHelper::DecodeFloat(PoseVector, FoundPermutationTime->GetChannelDataOffset());
+		}
+	}
+	return 0.f;
+}
+
+FVector FDebugDrawParams::ExtractPosition(TConstArrayView<float> PoseVector, float SampleTimeOffset, int8 SchemaBoneIdx, const FRole& Role, EPermutationTimeType PermutationTimeType, int32 SamplingAttributeId, float PermutationSampleTimeOffset) const
 {
 	// we don't wanna ask for a SchemaOriginBoneIdx in the future or past
 	check(PermutationTimeType != EPermutationTimeType::UsePermutationTime);
@@ -97,7 +119,7 @@ FVector FDebugDrawParams::ExtractPosition(TConstArrayView<float> PoseVector, flo
 					const FBoneIndexType BoneIndexType = Schema->GetBoneReferences(Role)[SchemaBoneIdx].BoneIndex;
 
 					FTransform WorldBoneTransform;
-					if (PoseHistory->GetTransformAtTime(SampleTimeOffset, WorldBoneTransform, Skeleton, BoneIndexType, WorldSpaceIndexType))
+					if (PoseHistory->GetTransformAtTime(SampleTimeOffset + PermutationSampleTimeOffset, WorldBoneTransform, Skeleton, BoneIndexType, WorldSpaceIndexType))
 					{
 						return WorldBoneTransform.GetTranslation();
 					}
@@ -110,11 +132,13 @@ FVector FDebugDrawParams::ExtractPosition(TConstArrayView<float> PoseVector, flo
 			}
 		}
 	}
-	return GetRootBoneTransform(Role, SampleTimeOffset).GetTranslation();
+	return GetRootBoneTransform(Role, SampleTimeOffset + PermutationSampleTimeOffset).GetTranslation();
 }
 
-FQuat FDebugDrawParams::ExtractRotation(TConstArrayView<float> PoseVector, float SampleTimeOffset, int8 SchemaBoneIdx, const FRole& Role, EPermutationTimeType PermutationTimeType, int32 SamplingAttributeId) const
+FQuat FDebugDrawParams::ExtractRotation(TConstArrayView<float> PoseVector, float SampleTimeOffset, int8 SchemaBoneIdx, const FRole& Role, EPermutationTimeType PermutationTimeType, int32 SamplingAttributeId, float PermutationSampleTimeOffset) const
 {
+	// we don't wanna ask for a SchemaOriginBoneIdx in the future or past
+	check(PermutationTimeType != EPermutationTimeType::UsePermutationTime);
 	if (const UPoseSearchSchema* Schema = GetSchema())
 	{
 		int32 HeadingAxisFoundNum = 0;
@@ -222,7 +246,7 @@ FQuat FDebugDrawParams::ExtractRotation(TConstArrayView<float> PoseVector, float
 					const FBoneIndexType BoneIndexType = Schema->GetBoneReferences(Role)[SchemaBoneIdx].BoneIndex;
 
 					FTransform WorldBoneTransform;
-					if (PoseHistory->GetTransformAtTime(SampleTimeOffset, WorldBoneTransform, Skeleton, BoneIndexType, WorldSpaceIndexType))
+					if (PoseHistory->GetTransformAtTime(SampleTimeOffset + PermutationSampleTimeOffset, WorldBoneTransform, Skeleton, BoneIndexType, WorldSpaceIndexType))
 					{
 						return WorldBoneTransform.GetRotation();
 					}
@@ -236,7 +260,7 @@ FQuat FDebugDrawParams::ExtractRotation(TConstArrayView<float> PoseVector, float
 		}
 	}
 
-	return GetRootBoneTransform(Role, SampleTimeOffset).GetRotation();
+	return GetRootBoneTransform(Role, SampleTimeOffset + PermutationSampleTimeOffset).GetRotation();
 }
 
 FTransform FDebugDrawParams::GetRootBoneTransform(const FRole& Role, float SampleTimeOffset) const
diff --git a/Engine/Plugins/Animation/PoseSearch/Source/Runtime/Private/PoseSearchDatabase.cpp b/Engine/Plugins/Animation/PoseSearch/Source/Runtime/Private/PoseSearchDatabase.cpp
--- a/Engine/Plugins/Animation/PoseSearch/Source/Runtime/Private/PoseSearchDatabase.cpp
+++ b/Engine/Plugins/Animation/PoseSearch/Source/Runtime/Private/PoseSearchDatabase.cpp
@@ -364,6 +364,13 @@ bool FPoseSearchDatabaseAnimationAssetBase::IsSkeletonCompatible(TObjectPtr<cons
 }
 #endif // WITH_EDITOR
 
+#if WITH_EDITORONLY_DATA
+const FString FPoseSearchDatabaseAnimationAssetBase::GetName() const
+{
+	return GetNameSafe(GetAnimationAsset());
+}
+#endif //WITH_EDITORONLY_DATA
+
 UAnimationAsset* FPoseSearchDatabaseAnimationAssetBase::GetAnimationAssetForRole(const UE::PoseSearch::FRole& Role) const
 {
 	check(GetNumRoles() == 1);
@@ -383,13 +390,15 @@ int64 FPoseSearchDatabaseAnimationAssetBase::GetEditorMemSize() const
 	return EditorMemCount.GetNum();
 }
 
-FFloatInterval FPoseSearchDatabaseAnimationAssetBase::GetEffectiveSamplingRange(const UAnimSequenceBase* SequenceBase, const FFloatInterval& RequestedSamplingRange)
+FFloatInterval FPoseSearchDatabaseAnimationAssetBase::GetEffectiveSamplingRange() const
 {
-	const bool bSampleAll = (RequestedSamplingRange.Min == 0.0f) && (RequestedSamplingRange.Max == 0.0f);
-	const float SequencePlayLength = SequenceBase->GetPlayLength();
+	const FFloatInterval SamplingRange = GetSamplingRange();
+
+	const bool bSampleAll = (SamplingRange.Min == 0.0f) && (SamplingRange.Max == 0.0f);
+	const float PlayLength = GetPlayLength();
 	FFloatInterval Range;
-	Range.Min = bSampleAll ? 0.0f : RequestedSamplingRange.Min;
-	Range.Max = bSampleAll ? SequencePlayLength : FMath::Min(SequencePlayLength, RequestedSamplingRange.Max);
+	Range.Min = bSampleAll ? 0.0f : SamplingRange.Min;
+	Range.Max = bSampleAll ? PlayLength : FMath::Min(PlayLength, SamplingRange.Max);
 	return Range;
 }
 #endif // WITH_EDITORONLY_DATA
@@ -415,11 +424,6 @@ bool FPoseSearchDatabaseSequence::IsLooping() const
 		SamplingRange.Max == 0.f;
 }
 
-const FString FPoseSearchDatabaseSequence::GetName() const
-{
-	return Sequence ? Sequence->GetName() : FString();
-}
-
 bool FPoseSearchDatabaseSequence::IsRootMotionEnabled() const
 {
 	return Sequence ? Sequence->HasRootMotion() : false;
@@ -452,11 +456,6 @@ bool FPoseSearchDatabaseBlendSpace::IsLooping() const
 	return BlendSpace && BlendSpace->bLoop;
 }
 
-const FString FPoseSearchDatabaseBlendSpace::GetName() const
-{
-	return BlendSpace ? BlendSpace->GetName() : FString();
-}
-
 bool FPoseSearchDatabaseBlendSpace::IsRootMotionEnabled() const
 {
 	bool bIsRootMotionUsedInBlendSpace = false;
@@ -566,11 +565,6 @@ bool FPoseSearchDatabaseAnimComposite::IsLooping() const
 		SamplingRange.Max == 0.f;
 }
 
-const FString FPoseSearchDatabaseAnimComposite::GetName() const
-{
-	return AnimComposite ? AnimComposite->GetName() : FString();
-}
-
 bool FPoseSearchDatabaseAnimComposite::IsRootMotionEnabled() const
 {
 	return AnimComposite ? AnimComposite->HasRootMotion() : false;
@@ -598,11 +592,6 @@ bool FPoseSearchDatabaseAnimMontage::IsLooping() const
 		SamplingRange.Max == 0.f;
 }
 
-const FString FPoseSearchDatabaseAnimMontage::GetName() const
-{
-	return AnimMontage ? AnimMontage->GetName() : FString();
-}
-
 bool FPoseSearchDatabaseAnimMontage::IsRootMotionEnabled() const
 {
 	return AnimMontage ? AnimMontage->HasRootMotion() : false;
@@ -662,11 +651,6 @@ bool FPoseSearchDatabaseMultiSequence::IsLooping() const
 		SamplingRange.Max == 0.f;
 }
 
-const FString FPoseSearchDatabaseMultiSequence::GetName() const
-{
-	return MultiSequence ? MultiSequence->GetName() : FString();
-}
-
 bool FPoseSearchDatabaseMultiSequence::IsRootMotionEnabled() const
 {
 	return MultiSequence ? MultiSequence->HasRootMotion() : false;
diff --git a/Engine/Plugins/Animation/PoseSearch/Source/Runtime/Private/PoseSearchDerivedData.cpp b/Engine/Plugins/Animation/PoseSearch/Source/Runtime/Private/PoseSearchDerivedData.cpp
--- a/Engine/Plugins/Animation/PoseSearch/Source/Runtime/Private/PoseSearchDerivedData.cpp
+++ b/Engine/Plugins/Animation/PoseSearch/Source/Runtime/Private/PoseSearchDerivedData.cpp
@@ -269,18 +269,18 @@ public:
 	}
 };
 
-static void FindValidSequenceIntervals(const UAnimSequenceBase* SequenceBase, FFloatInterval SamplingRange, bool bIsLooping,
-	const FFloatInterval& ExcludeFromDatabaseParameters, TArray<FFloatRange>& ValidRanges)
+static void FindValidSequenceIntervals(const FPoseSearchDatabaseAnimationAssetBase* DatabaseAsset, const FFloatInterval& ExcludeFromDatabaseParameters, TArray<FFloatRange>& ValidRanges)
 {
-	check(SequenceBase);
+	check(DatabaseAsset);
 
-	const float SequenceLength = SequenceBase->GetPlayLength();
+	bool bIsLooping = DatabaseAsset->IsLooping();
+	const float PlayLength = DatabaseAsset->GetPlayLength();
 
-	const FFloatInterval EffectiveSamplingInterval = FPoseSearchDatabaseAnimationAssetBase::GetEffectiveSamplingRange(SequenceBase, SamplingRange);
+	const FFloatInterval EffectiveSamplingInterval = DatabaseAsset->GetEffectiveSamplingRange();
 	FFloatRange EffectiveSamplingRange = FFloatRange::Inclusive(EffectiveSamplingInterval.Min, EffectiveSamplingInterval.Max);
 	if (!bIsLooping)
 	{
-		const FFloatRange ExcludeFromDatabaseRange(ExcludeFromDatabaseParameters.Min, SequenceLength + ExcludeFromDatabaseParameters.Max);
+		const FFloatRange ExcludeFromDatabaseRange(ExcludeFromDatabaseParameters.Min, PlayLength + ExcludeFromDatabaseParameters.Max);
 		EffectiveSamplingRange = FFloatRange::Intersection(EffectiveSamplingRange, ExcludeFromDatabaseRange);
 	}
 
@@ -288,40 +288,52 @@ static void FindValidSequenceIntervals(const UAnimSequenceBase* SequenceBase, FF
 	ValidRanges.Empty();
 	ValidRanges.Add(EffectiveSamplingRange);
 
-	FAnimNotifyContext NotifyContext;
-	SequenceBase->GetAnimNotifies(0.0f, SequenceLength, NotifyContext);
-
-	for (const FAnimNotifyEventReference& EventReference : NotifyContext.ActiveNotifies)
+	for (int32 RoleIndex = 0; RoleIndex < DatabaseAsset->GetNumRoles(); ++RoleIndex)
 	{
-		if (const FAnimNotifyEvent* NotifyEvent = EventReference.GetNotify())
+		const FRole Role = DatabaseAsset->GetRole(RoleIndex);
+		if (UAnimSequenceBase* SequenceBase = Cast<UAnimSequenceBase>(DatabaseAsset->GetAnimationAssetForRole(Role)))
 		{
-			if (const UAnimNotifyState_PoseSearchExcludeFromDatabase* ExclusionNotifyState = Cast<const UAnimNotifyState_PoseSearchExcludeFromDatabase>(NotifyEvent->NotifyStateClass))
-			{
-				FFloatRange ExclusionRange = FFloatRange::Inclusive(NotifyEvent->GetTime(), NotifyEvent->GetTime() + NotifyEvent->GetDuration());
+			FAnimNotifyContext NotifyContext;
+			SequenceBase->GetAnimNotifies(0.0f, PlayLength, NotifyContext);
 
-				// Split every valid range based on the exclusion range just found. Because this might increase the 
-				// number of ranges in ValidRanges, the algorithm iterates from end to start.
-				for (int32 RangeIdx = ValidRanges.Num() - 1; RangeIdx >= 0; --RangeIdx)
+			for (const FAnimNotifyEventReference& EventReference : NotifyContext.ActiveNotifies)
+			{
+				if (const FAnimNotifyEvent* NotifyEvent = EventReference.GetNotify())
 				{
-					FFloatRange EvaluatedRange = ValidRanges[RangeIdx];
-					ValidRanges.RemoveAt(RangeIdx);
+					if (const UAnimNotifyState_PoseSearchExcludeFromDatabase* ExclusionNotifyState = Cast<const UAnimNotifyState_PoseSearchExcludeFromDatabase>(NotifyEvent->NotifyStateClass))
+					{
+						FFloatRange ExclusionRange = FFloatRange::Inclusive(NotifyEvent->GetTime(), NotifyEvent->GetTime() + NotifyEvent->GetDuration());
+
+						// Split every valid range based on the exclusion range just found. Because this might increase the 
+						// number of ranges in ValidRanges, the algorithm iterates from end to start.
+						for (int32 RangeIdx = ValidRanges.Num() - 1; RangeIdx >= 0; --RangeIdx)
+						{
+							FFloatRange EvaluatedRange = ValidRanges[RangeIdx];
+							ValidRanges.RemoveAt(RangeIdx);
 
-					TArray<FFloatRange> Diff = FFloatRange::Difference(EvaluatedRange, ExclusionRange);
-					ValidRanges.Append(Diff);
+							TArray<FFloatRange> Diff = FFloatRange::Difference(EvaluatedRange, ExclusionRange);
+							ValidRanges.Append(Diff);
+						}
+					}
 				}
 			}
 		}
 	}
 }
 
-static void InitSearchIndexAssets(FSearchIndexBase& SearchIndex, const TArray<FInstancedStruct>& DatabaseAnimationAssets, const UPoseSearchSchema* Schema, const FFloatInterval& ExcludeFromDatabaseParameters)
+// returns false in case of errors
+static bool InitSearchIndexAssets(FSearchIndexBase& SearchIndex, const TArray<FInstancedStruct>& DatabaseAnimationAssets, const UPoseSearchSchema* Schema, const FFloatInterval& ExcludeFromDatabaseParameters)
 {
 	using namespace UE::PoseSearch;
 
+	check(Schema);
+
 	SearchIndex.Assets.Empty();
 	TArray<FFloatRange> ValidRanges;
 	TArray<FBlendSampleData> BlendSamples;
 
+	bool bAnyErrors = false;
+
 	int32 TotalPoses = 0;
 	for (int32 AnimationAssetIndex = 0; AnimationAssetIndex < DatabaseAnimationAssets.Num(); ++AnimationAssetIndex)
 	{
@@ -333,6 +345,29 @@ static void InitSearchIndexAssets(FSearchIndexBase& SearchIndex, const TArray<FI
 				continue;
 			}
 
+			if (DatabaseAsset->GetNumRoles() != Schema->GetRoledSkeletons().Num())
+			{
+				UE_LOG(LogPoseSearch, Error, TEXT("DatabaseAsset '%s' Roles don't match Schema '%s' Skeletons Roles"), *DatabaseAsset->GetAnimationAsset()->GetName(), *Schema->GetName());
+				bAnyErrors = true;
+				continue;
+			}
+
+			// checking for valid roles in DatabaseMultiSequence against the Schema
+			bool bAreAllRolesSupported = true;
+			for (const FPoseSearchRoledSkeleton& RoledSkeleton : Schema->GetRoledSkeletons())
+			{
+				if (!DatabaseAsset->GetAnimationAssetForRole(RoledSkeleton.Role))
+				{
+					UE_LOG(LogPoseSearch, Error, TEXT("DatabaseAsset '%s' doesn't support Role '%s' required by Schema '%s' Skeletons"), *DatabaseAsset->GetAnimationAsset()->GetName(), *RoledSkeleton.Role.ToString(), *Schema->GetName());
+					bAreAllRolesSupported = false;
+				}
+			}
+			if (!bAreAllRolesSupported)
+			{
+				bAnyErrors = true;
+				continue;
+			}
+
 			const bool bAddUnmirrored = DatabaseAsset->GetMirrorOption() == EPoseSearchMirrorOption::UnmirroredOnly || DatabaseAsset->GetMirrorOption() == EPoseSearchMirrorOption::UnmirroredAndMirrored;
 			const bool bAddMirrored = DatabaseAsset->GetMirrorOption() == EPoseSearchMirrorOption::MirroredOnly || DatabaseAsset->GetMirrorOption() == EPoseSearchMirrorOption::UnmirroredAndMirrored;
 			const bool bIsLooping = DatabaseAsset->IsLooping();
@@ -383,42 +418,12 @@ static void InitSearchIndexAssets(FSearchIndexBase& SearchIndex, const TArray<FI
 					}
 				}
 			}
-			else if (const FPoseSearchDatabaseMultiSequence* DatabaseMultiSequence = DatabaseAssetStruct.GetPtr<FPoseSearchDatabaseMultiSequence>())
-			{
-				// @todo: support FindValidSequenceIntervals(SequenceBase, DatabaseAsset->GetSamplingRange(), bIsLooping, Database->ExcludeFromDatabaseParameters, ValidRanges);
-				const float PlayLength = DatabaseMultiSequence->GetPlayLength();
-
-				for (int32 PermutationIdx = 0; PermutationIdx < Schema->NumberOfPermutations; ++PermutationIdx)
-				{
-					if (bAddUnmirrored)
-					{
-						const FSearchIndexAsset PoseSearchIndexAsset(AnimationAssetIndex, TotalPoses, false, bIsLooping, 
-							bDisableReselection, FFloatInterval(0.f, PlayLength), Schema->SampleRate, PermutationIdx);
-						if (PoseSearchIndexAsset.GetNumPoses() > 0)
-						{
-							SearchIndex.Assets.Add(PoseSearchIndexAsset);
-							TotalPoses += PoseSearchIndexAsset.GetNumPoses();
-						}
-					}
-
-					if (bAddMirrored)
-					{
-						const FSearchIndexAsset PoseSearchIndexAsset(AnimationAssetIndex, TotalPoses, true, bIsLooping,
-							bDisableReselection, FFloatInterval(0.f, PlayLength), Schema->SampleRate, PermutationIdx);
-						if (PoseSearchIndexAsset.GetNumPoses() > 0)
-						{
-							SearchIndex.Assets.Add(PoseSearchIndexAsset);
-							TotalPoses += PoseSearchIndexAsset.GetNumPoses();
-						}
-					}
-				}
-			}
-			// support for FPoseSearchDatabaseSequence, FPoseSearchDatabaseAnimComposite, FPoseSearchDatabaseAnimMontage
-			else if (const UAnimSequenceBase* SequenceBase = Cast<UAnimSequenceBase>(DatabaseAsset->GetAnimationAsset()))
+			// support for FPoseSearchDatabaseSequence, FPoseSearchDatabaseAnimComposite, FPoseSearchDatabaseAnimMontage, FPoseSearchDatabaseMultiSequence
+			else
 			{
 				ValidRanges.Reset();
 
-				FindValidSequenceIntervals(SequenceBase, DatabaseAsset->GetSamplingRange(), bIsLooping, ExcludeFromDatabaseParameters, ValidRanges);
+				FindValidSequenceIntervals(DatabaseAsset, ExcludeFromDatabaseParameters, ValidRanges);
 				for (const FFloatRange& Range : ValidRanges)
 				{
 					for (int32 PermutationIdx = 0; PermutationIdx < Schema->NumberOfPermutations; ++PermutationIdx)
@@ -447,12 +452,10 @@ static void InitSearchIndexAssets(FSearchIndexBase& SearchIndex, const TArray<FI
 					}
 				}
 			}
-			else
-			{
-				checkNoEntry();
-			}
 		}
 	}
+
+	return !bAnyErrors;
 }
 
 static void PreprocessSearchIndexWeights(FSearchIndex& SearchIndex, const UPoseSearchSchema* Schema, TConstArrayView<float> Deviation)
@@ -1611,7 +1614,12 @@ void FPoseSearchDatabaseAsyncCacheTask::OnGetComplete(UE::DerivedData::FCacheGet
 					}
 
 					// Building all the related FPoseSearchBaseIndex first
-					InitSearchIndexAssets(SearchIndexBase, DependentDatabaseAnimationAssets, DependentDatabaseSchema, DependentExcludeFromDatabaseParameters);
+					if (!InitSearchIndexAssets(SearchIndexBase, DependentDatabaseAnimationAssets, DependentDatabaseSchema, DependentExcludeFromDatabaseParameters))
+					{
+						UE_LOG(LogPoseSearch, Error, TEXT("%s - %s BuildIndex Failed becasue of invalid assets"), *LexToString(FullIndexKey.Hash), *MainDatabaseName);
+						ResetSearchIndex();
+						return;
+					}
 
 					if (Owner.IsCanceled())
 					{
diff --git a/Engine/Plugins/Animation/PoseSearch/Source/Runtime/Private/PoseSearchDerivedDataKey.cpp b/Engine/Plugins/Animation/PoseSearch/Source/Runtime/Private/PoseSearchDerivedDataKey.cpp
--- a/Engine/Plugins/Animation/PoseSearch/Source/Runtime/Private/PoseSearchDerivedDataKey.cpp
+++ b/Engine/Plugins/Animation/PoseSearch/Source/Runtime/Private/PoseSearchDerivedDataKey.cpp
@@ -42,7 +42,7 @@ FKeyBuilder::FKeyBuilder(const UObject* Object, bool bUseDataVer, bool bPerformC
 	if (bUseDataVer)
 	{
 		// used to invalidate the key without having to change POSESEARCHDB_DERIVEDDATA_VER all the times
-		int32 POSESEARCHDB_DERIVEDDATA_VER_SMALL = 247;
+		int32 POSESEARCHDB_DERIVEDDATA_VER_SMALL = 249;
 		FGuid VersionGuid = FDevSystemGuids::GetSystemGuid(FDevSystemGuids::Get().POSESEARCHDB_DERIVEDDATA_VER);
 		FString AnimationCompressionVersionString = UE::Anim::Compression::AnimationCompressionVersionString;
 
diff --git a/Engine/Plugins/Animation/PoseSearch/Source/Runtime/Private/PoseSearchFeatureChannel_Heading.cpp b/Engine/Plugins/Animation/PoseSearch/Source/Runtime/Private/PoseSearchFeatureChannel_Heading.cpp
--- a/Engine/Plugins/Animation/PoseSearch/Source/Runtime/Private/PoseSearchFeatureChannel_Heading.cpp
+++ b/Engine/Plugins/Animation/PoseSearch/Source/Runtime/Private/PoseSearchFeatureChannel_Heading.cpp
@@ -70,16 +70,13 @@ void UPoseSearchFeatureChannel_Heading::AddDependentChannels(UPoseSearchSchema*
 {
 	if (Schema->bInjectAdditionalDebugChannels)
 	{
-		UPoseSearchFeatureChannel_Position::FindOrAddToSchema(Schema, SampleTimeOffset, Bone.BoneName, SampleRole);
-		if (!FMath::IsNearlyZero(OriginTimeOffset))
-		{
-			// adding the position OriginTimeOffset seconds ahead
-			UPoseSearchFeatureChannel_Position::FindOrAddToSchema(Schema, OriginTimeOffset, Bone.BoneName, OriginRole);
-			
-			// adding the rotation (X, Y axis) OriginTimeOffset seconds ahead
-			UPoseSearchFeatureChannel_Heading::FindOrAddToSchema(Schema, OriginTimeOffset, Bone.BoneName, OriginRole, EHeadingAxis::X);
-			UPoseSearchFeatureChannel_Heading::FindOrAddToSchema(Schema, OriginTimeOffset, Bone.BoneName, OriginRole, EHeadingAxis::Y);
-		}
+		const EPermutationTimeType SamplePermutationTimeType = PermutationTimeType != EPermutationTimeType::UseSampleTime ? EPermutationTimeType::UseSampleToPermutationTime : EPermutationTimeType::UseSampleTime;
+		UPoseSearchFeatureChannel_Position::FindOrAddToSchema(Schema, SampleTimeOffset, Bone.BoneName, SampleRole, SamplePermutationTimeType);
+		
+		// injecting 2 Heading channels to be able to reconstruct the Origin bone rotation
+		const EPermutationTimeType OriginPermutationTimeType = PermutationTimeType == EPermutationTimeType::UsePermutationTime ? EPermutationTimeType::UseSampleToPermutationTime : EPermutationTimeType::UseSampleTime;
+		UPoseSearchFeatureChannel_Heading::FindOrAddToSchema(Schema, OriginTimeOffset, Bone.BoneName, OriginRole, EHeadingAxis::X, OriginPermutationTimeType);
+		UPoseSearchFeatureChannel_Heading::FindOrAddToSchema(Schema, OriginTimeOffset, Bone.BoneName, OriginRole, EHeadingAxis::Y, OriginPermutationTimeType);
 	}
 }
 
@@ -183,6 +180,18 @@ void UPoseSearchFeatureChannel_Heading::DebugDraw(const UE::PoseSearch::FDebugDr
 {
 	using namespace UE::PoseSearch;
 
+#if WITH_EDITORONLY_DATA
+	bool bDrawInjectAdditionalDebugChannels = false;
+	if (const UPoseSearchSchema* Schema = GetSchema())
+	{
+		bDrawInjectAdditionalDebugChannels = Schema->bDrawInjectAdditionalDebugChannels;
+	}
+	if (!bDrawInjectAdditionalDebugChannels && Weight <= 0.f)
+	{
+		return;
+	}
+#endif // WITH_EDITORONLY_DATA
+
 	FColor Color;
 #if WITH_EDITORONLY_DATA
 	Color = DebugColor.ToFColor(true);
@@ -190,8 +199,15 @@ void UPoseSearchFeatureChannel_Heading::DebugDraw(const UE::PoseSearch::FDebugDr
 	Color = FLinearColor::White.ToFColor(true);
 #endif // WITH_EDITORONLY_DATA
 
-	const FVector BoneHeading = DrawParams.ExtractRotation(PoseVector, OriginTimeOffset, RootSchemaBoneIdx, OriginRole).RotateVector(FFeatureVectorHelper::DecodeVector(PoseVector, ChannelDataOffset, ComponentStripping));
-	const FVector BonePos = DrawParams.ExtractPosition(PoseVector, SampleTimeOffset, SchemaBoneIdx, SampleRole, PermutationTimeType, SamplingAttributeId);
+	float PermutationSampleTimeOffset = 0.f;
+	float PermutationOriginTimeOffset = 0.f;
+	UPoseSearchFeatureChannel::GetPermutationTimeOffsets(PermutationTimeType, DrawParams.ExtractPermutationTime(PoseVector), PermutationSampleTimeOffset, PermutationOriginTimeOffset);
+	const EPermutationTimeType SamplePermutationTimeType = PermutationTimeType != EPermutationTimeType::UseSampleTime ? EPermutationTimeType::UseSampleToPermutationTime : EPermutationTimeType::UseSampleTime;
+	const EPermutationTimeType OriginPermutationTimeType = PermutationTimeType == EPermutationTimeType::UsePermutationTime ? EPermutationTimeType::UseSampleToPermutationTime : EPermutationTimeType::UseSampleTime;
+
+	const FVector FeaturesVector = FFeatureVectorHelper::DecodeVector(PoseVector, ChannelDataOffset, ComponentStripping);
+	const FVector BoneHeading = DrawParams.ExtractRotation(PoseVector, OriginTimeOffset, SchemaOriginBoneIdx, OriginRole, OriginPermutationTimeType, INDEX_NONE, PermutationOriginTimeOffset).RotateVector(FeaturesVector);
+	const FVector BonePos = DrawParams.ExtractPosition(PoseVector, SampleTimeOffset, SchemaBoneIdx, SampleRole, SamplePermutationTimeType, SamplingAttributeId, PermutationSampleTimeOffset);
 
 	DrawParams.DrawPoint(BonePos, Color, 3.f);
 	DrawParams.DrawLine(BonePos + BoneHeading * 4.f, BonePos + BoneHeading * 15.f, Color);
diff --git a/Engine/Plugins/Animation/PoseSearch/Source/Runtime/Private/PoseSearchFeatureChannel_Position.cpp b/Engine/Plugins/Animation/PoseSearch/Source/Runtime/Private/PoseSearchFeatureChannel_Position.cpp
--- a/Engine/Plugins/Animation/PoseSearch/Source/Runtime/Private/PoseSearchFeatureChannel_Position.cpp
+++ b/Engine/Plugins/Animation/PoseSearch/Source/Runtime/Private/PoseSearchFeatureChannel_Position.cpp
@@ -69,7 +69,7 @@ void UPoseSearchFeatureChannel_Position::AddDependentChannels(UPoseSearchSchema*
 
 	if (Schema->bInjectAdditionalDebugChannels)
 	{
-		if (SchemaOriginBoneIdx != RootSchemaBoneIdx)
+		if (SchemaOriginBoneIdx != RootSchemaBoneIdx || PermutationTimeType == EPermutationTimeType::UsePermutationTime)
 		{
 			const EPermutationTimeType DependentChannelsPermutationTimeType = PermutationTimeType == EPermutationTimeType::UsePermutationTime ? EPermutationTimeType::UseSampleToPermutationTime : EPermutationTimeType::UseSampleTime;
 			UPoseSearchFeatureChannel_Position::FindOrAddToSchema(Schema, 0.f, OriginBone.BoneName, OriginRole, DependentChannelsPermutationTimeType);
@@ -159,6 +159,18 @@ void UPoseSearchFeatureChannel_Position::DebugDraw(const UE::PoseSearch::FDebugD
 {
 	using namespace UE::PoseSearch;
 
+#if WITH_EDITORONLY_DATA
+	bool bDrawInjectAdditionalDebugChannels = false;
+	if (const UPoseSearchSchema* Schema = GetSchema())
+	{
+		bDrawInjectAdditionalDebugChannels = Schema->bDrawInjectAdditionalDebugChannels;
+	}
+	if (!bDrawInjectAdditionalDebugChannels && Weight <= 0.f)
+	{
+		return;
+	}
+#endif // WITH_EDITORONLY_DATA
+
 	FColor Color;
 #if WITH_EDITORONLY_DATA
 	Color = DebugColor.ToFColor(true);
@@ -166,10 +178,14 @@ void UPoseSearchFeatureChannel_Position::DebugDraw(const UE::PoseSearch::FDebugD
 	Color = FLinearColor::Blue.ToFColor(true);
 #endif // WITH_EDITORONLY_DATA
 
+	float PermutationSampleTimeOffset = 0.f;
+	float PermutationOriginTimeOffset = 0.f;
+	UPoseSearchFeatureChannel::GetPermutationTimeOffsets(PermutationTimeType, DrawParams.ExtractPermutationTime(PoseVector), PermutationSampleTimeOffset, PermutationOriginTimeOffset);
+	const EPermutationTimeType OriginPermutationTimeType = PermutationTimeType == EPermutationTimeType::UsePermutationTime ? EPermutationTimeType::UseSampleToPermutationTime : EPermutationTimeType::UseSampleTime;
+
 	const FVector FeaturesVector = FFeatureVectorHelper::DecodeVector(PoseVector, ChannelDataOffset, ComponentStripping);
-	const EPermutationTimeType TimeType = PermutationTimeType == EPermutationTimeType::UsePermutationTime ? EPermutationTimeType::UseSampleToPermutationTime : EPermutationTimeType::UseSampleTime;
-	const FVector OriginBonePos = DrawParams.ExtractPosition(PoseVector, OriginTimeOffset, SchemaOriginBoneIdx, OriginRole, TimeType);
-	const FVector DeltaPos = DrawParams.GetRootBoneTransform(OriginRole).TransformVector(FeaturesVector);
+	const FVector OriginBonePos = DrawParams.ExtractPosition(PoseVector, OriginTimeOffset, SchemaOriginBoneIdx, OriginRole, OriginPermutationTimeType, INDEX_NONE, PermutationOriginTimeOffset);
+	const FVector DeltaPos = DrawParams.ExtractRotation(PoseVector, OriginTimeOffset, RootSchemaBoneIdx, OriginRole, OriginPermutationTimeType, INDEX_NONE, PermutationOriginTimeOffset).RotateVector(FeaturesVector);
 	const FVector BonePos = OriginBonePos + DeltaPos;
 	DrawParams.DrawPoint(BonePos, Color);
 
diff --git a/Engine/Plugins/Animation/PoseSearch/Source/Runtime/Private/PoseSearchFeatureChannel_Velocity.cpp b/Engine/Plugins/Animation/PoseSearch/Source/Runtime/Private/PoseSearchFeatureChannel_Velocity.cpp
--- a/Engine/Plugins/Animation/PoseSearch/Source/Runtime/Private/PoseSearchFeatureChannel_Velocity.cpp
+++ b/Engine/Plugins/Animation/PoseSearch/Source/Runtime/Private/PoseSearchFeatureChannel_Velocity.cpp
@@ -35,7 +35,8 @@ void UPoseSearchFeatureChannel_Velocity::AddDependentChannels(UPoseSearchSchema*
 {
 	if (Schema->bInjectAdditionalDebugChannels)
 	{
-		UPoseSearchFeatureChannel_Position::FindOrAddToSchema(Schema, SampleTimeOffset, Bone.BoneName, SampleRole);
+		const EPermutationTimeType DependentChannelsPermutationTimeType = PermutationTimeType != EPermutationTimeType::UseSampleTime ? EPermutationTimeType::UseSampleToPermutationTime : EPermutationTimeType::UseSampleTime;
+		UPoseSearchFeatureChannel_Position::FindOrAddToSchema(Schema, SampleTimeOffset, Bone.BoneName, SampleRole, DependentChannelsPermutationTimeType);
 	}
 }
 
@@ -143,9 +144,14 @@ void UPoseSearchFeatureChannel_Velocity::DebugDraw(const UE::PoseSearch::FDebugD
 
 	const float LinearVelocityScale = bNormalize ? 15.f : 0.08f;
 
-	const FVector LinearVelocity = DrawParams.GetRootBoneTransform(SampleRole).TransformVector(FFeatureVectorHelper::DecodeVector(PoseVector, ChannelDataOffset, ComponentStripping));
-	const FVector BoneVelDirection = LinearVelocity.GetSafeNormal();
-	const FVector BonePos = DrawParams.ExtractPosition(PoseVector, SampleTimeOffset, SchemaBoneIdx, SampleRole, PermutationTimeType, SamplingAttributeId);
+	float PermutationSampleTimeOffset = 0.f;
+	float PermutationOriginTimeOffset = 0.f;
+	UPoseSearchFeatureChannel::GetPermutationTimeOffsets(PermutationTimeType, DrawParams.ExtractPermutationTime(PoseVector), PermutationSampleTimeOffset, PermutationOriginTimeOffset);
+	const EPermutationTimeType SamplePermutationTimeType = PermutationTimeType == EPermutationTimeType::UsePermutationTime ? EPermutationTimeType::UseSampleToPermutationTime : EPermutationTimeType::UseSampleTime;
+	
+	const FVector FeaturesVector = FFeatureVectorHelper::DecodeVector(PoseVector, ChannelDataOffset, ComponentStripping);
+	const FVector LinearVelocity = DrawParams.ExtractRotation(PoseVector, SampleTimeOffset, RootSchemaBoneIdx, SampleRole, SamplePermutationTimeType, SamplingAttributeId, PermutationSampleTimeOffset).RotateVector(FeaturesVector);
+	const FVector BonePos = DrawParams.ExtractPosition(PoseVector, SampleTimeOffset, SchemaBoneIdx, SampleRole, SamplePermutationTimeType, SamplingAttributeId, PermutationSampleTimeOffset);
 
 	DrawParams.DrawLine(BonePos, BonePos + LinearVelocity * LinearVelocityScale, Color);
 }
diff --git a/Engine/Plugins/Animation/PoseSearch/Source/Runtime/Private/PoseSearchSchema.cpp b/Engine/Plugins/Animation/PoseSearch/Source/Runtime/Private/PoseSearchSchema.cpp
--- a/Engine/Plugins/Animation/PoseSearch/Source/Runtime/Private/PoseSearchSchema.cpp
+++ b/Engine/Plugins/Animation/PoseSearch/Source/Runtime/Private/PoseSearchSchema.cpp
@@ -1,6 +1,7 @@
 // Copyright Epic Games, Inc. All Rights Reserved.
 
 #include "PoseSearch/PoseSearchSchema.h"
+#include "Animation/MirrorDataTable.h"
 #include "AnimationRuntime.h"
 #include "PoseSearch/PoseSearchContext.h"
 #include "PoseSearch/PoseSearchDefines.h"
@@ -305,7 +306,16 @@ void UPoseSearchSchema::Finalize()
 		for (FBoneReference& BoneRef : RoledSkeleton.BoneReferences)
 		{
 			check(BoneRef.HasValidSetup());
-			RoledSkeleton.BoneIndicesWithParents.Add(BoneRef.BoneIndex);
+			RoledSkeleton.BoneIndicesWithParents.AddUnique(BoneRef.BoneIndex);
+
+			if (RoledSkeleton.MirrorDataTable)
+			{
+				const FSkeletonPoseBoneIndex MirroredBoneIndex = RoledSkeleton.MirrorDataTable->BoneToMirrorBoneIndex[BoneRef.BoneIndex];
+				if (MirroredBoneIndex.IsValid())
+				{
+					RoledSkeleton.BoneIndicesWithParents.AddUnique(MirroredBoneIndex.GetInt());
+				}
+			}
 		}
 
 		// Build separate index array with parent indices guaranteed to be present. Sort for EnsureParentsPresent.
diff --git a/Engine/Plugins/Animation/PoseSearch/Source/Runtime/Public/PoseSearch/PoseSearchAssetIndexer.h b/Engine/Plugins/Animation/PoseSearch/Source/Runtime/Public/PoseSearch/PoseSearchAssetIndexer.h
--- a/Engine/Plugins/Animation/PoseSearch/Source/Runtime/Public/PoseSearch/PoseSearchAssetIndexer.h
+++ b/Engine/Plugins/Animation/PoseSearch/Source/Runtime/Public/PoseSearch/PoseSearchAssetIndexer.h
@@ -126,6 +126,7 @@ private:
 		float SampleTime = 0.f;
 		bool bClamped = false;
 
+		// RootTransform and ComponentSpacePose are stored mirrored in case SearchIndexAsset.IsMirrored
 		TArray<FTransform> RootTransform;
 		TArray<FCSPose<FCompactHeapPose>> ComponentSpacePose;
 	};
diff --git a/Engine/Plugins/Animation/PoseSearch/Source/Runtime/Public/PoseSearch/PoseSearchAssetSampler.h b/Engine/Plugins/Animation/PoseSearch/Source/Runtime/Public/PoseSearch/PoseSearchAssetSampler.h
--- a/Engine/Plugins/Animation/PoseSearch/Source/Runtime/Public/PoseSearch/PoseSearchAssetSampler.h
+++ b/Engine/Plugins/Animation/PoseSearch/Source/Runtime/Public/PoseSearch/PoseSearchAssetSampler.h
@@ -51,6 +51,8 @@ struct POSESEARCH_API FAnimationAssetSampler
 
 	static float GetPlayLength(const UAnimationAsset* AnimAsset, const FVector& BlendParameters = FVector::ZeroVector);
 
+	void SetRootTransformOrigin(const FTransform& InRootTransformOrigin) { RootTransformOrigin = InRootTransformOrigin; }
+
 protected:
 	TWeakObjectPtr<const UAnimationAsset> AnimationAssetPtr;
 	FTransform RootTransformOrigin = FTransform::Identity;
diff --git a/Engine/Plugins/Animation/PoseSearch/Source/Runtime/Public/PoseSearch/PoseSearchContext.h b/Engine/Plugins/Animation/PoseSearch/Source/Runtime/Public/PoseSearch/PoseSearchContext.h
--- a/Engine/Plugins/Animation/PoseSearch/Source/Runtime/Public/PoseSearch/PoseSearchContext.h
+++ b/Engine/Plugins/Animation/PoseSearch/Source/Runtime/Public/PoseSearch/PoseSearchContext.h
@@ -58,10 +58,9 @@ struct POSESEARCH_API FDebugDrawParams
 	const FSearchIndex* GetSearchIndex() const;
 	const UPoseSearchSchema* GetSchema() const;
 
-	FVector ExtractPosition(TConstArrayView<float> PoseVector, float SampleTimeOffset, int8 SchemaBoneIdx, const FRole& Role, EPermutationTimeType PermutationTimeType = EPermutationTimeType::UseSampleTime, int32 SamplingAttributeId = INDEX_NONE) const;
-
-	FQuat ExtractRotation(TConstArrayView<float> PoseVector, float SampleTimeOffset, int8 SchemaBoneIdx, const FRole& Role, EPermutationTimeType PermutationTimeType = EPermutationTimeType::UseSampleTime, int32 SamplingAttributeId = INDEX_NONE) const;
-
+	float ExtractPermutationTime(TConstArrayView<float> PoseVector) const;
+	FVector ExtractPosition(TConstArrayView<float> PoseVector, float SampleTimeOffset, int8 SchemaBoneIdx, const FRole& Role, EPermutationTimeType PermutationTimeType = EPermutationTimeType::UseSampleTime, int32 SamplingAttributeId = INDEX_NONE, float PermutationSampleTimeOffset = 0.f) const;
+	FQuat ExtractRotation(TConstArrayView<float> PoseVector, float SampleTimeOffset, int8 SchemaBoneIdx, const FRole& Role, EPermutationTimeType PermutationTimeType = EPermutationTimeType::UseSampleTime, int32 SamplingAttributeId = INDEX_NONE, float PermutationSampleTimeOffset = 0.f) const;
 	FTransform GetRootBoneTransform(const FRole& Role, float SampleTimeOffset = 0.f) const;
 
 	void DrawLine(const FVector& LineStart, const FVector& LineEnd, const FColor& Color, float Thickness = 0.f) const;
diff --git a/Engine/Plugins/Animation/PoseSearch/Source/Runtime/Public/PoseSearch/PoseSearchDatabase.h b/Engine/Plugins/Animation/PoseSearch/Source/Runtime/Public/PoseSearch/PoseSearchDatabase.h
--- a/Engine/Plugins/Animation/PoseSearch/Source/Runtime/Public/PoseSearch/PoseSearchDatabase.h
+++ b/Engine/Plugins/Animation/PoseSearch/Source/Runtime/Public/PoseSearch/PoseSearchDatabase.h
@@ -69,7 +69,7 @@ struct POSESEARCH_API FPoseSearchDatabaseAnimationAssetBase
 	virtual void SetDisableReselection(bool bValue) { bDisableReselection = bValue; }
 	virtual UClass* GetAnimationAssetStaticClass() const { return nullptr; }
 	virtual bool IsLooping() const { return false; }
-	virtual const FString GetName() const { return FString(); }
+	virtual const FString GetName() const;
 	virtual bool IsEnabled() const { return bEnabled; }
 	virtual void SetIsEnabled(bool bValue) { bEnabled = bValue; }
 	virtual bool IsRootMotionEnabled() const { return false; }
@@ -77,7 +77,7 @@ struct POSESEARCH_API FPoseSearchDatabaseAnimationAssetBase
 
 	// [0, 0] represents the entire frame range of the original animation.
 	virtual FFloatInterval GetSamplingRange() const { return FFloatInterval(0.f, 0.f); }
-	static FFloatInterval GetEffectiveSamplingRange(const UAnimSequenceBase* Sequence, const FFloatInterval& RequestedSamplingRange);
+	FFloatInterval GetEffectiveSamplingRange() const;
 
 	virtual int64 GetEditorMemSize() const;
 	virtual int64 GetApproxCookedSize() const { return GetEditorMemSize(); }
@@ -121,7 +121,6 @@ struct POSESEARCH_API FPoseSearchDatabaseSequence : public FPoseSearchDatabaseAn
 
 	virtual UClass* GetAnimationAssetStaticClass() const override;
 	virtual bool IsLooping() const override;
-	virtual const FString GetName() const override;
 	virtual bool IsRootMotionEnabled() const override;
 	virtual FFloatInterval GetSamplingRange() const override { return SamplingRange; }
 #endif // WITH_EDITORONLY_DATA
@@ -167,7 +166,6 @@ struct POSESEARCH_API FPoseSearchDatabaseBlendSpace : public FPoseSearchDatabase
 
 	virtual UClass* GetAnimationAssetStaticClass() const override;
 	virtual bool IsLooping() const override;
-	virtual const FString GetName() const override;
 	virtual bool IsRootMotionEnabled() const override;
 
 	void GetBlendSpaceParameterSampleRanges(int32& HorizontalBlendNum, int32& VerticalBlendNum) const;
@@ -200,7 +198,6 @@ struct POSESEARCH_API FPoseSearchDatabaseAnimComposite : public FPoseSearchDatab
 
 	virtual UClass* GetAnimationAssetStaticClass() const override;
 	virtual bool IsLooping() const override;
-	virtual const FString GetName() const override;
 	virtual bool IsRootMotionEnabled() const override;
 	virtual FFloatInterval GetSamplingRange() const override { return SamplingRange; }
 #endif // WITH_EDITORONLY_DATA
@@ -227,7 +224,6 @@ struct POSESEARCH_API FPoseSearchDatabaseAnimMontage : public FPoseSearchDatabas
 
 	virtual UClass* GetAnimationAssetStaticClass() const override;
 	virtual bool IsLooping() const override;
-	virtual const FString GetName() const override;
 	virtual bool IsRootMotionEnabled() const override;
 	virtual FFloatInterval GetSamplingRange() const override { return SamplingRange; }
 #endif // WITH_EDITORONLY_DATA
@@ -253,7 +249,6 @@ struct POSESEARCH_API FPoseSearchDatabaseMultiSequence : public FPoseSearchDatab
 
 	virtual UClass* GetAnimationAssetStaticClass() const override;
 	virtual bool IsLooping() const override;
-	virtual const FString GetName() const override;
 	virtual bool IsRootMotionEnabled() const override;
 	virtual FFloatInterval GetSamplingRange() const override { return SamplingRange; }
 #endif // WITH_EDITORONLY_DATA
diff --git a/Engine/Plugins/Animation/PoseSearch/Source/Runtime/Public/PoseSearch/PoseSearchSchema.h b/Engine/Plugins/Animation/PoseSearch/Source/Runtime/Public/PoseSearch/PoseSearchSchema.h
--- a/Engine/Plugins/Animation/PoseSearch/Source/Runtime/Public/PoseSearch/PoseSearchSchema.h
+++ b/Engine/Plugins/Animation/PoseSearch/Source/Runtime/Public/PoseSearch/PoseSearchSchema.h
@@ -114,9 +114,14 @@ public:
 	// the original intent is to add UPoseSearchFeatureChannel_Position(s) to help with the complexity of the debug drawing
 	// (the database will have all the necessary positions to draw lines at the right location and time).
 	UPROPERTY(EditAnywhere, Category = "Debug")
-	bool bInjectAdditionalDebugChannels;
+	bool bInjectAdditionalDebugChannels = false;
 
-	//bool IsValid () const;
+#if WITH_EDITORONLY_DATA
+	// if bDrawInjectAdditionalDebugChannels is true, all the channels added for debug purposes with 
+	// bInjectAdditionalDebugChannels (as well as all those channels with an associated zero weight) will be drawn
+	UPROPERTY(EditAnywhere, Category = "Debug")
+	bool bDrawInjectAdditionalDebugChannels = false;
+#endif // WITH_EDITORONLY_DATA
 
 	TConstArrayView<TObjectPtr<UPoseSearchFeatureChannel>> GetChannels() const { return FinalizedChannels; }
 
