motion matching - multicharacter improvements

diff --git a/Engine/Plugins/Animation/PoseSearch/Source/Runtime/Private/PoseSearchLibrary.cpp b/Engine/Plugins/Animation/PoseSearch/Source/Runtime/Private/PoseSearchLibrary.cpp
--- a/Engine/Plugins/Animation/PoseSearch/Source/Runtime/Private/PoseSearchLibrary.cpp
+++ b/Engine/Plugins/Animation/PoseSearch/Source/Runtime/Private/PoseSearchLibrary.cpp
@@ -108,6 +108,47 @@ namespace UE::PoseSearch
 		return false;
 	}
 
+	static const FAnimNode_PoseSearchHistoryCollector_Base* FindPoseHistoryNode(const FName PoseHistoryName, const UAnimInstance* AnimInstance)
+	{
+		if (AnimInstance)
+		{
+			TSet<const UAnimInstance*, DefaultKeyFuncs<const UAnimInstance*>, TInlineSetAllocator<128>> AlreadyVisited;
+			TArray<const UAnimInstance*, TInlineAllocator<128>> ToVisit;
+
+			ToVisit.Add(AnimInstance);
+			AlreadyVisited.Add(AnimInstance);
+
+			while (!ToVisit.IsEmpty())
+			{
+				const UAnimInstance* Visiting = ToVisit.Pop();
+
+				if (IAnimClassInterface* AnimBlueprintClass = IAnimClassInterface::GetFromClass(Visiting->GetClass()))
+				{
+					if (const FAnimSubsystem_Tag* TagSubsystem = AnimBlueprintClass->FindSubsystem<FAnimSubsystem_Tag>())
+					{
+						if (const FAnimNode_PoseSearchHistoryCollector_Base* HistoryCollector = TagSubsystem->FindNodeByTag<FAnimNode_PoseSearchHistoryCollector_Base>(PoseHistoryName, Visiting))
+						{
+							return HistoryCollector;
+						}
+					}
+				}
+
+				const USkeletalMeshComponent* SkeletalMeshComponent = Visiting->GetSkelMeshComponent();
+				const TArray<UAnimInstance*>& LinkedAnimInstances = SkeletalMeshComponent->GetLinkedAnimInstances();
+				for (const UAnimInstance* LinkedAnimInstance : LinkedAnimInstances)
+				{
+					bool bIsAlreadyInSet = false;
+					AlreadyVisited.Add(LinkedAnimInstance, &bIsAlreadyInSet);
+					if (!bIsAlreadyInSet)
+					{
+						ToVisit.Add(LinkedAnimInstance);
+					}
+				}
+			}
+		}
+		return nullptr;
+	}
+
 #if ENABLE_DRAW_DEBUG && ENABLE_ANIM_DEBUG
 	class UAnimInstanceProxyProvider : public UAnimInstance
 	{
@@ -547,6 +588,7 @@ void UPoseSearchLibrary::MotionMatch(
 	UAnimInstance* AnimInstance,
 	TArray<UObject*> AssetsToSearch,
 	const FName PoseHistoryName,
+	const FPoseSearchContinuingProperties ContinuingProperties,
 	FPoseSearchFutureProperties Future,
 	FPoseSearchBlueprintResult& Result,
 	const int32 DebugSessionUniqueIdentifier)
@@ -562,7 +604,7 @@ void UPoseSearchLibrary::MotionMatch(
 	Roles.Add(UE::PoseSearch::DefaultRole);
 
 	TArray<const UObject*>& AssetsToSearchConst = reinterpret_cast<TArray<const UObject*>&>(AssetsToSearch);
-	MotionMatch(AnimInstances, Roles, AssetsToSearchConst, PoseHistoryName, Future, Result, DebugSessionUniqueIdentifier);
+	MotionMatch(AnimInstances, Roles, AssetsToSearchConst, PoseHistoryName, ContinuingProperties, Future, Result, DebugSessionUniqueIdentifier);
 }
 
 void UPoseSearchLibrary::MotionMatchMulti(
@@ -570,6 +612,7 @@ void UPoseSearchLibrary::MotionMatchMulti(
 	TArray<FName> Roles,
 	TArray<UObject*> AssetsToSearch,
 	const FName PoseHistoryName,
+	const FPoseSearchContinuingProperties ContinuingProperties,
 	FPoseSearchBlueprintResult& Result,
 	const int32 DebugSessionUniqueIdentifier)
 {
@@ -589,7 +632,7 @@ void UPoseSearchLibrary::MotionMatchMulti(
 	}
 
 	TArray<const UObject*>& AssetsToSearchConst = reinterpret_cast<TArray<const UObject*>&>(AssetsToSearch);
-	MotionMatch(AnimInstances, Roles, AssetsToSearchConst, PoseHistoryName, FPoseSearchFutureProperties(), Result, DebugSessionUniqueIdentifier);
+	MotionMatch(AnimInstances, Roles, AssetsToSearchConst, PoseHistoryName, ContinuingProperties, FPoseSearchFutureProperties(), Result, DebugSessionUniqueIdentifier);
 }
 
 void UPoseSearchLibrary::MotionMatch(
@@ -597,6 +640,7 @@ void UPoseSearchLibrary::MotionMatch(
 	TArrayView<const UE::PoseSearch::FRole> Roles,
 	TArrayView<const UObject*> AssetsToSearch,
 	const FName PoseHistoryName,
+	const FPoseSearchContinuingProperties& ContinuingProperties,
 	const FPoseSearchFutureProperties& Future,
 	FPoseSearchBlueprintResult& Result,
 	const int32 DebugSessionUniqueIdentifier)
@@ -643,23 +687,17 @@ void UPoseSearchLibrary::MotionMatch(
 	TArray<FMemStackPoseHistory, TInlineAllocator<PreallocatedRolesNum, TMemStackAllocator<>>> MemStackPoseHistories;
 	for (UAnimInstance* AnimInstance : AnimInstances)
 	{
-		if (IAnimClassInterface* AnimBlueprintClass = IAnimClassInterface::GetFromClass(AnimInstance->GetClass()))
+		if (const FAnimNode_PoseSearchHistoryCollector_Base* PoseHistoryNode = FindPoseHistoryNode(PoseHistoryName, AnimInstance))
 		{
-			if (const FAnimSubsystem_Tag* TagSubsystem = AnimBlueprintClass->FindSubsystem<FAnimSubsystem_Tag>())
-			{
-				if (const FAnimNode_PoseSearchHistoryCollector_Base* PoseHistoryNode = TagSubsystem->FindNodeByTag<FAnimNode_PoseSearchHistoryCollector_Base>(PoseHistoryName, AnimInstance))
-				{
-					MemStackPoseHistories.AddDefaulted_GetRef().Init(&PoseHistoryNode->GetPoseHistory());
+			MemStackPoseHistories.AddDefaulted_GetRef().Init(&PoseHistoryNode->GetPoseHistory());
 
 #if ENABLE_DRAW_DEBUG && ENABLE_ANIM_DEBUG
-	#if WITH_EDITORONLY_DATA
-					HistoryCollectorColors.Emplace(PoseHistoryNode->DebugColor.ToFColor(true));
-	#else // WITH_EDITORONLY_DATA
-					HistoryCollectorColors.Emplace(FColor::Red);
-	#endif // WITH_EDITORONLY_DATA
+#if WITH_EDITORONLY_DATA
+			HistoryCollectorColors.Emplace(PoseHistoryNode->DebugColor.ToFColor(true));
+#else // WITH_EDITORONLY_DATA
+			HistoryCollectorColors.Emplace(FColor::Red);
+#endif // WITH_EDITORONLY_DATA
 #endif // ENABLE_DRAW_DEBUG && ENABLE_ANIM_DEBUG && WITH_EDITORONLY_DATA
-				}
-			}
 		}
 	}
 
@@ -733,7 +771,7 @@ void UPoseSearchLibrary::MotionMatch(
 		PoseHistories.Add(MemStackPoseHistory.GetThisOrPoseHistory());
 	}
 
-	const FSearchResult SearchResult = MotionMatch(AnimInstances, Roles, PoseHistories, AssetsToSearch, nullptr, 0.f, DebugSessionUniqueIdentifier, FutureIntervalTime);
+	const FSearchResult SearchResult = MotionMatch(AnimInstances, Roles, PoseHistories, AssetsToSearch, ContinuingProperties, DebugSessionUniqueIdentifier, FutureIntervalTime);
 	if (SearchResult.IsValid())
 	{
 		const UPoseSearchDatabase* Database = SearchResult.Database.Get();
@@ -770,8 +808,7 @@ void UPoseSearchLibrary::MotionMatch(
 UE::PoseSearch::FSearchResult UPoseSearchLibrary::MotionMatch(
 	const FAnimationBaseContext& Context,
 	TArrayView<const UObject*> AssetsToSearch,
-	const UObject* PlayingAsset,
-	float PlayingAssetAccumulatedTime)
+	const FPoseSearchContinuingProperties& ContinuingProperties)
 {
 	using namespace UE::PoseSearch;
 
@@ -784,7 +821,7 @@ UE::PoseSearch::FSearchResult UPoseSearchLibrary::MotionMatch(
 	UAnimInstance* AnimInstance = Cast<UAnimInstance>(Context.AnimInstanceProxy->GetAnimInstanceObject());
 	check(AnimInstance);
 
-	return MotionMatch(MakeArrayView(&AnimInstance, 1), MakeArrayView(&DefaultRole, 1), MakeArrayView(&PoseHistory, 1), AssetsToSearch, PlayingAsset, PlayingAssetAccumulatedTime, Context.GetCurrentNodeId());
+	return MotionMatch(MakeArrayView(&AnimInstance, 1), MakeArrayView(&DefaultRole, 1), MakeArrayView(&PoseHistory, 1), AssetsToSearch, ContinuingProperties, Context.GetCurrentNodeId());
 }
 	
 UE::PoseSearch::FSearchResult UPoseSearchLibrary::MotionMatch(
@@ -792,8 +829,7 @@ UE::PoseSearch::FSearchResult UPoseSearchLibrary::MotionMatch(
 	TArrayView<const UE::PoseSearch::FRole> Roles,
 	TArrayView<const UE::PoseSearch::IPoseHistory*> PoseHistories, 
 	TArrayView<const UObject*> AssetsToSearch,
-	const UObject* PlayingAsset,
-	float PlayingAssetAccumulatedTime,
+	const FPoseSearchContinuingProperties& ContinuingProperties,
 	const int32 DebugSessionUniqueIdentifier,
 	float DesiredPermutationTimeOffset)
 {
@@ -869,9 +905,21 @@ UE::PoseSearch::FSearchResult UPoseSearchLibrary::MotionMatch(
 
 	// collecting all the possible continuing pose search (it could be multiple searches, but most likely only one)
 	const float DeltaSeconds = AnimInstances[0]->GetDeltaSeconds();
-	if (const UAnimationAsset* PlayingAnimationAsset = Cast<UAnimationAsset>(PlayingAsset))
+	if (const UAnimationAsset* PlayingAnimationAsset = Cast<UAnimationAsset>(ContinuingProperties.PlayingAsset))
 	{
 		AddToSearch(AssetsToSearchPerDatabaseMap, PlayingAnimationAsset);
+		
+		for (const UObject* AssetToSearch : AssetsToSearch)
+		{
+			if (const UPoseSearchDatabase* Database = Cast<UPoseSearchDatabase>(AssetToSearch))
+			{
+				if (Database->Contains(PlayingAnimationAsset))
+				{
+					AddToSearch(AssetsToSearchPerDatabaseMap, AssetToSearch);
+				}
+			}
+		}
+
 		for (const TAssetsToSearchPerDatabasePair& AssetsToSearchPerDatabasePair : AssetsToSearchPerDatabaseMap)
 		{
 			const UPoseSearchDatabase* Database = AssetsToSearchPerDatabasePair.Key;
@@ -901,14 +949,14 @@ UE::PoseSearch::FSearchResult UPoseSearchLibrary::MotionMatch(
 								const float LastSampleTime = SearchIndexAsset.GetLastSampleTime(Database->Schema->SampleRate) - DeltaSeconds;
 								const float MaxTimeToBeAbleToContinuingPlayingAnimation = LastSampleTime - DeltaSeconds;
 
-								bCanAdvance = PlayingAssetAccumulatedTime >= FirstSampleTime && PlayingAssetAccumulatedTime < MaxTimeToBeAbleToContinuingPlayingAnimation;
+								bCanAdvance = ContinuingProperties.PlayingAssetAccumulatedTime >= FirstSampleTime && ContinuingProperties.PlayingAssetAccumulatedTime < MaxTimeToBeAbleToContinuingPlayingAnimation;
 							}
 
 							if (bCanAdvance)
 							{
 								ReconstructedPreviousSearchResult.Database = Database;
-								ReconstructedPreviousSearchResult.AssetTime = PlayingAssetAccumulatedTime;
-								ReconstructedPreviousSearchResult.PoseIdx = Database->GetPoseIndexFromTime(PlayingAssetAccumulatedTime, SearchIndexAsset);
+								ReconstructedPreviousSearchResult.AssetTime = ContinuingProperties.PlayingAssetAccumulatedTime;
+								ReconstructedPreviousSearchResult.PoseIdx = Database->GetPoseIndexFromTime(ContinuingProperties.PlayingAssetAccumulatedTime, SearchIndexAsset);
 								SearchContext.UpdateCurrentResultPoseVector();
 
 								const FSearchResult NewSearchResult = Database->SearchContinuingPose(SearchContext);
diff --git a/Engine/Plugins/Animation/PoseSearch/Source/Runtime/Private/PoseSearchModule.cpp b/Engine/Plugins/Animation/PoseSearch/Source/Runtime/Private/PoseSearchModule.cpp
--- a/Engine/Plugins/Animation/PoseSearch/Source/Runtime/Private/PoseSearchModule.cpp
+++ b/Engine/Plugins/Animation/PoseSearch/Source/Runtime/Private/PoseSearchModule.cpp
@@ -24,7 +24,10 @@ public:
 	// IPoseSearchProvider
 	virtual UE::Anim::IPoseSearchProvider::FSearchResult Search(const FAnimationBaseContext& GraphContext, TArrayView<const UObject*> AssetsToSearch, const UObject* PlayingAsset, float PlayingAssetAccumulatedTime) const override
 	{
-		const UE::PoseSearch::FSearchResult SearchResult = UPoseSearchLibrary::MotionMatch(GraphContext, AssetsToSearch, PlayingAsset, PlayingAssetAccumulatedTime);
+		FPoseSearchContinuingProperties ContinuingProperties;
+		ContinuingProperties.PlayingAsset = PlayingAsset;
+		ContinuingProperties.PlayingAssetAccumulatedTime = PlayingAssetAccumulatedTime;
+		const UE::PoseSearch::FSearchResult SearchResult = UPoseSearchLibrary::MotionMatch(GraphContext, AssetsToSearch, ContinuingProperties);
 		UE::Anim::IPoseSearchProvider::FSearchResult ProviderResult;
 		if (const UE::PoseSearch::FSearchIndexAsset* SearchIndexAsset = SearchResult.GetSearchIndexAsset())
 		{
diff --git a/Engine/Plugins/Animation/PoseSearch/Source/Runtime/Private/PoseSearchMultiSequence.cpp b/Engine/Plugins/Animation/PoseSearch/Source/Runtime/Private/PoseSearchMultiSequence.cpp
--- a/Engine/Plugins/Animation/PoseSearch/Source/Runtime/Private/PoseSearchMultiSequence.cpp
+++ b/Engine/Plugins/Animation/PoseSearch/Source/Runtime/Private/PoseSearchMultiSequence.cpp
@@ -8,18 +8,18 @@ bool UPoseSearchMultiSequence::IsLooping() const
 	float CommonPlayLength = -1.f;
 	for (const FPoseSearchMultiSequenceItem& Item : Items)
 	{
-		if (Item.Sequence)
+		if (const UAnimSequenceBase* Sequence = Item.Sequence.Get())
 		{
-			if (!Item.Sequence->bLoop)
+			if (!Sequence->bLoop)
 			{
 				return false;
 			}
 
 			if (CommonPlayLength < 0.f)
 			{
-				CommonPlayLength = Item.Sequence->GetPlayLength();
+				CommonPlayLength = Sequence->GetPlayLength();
 			}
-			else if (!FMath::IsNearlyEqual(CommonPlayLength, Item.Sequence->GetPlayLength()))
+			else if (!FMath::IsNearlyEqual(CommonPlayLength, Sequence->GetPlayLength()))
 			{
 				return false;
 			}
@@ -28,30 +28,6 @@ bool UPoseSearchMultiSequence::IsLooping() const
 	return true;
 }
 
-const FString UPoseSearchMultiSequence::GetName() const
-{
-	FString Name;
-	Name.Reserve(256);
-	bool bAddComma = false;
-	for (const FPoseSearchMultiSequenceItem& Item : Items)
-	{
-		if (bAddComma)
-		{
-			Name += ", ";
-		}
-		else
-		{
-			bAddComma = true;
-		}
-
-		Name += "[";
-		Name += Item.Role.ToString();
-		Name += "] ";
-		Name += GetNameSafe(Item.Sequence);
-	}
-	return Name;
-}
-
 bool UPoseSearchMultiSequence::HasRootMotion() const
 {
 	bool bHasAtLeastOneValidItem = false;
@@ -59,9 +35,9 @@ bool UPoseSearchMultiSequence::HasRootMotion() const
 
 	for (const FPoseSearchMultiSequenceItem& Item : Items)
 	{
-		if (Item.Sequence)
+		if (const UAnimSequenceBase* Sequence = Item.Sequence.Get())
 		{
-			bHasRootMotion &= Item.Sequence->HasRootMotion();
+			bHasRootMotion &= Sequence->HasRootMotion();
 			bHasAtLeastOneValidItem = true;
 		}
 	}
@@ -71,27 +47,40 @@ bool UPoseSearchMultiSequence::HasRootMotion() const
 
 float UPoseSearchMultiSequence::GetPlayLength() const
 {
-	float PlayLength = 0.f;
+	float MaxPlayLength = 0.f;
 	for (const FPoseSearchMultiSequenceItem& Item : Items)
 	{
-		if (Item.Sequence)
+		if (const UAnimSequenceBase* Sequence = Item.Sequence.Get())
 		{
-			PlayLength = FMath::Max(PlayLength, Item.Sequence->GetPlayLength());
+			MaxPlayLength = FMath::Max(MaxPlayLength, Sequence->GetPlayLength());
 		}
 	}
-	return PlayLength;
+	return MaxPlayLength;
 }
 
 #if WITH_EDITOR
 int32 UPoseSearchMultiSequence::GetFrameAtTime(float Time) const
 {
+	const UAnimSequenceBase* MaxPlayLengthAnim = nullptr;
+	float MaxPlayLength = -1.f;
 	for (const FPoseSearchMultiSequenceItem& Item : Items)
 	{
-		if (Item.Sequence)
+		if (const UAnimSequenceBase* Sequence = Item.Sequence.Get())
 		{
-			return Item.Sequence->GetFrameAtTime(Time);
+			const float PlayLength = Sequence->GetPlayLength();
+			if (PlayLength > MaxPlayLength)
+			{
+				MaxPlayLength = PlayLength;
+				MaxPlayLengthAnim = Sequence;
+			}
 		}
 	}
+	
+	if (MaxPlayLengthAnim)
+	{
+		return MaxPlayLengthAnim->GetFrameAtTime(Time);
+	}
+
 	return 0;
 }
 #endif // WITH_EDITOR
diff --git a/Engine/Plugins/Animation/PoseSearch/Source/Runtime/Public/PoseSearch/PoseSearchLibrary.h b/Engine/Plugins/Animation/PoseSearch/Source/Runtime/Public/PoseSearch/PoseSearchLibrary.h
--- a/Engine/Plugins/Animation/PoseSearch/Source/Runtime/Public/PoseSearch/PoseSearchLibrary.h
+++ b/Engine/Plugins/Animation/PoseSearch/Source/Runtime/Public/PoseSearch/PoseSearchLibrary.h
@@ -105,6 +105,21 @@ public:
 	float IntervalTime = 0.f;
 };
 
+USTRUCT(Experimental, BlueprintType, Category="Animation|Pose Search")
+struct POSESEARCH_API FPoseSearchContinuingProperties
+{
+	GENERATED_BODY()
+
+public:
+	// Currently playing animation
+	UPROPERTY(EditAnywhere, BlueprintReadWrite, Category=State)
+	TObjectPtr<const UObject> PlayingAsset = nullptr;
+
+	// Currently playing animation accumulated time
+	UPROPERTY(EditAnywhere, BlueprintReadWrite, Category=State)
+	float PlayingAssetAccumulatedTime = 0.f;
+};
+
 UCLASS()
 class POSESEARCH_API UPoseSearchLibrary : public UBlueprintFunctionLibrary
 {
@@ -172,6 +187,7 @@ public:
 		UAnimInstance* AnimInstance,
 		TArray<UObject*> AssetsToSearch,
 		const FName PoseHistoryName,
+		const FPoseSearchContinuingProperties ContinuingProperties,
 		FPoseSearchFutureProperties Future,
 		FPoseSearchBlueprintResult& Result,
 		const int32 DebugSessionUniqueIdentifier = 6174);
@@ -188,10 +204,11 @@ public:
 	*/
 	UFUNCTION(BlueprintPure, Category = "Animation|Pose Search|Experimental", meta = (BlueprintThreadSafe, Keywords = "PoseMatch"))
 	static void MotionMatchMulti(
-		TArray<ACharacter*> AnimInstances,
+		TArray<ACharacter*> Characters,
 		TArray<FName> Roles,
 		TArray<UObject*> AssetsToSearch,
 		const FName PoseHistoryName,
+		const FPoseSearchContinuingProperties ContinuingProperties,
 		FPoseSearchBlueprintResult& Result,
 		const int32 DebugSessionUniqueIdentifier = 6174);
 
@@ -200,6 +217,7 @@ public:
 		TArrayView<const UE::PoseSearch::FRole> Roles,
 		TArrayView<const UObject*> AssetsToSearch,
 		const FName PoseHistoryName,
+		const FPoseSearchContinuingProperties& ContinuingProperties,
 		const FPoseSearchFutureProperties& Future,
 		FPoseSearchBlueprintResult& Result,
 		const int32 DebugSessionUniqueIdentifier);
@@ -207,16 +225,14 @@ public:
 	static UE::PoseSearch::FSearchResult MotionMatch(
 		const FAnimationBaseContext& Context,
 		TArrayView<const UObject*> AssetsToSearch,
-		const UObject* PlayingAsset = nullptr,
-		float PlayingAssetAccumulatedTime = 0.f);
+		const FPoseSearchContinuingProperties& ContinuingProperties);
 
 	static UE::PoseSearch::FSearchResult MotionMatch(
 		TArrayView<UAnimInstance*> AnimInstances,
 		TArrayView<const UE::PoseSearch::FRole> Roles,
 		TArrayView<const UE::PoseSearch::IPoseHistory*> PoseHistories, 
 		TArrayView<const UObject*> AssetsToSearch,
-		const UObject* PlayingAsset,
-		float PlayingAssetAccumulatedTime,
+		const FPoseSearchContinuingProperties& ContinuingProperties,
 		const int32 DebugSessionUniqueIdentifier,
 		float DesiredPermutationTimeOffset = 0.f);
 };
diff --git a/Engine/Plugins/Animation/PoseSearch/Source/Runtime/Public/PoseSearch/PoseSearchMultiSequence.h b/Engine/Plugins/Animation/PoseSearch/Source/Runtime/Public/PoseSearch/PoseSearchMultiSequence.h
--- a/Engine/Plugins/Animation/PoseSearch/Source/Runtime/Public/PoseSearch/PoseSearchMultiSequence.h
+++ b/Engine/Plugins/Animation/PoseSearch/Source/Runtime/Public/PoseSearch/PoseSearchMultiSequence.h
@@ -33,7 +33,6 @@ class POSESEARCH_API UPoseSearchMultiSequence : public UDataAsset
 
 public:
 	bool IsLooping() const;
-	const FString GetName() const;
 	bool HasRootMotion() const;
 	float GetPlayLength() const;
 
