motion matching - fix for UPoseSearchLibrary::UpdateMotionMatchingState log formatting

diff --git a/Engine/Plugins/Animation/PoseSearch/Source/Runtime/Private/PoseSearchLibrary.cpp b/Engine/Plugins/Animation/PoseSearch/Source/Runtime/Private/PoseSearchLibrary.cpp
--- a/Engine/Plugins/Animation/PoseSearch/Source/Runtime/Private/PoseSearchLibrary.cpp
+++ b/Engine/Plugins/Animation/PoseSearch/Source/Runtime/Private/PoseSearchLibrary.cpp
@@ -500,7 +500,7 @@ void UPoseSearchLibrary::UpdateMotionMatchingState(
 
 #if WITH_EDITOR
 			StringBuilder << SearchContext.IsAsyncBuildIndexInProgress();
-//#else // WITH_EDITOR
+#else // WITH_EDITOR
 			StringBuilder << false;
 #endif // WITH_EDITOR
 
