motion matching - renaming config to schema in the database factory

diff --git a/Engine/Plugins/Animation/PoseSearch/Source/Editor/Private/PoseSearchDatabaseFactory.cpp b/Engine/Plugins/Animation/PoseSearch/Source/Editor/Private/PoseSearchDatabaseFactory.cpp
--- a/Engine/Plugins/Animation/PoseSearch/Source/Editor/Private/PoseSearchDatabaseFactory.cpp
+++ b/Engine/Plugins/Animation/PoseSearch/Source/Editor/Private/PoseSearchDatabaseFactory.cpp
@@ -42,7 +42,7 @@ bool UPoseSearchDatabaseFactory::ConfigureProperties()
 	AssetPickerConfig.InitialAssetViewType = EAssetViewType::List;
 
 	PickerWindow = SNew(SWindow)
-	.Title(LOCTEXT("CreatePoseSearchDatabaseOptions", "Pick Config"))
+	.Title(LOCTEXT("CreatePoseSearchDatabaseOptions", "Pick Schema"))
 	.ClientSize(FVector2D(500, 600))
 	.SupportsMinimize(false)
 	.SupportsMaximize(false)
