motion matching - removing unnecessary check in FAssetSamplingContext::FAssetSamplingContext, that is actually crashing the editor when indexing a database referencing a normalization set containing invalid databases (with null schema)

diff --git a/Engine/Plugins/Animation/PoseSearch/Source/Runtime/Private/PoseSearchAssetIndexer.cpp b/Engine/Plugins/Animation/PoseSearch/Source/Runtime/Private/PoseSearchAssetIndexer.cpp
--- a/Engine/Plugins/Animation/PoseSearch/Source/Runtime/Private/PoseSearchAssetIndexer.cpp
+++ b/Engine/Plugins/Animation/PoseSearch/Source/Runtime/Private/PoseSearchAssetIndexer.cpp
@@ -83,7 +83,6 @@ static FSamplingParam WrapOrClampSamplingParam(bool bCanWrap, float SamplingPara
 // FAssetSamplingContext
 FAssetSamplingContext::FAssetSamplingContext(const UPoseSearchDatabase& Database)
 {
-	check(Database.Schema);
 	BaseCostBias = Database.BaseCostBias;
 	LoopingCostBias = Database.LoopingCostBias;
 }
