motion matching - auto reindex cancelled databases (enabled if CVarMotionMatchReindexCancelledDatabases is true)

diff --git a/Engine/Plugins/Animation/PoseSearch/Source/Runtime/Private/PoseSearchDerivedData.cpp b/Engine/Plugins/Animation/PoseSearch/Source/Runtime/Private/PoseSearchDerivedData.cpp
--- a/Engine/Plugins/Animation/PoseSearch/Source/Runtime/Private/PoseSearchDerivedData.cpp
+++ b/Engine/Plugins/Animation/PoseSearch/Source/Runtime/Private/PoseSearchDerivedData.cpp
@@ -81,6 +81,8 @@ static TAutoConsoleVariable<int32> CVarMotionMatchTestNumIterations(TEXT("a.Moti
 static bool AnyTestFlags(int32 Flags) { return (CVarMotionMatchTestFlags.GetValueOnAnyThread() & Flags) != 0; }
 #endif // ENABLE_ANIM_DEBUG
 
+static TAutoConsoleVariable<bool> CVarMotionMatchReindexCancelledDatabases(TEXT("a.MotionMatch.ReindexCancelledDatabases"), true, TEXT("Reindex Cancelled Databases"));
+
 static const UE::DerivedData::FValueId Id(UE::DerivedData::FValueId::FromName("Data"));
 static const UE::DerivedData::FCacheBucket Bucket("PoseSearchDatabase");
 
@@ -2113,13 +2115,27 @@ void FAsyncPoseSearchDatabasesManagement::Tick(float DeltaTime)
 
 #endif // ENABLE_ANIM_DEBUG
 	
+	const bool bReindexCancelledDatabases = CVarMotionMatchReindexCancelledDatabases.GetValueOnAnyThread();
+
 	// iterating backwards because of the possible RemoveAtSwap 
 	for (int32 TaskIndex = Tasks.Num() - 1; TaskIndex >= 0; --TaskIndex)
 	{
-		if (!Tasks[TaskIndex]->IsValid() || Tasks[TaskIndex]->GetState() == FPoseSearchDatabaseAsyncCacheTask::EState::Cancelled)
+		if (!Tasks[TaskIndex]->IsValid())
 		{
 			Tasks.RemoveAtSwap(TaskIndex, EAllowShrinking::No);
 		}
+		else if (Tasks[TaskIndex]->GetState() == FPoseSearchDatabaseAsyncCacheTask::EState::Cancelled)
+		{
+			if (bReindexCancelledDatabases)
+			{
+				RequestAsyncBuildIndex(Tasks[TaskIndex]->GetDatabase(), ERequestAsyncBuildFlag::NewRequest);
+				Tasks.RemoveAtSwap(TaskIndex, EAllowShrinking::No);
+			}
+			else
+			{
+				Tasks.RemoveAtSwap(TaskIndex, EAllowShrinking::No);
+			}
+		}
 		else
 		{
 			Tasks[TaskIndex]->Update(Mutex);
