Motion Matching Editor - Improved PoseSearchDatabaseItem row asset picker

#rnx

diff --git a/Engine/Plugins/Animation/PoseSearch/Source/Editor/Private/PoseSearchDatabaseAssetListItem.cpp b/Engine/Plugins/Animation/PoseSearch/Source/Editor/Private/PoseSearchDatabaseAssetListItem.cpp
--- a/Engine/Plugins/Animation/PoseSearch/Source/Editor/Private/PoseSearchDatabaseAssetListItem.cpp
+++ b/Engine/Plugins/Animation/PoseSearch/Source/Editor/Private/PoseSearchDatabaseAssetListItem.cpp
@@ -8,6 +8,7 @@
 #include "Animation/DebugSkelMeshComponent.h"
 #include "AnimPreviewInstance.h"
 #include "AssetSelection.h"
+#include "AssetToolsModule.h"
 #include "ClassIconFinder.h"
 #include "DetailColumnSizeData.h"
 #include "DragAndDrop/AssetDragDropOp.h"
@@ -21,6 +22,7 @@
 #include "PoseSearch/PoseSearchDatabase.h"
 #include "PoseSearchDatabaseAssetTree.h"
 #include "PoseSearchDatabaseViewModel.h"
+#include "PropertyCustomizationHelpers.h"
 #include "ScopedTransaction.h"
 #include "SPositiveActionButton.h"
 #include "Styling/AppStyle.h"
@@ -48,6 +50,19 @@ namespace UE::PoseSearch
 		EditorViewModel = InEditorViewModel;
 		SkeletonView = InHierarchy;
 
+		AssetTypeColor = FColor::White;
+		if (UPoseSearchDatabase* Database = InEditorViewModel->GetPoseSearchDatabase())
+		{
+			if (const FPoseSearchDatabaseAnimationAssetBase* DatabaseAnimationAsset = Database->GetAnimationAssetBase(WeakAssetTreeNode.Pin()->SourceAssetIdx))
+			{
+				static FAssetToolsModule& AssetToolsModule = FModuleManager::LoadModuleChecked<FAssetToolsModule>(TEXT("AssetTools"));
+				if (TSharedPtr<IAssetTypeActions> AssetTypeActions = AssetToolsModule.Get().GetAssetTypeActionsForClass(DatabaseAnimationAsset->GetAnimationAssetStaticClass()).Pin())
+				{
+					AssetTypeColor = AssetTypeActions->GetTypeColor();
+				}
+			}
+		}
+		
 		if (InAssetTreeNode->SourceAssetIdx == INDEX_NONE)
 		{
 			ConstructGroupItem(OwnerTable);
@@ -181,7 +196,6 @@ namespace UE::PoseSearch
 		}
 
 		TSharedPtr<SWidget> ItemWidget;
-		const FDetailColumnSizeData& ColumnSizeData = SkeletonView.Pin()->GetColumnSizeData();
 		
 		if (SourceAssetIdx == INDEX_NONE)
 		{
@@ -212,151 +226,199 @@ namespace UE::PoseSearch
 		}
 		else
 		{
-			// Item Icon
-			TSharedPtr<SImage> ItemIconWidget;
 			TSharedPtr<FDatabaseViewModel> ViewModel = EditorViewModel.Pin();
-			if (UPoseSearchDatabase* Database = ViewModel->GetPoseSearchDatabase())
+
+			// Item Thumbnail
 			{
-				if (const FPoseSearchDatabaseAnimationAssetBase* DatabaseAnimationAsset = Database->GetAnimationAssetBase(SourceAssetIdx))
+				// Get item Icon
+				TSharedPtr<SImage> ItemIconWidget;
+				if (UPoseSearchDatabase* Database = ViewModel->GetPoseSearchDatabase())
 				{
-					SAssignNew(ItemIconWidget, SImage)
+					if (const FPoseSearchDatabaseAnimationAssetBase* DatabaseAnimationAsset = Database->GetAnimationAssetBase(SourceAssetIdx))
+					{
+						SAssignNew(ItemIconWidget, SImage)
 						.Image(FSlateIconFinder::FindIconBrushForClass(DatabaseAnimationAsset->GetAnimationAssetStaticClass()));
+					}
 				}
-			}
-
-			// Setup table row to display 
-			SAssignNew(ItemWidget, SHorizontalBox)
-			+SHorizontalBox::Slot()
-			.FillWidth(1.0f)
-			[
-				SNew(SSplitter)
-				.Style(FAppStyle::Get(), "FoliageEditMode.Splitter")
-				.PhysicalSplitterHandleSize(1.0f)
-				.HitDetectionSplitterHandleSize(5.0f)
-				.HighlightedHandleIndex(ColumnSizeData.GetHoveredSplitterIndex())
-				.MinimumSlotHeight(0.5f)
 				
-				// Asset Name with type icon
-				+SSplitter::Slot()
-				.Value(ColumnSizeData.GetNameColumnWidth())
-				.MinSize(0.3f)
-				.OnSlotResized(ColumnSizeData.GetOnNameColumnResized())
+				SAssignNew(AssetThumbnailOverlay, SOverlay)
+				
+				// Item Icon
+				+ SOverlay::Slot()
+				.Padding(1.0f)
 				[
-					SNew(SHorizontalBox)
-					.Clipping(EWidgetClipping::ClipToBounds)
-					+ SHorizontalBox::Slot()
-					.MaxWidth(18)
-					.AutoWidth()
-					.Padding(0.0f, 0.0f, 5.0f, 0.0f)
-					.HAlign(HAlign_Left)
-					.VAlign(VAlign_Center)
+					SNew(SOverlay)
+					+ SOverlay::Slot()
 					[
-						ItemIconWidget.ToSharedRef()
+						SNew(SBorder)
+						.Padding(0.0f)
+						.VAlign(VAlign_Fill)
+						.HAlign(HAlign_Fill)
+						.BorderImage(FAppStyle::GetBrush("AssetThumbnail.AssetBackground"))
+						[
+							SNew(SBorder)
+							.Padding(3.0f)
+							.BorderImage(FStyleDefaults::GetNoBrush())
+							.VAlign(VAlign_Center)
+							.HAlign(HAlign_Center)
+							[
+								ItemIconWidget.ToSharedRef()
+							]
+						]
 					]
-					+ SHorizontalBox::Slot()
-					.FillWidth(1.0f)
-					.VAlign(VAlign_Center)
+
+					// Color strip
+					+ SOverlay::Slot()
+					.HAlign(HAlign_Fill)
+					.VAlign(VAlign_Bottom )
 					[
-						SNew(STextBlock)
-						.Text(this, &SDatabaseAssetListItem::GetName)
-						.ColorAndOpacity(this, &SDatabaseAssetListItem::GetNameTextColorAndOpacity)
+						SNew(SBorder)
+						.BorderImage(FAppStyle::GetBrush("WhiteBrush"))
+						.BorderBackgroundColor(AssetTypeColor)
+						.Padding(FMargin(0, 2, 0, 0))
 					]
 				]
-				
-				// Display information via icons
-				+SSplitter::Slot()
-				.Value(ColumnSizeData.GetValueColumnWidth())
-				.MinSize(0.3f)
-				.OnSlotResized(ColumnSizeData.GetOnValueColumnResized())
+
+				// Square border
+				+ SOverlay::Slot()
 				[
-					// Asset Info.
-
-					// Looping
-					SNew(SHorizontalBox)
-					.Clipping(EWidgetClipping::ClipToBounds)
-					+ SHorizontalBox::Slot()
-					.AutoWidth()
-					.Padding(4.0f, 1.0f)
-					.HAlign(HAlign_Center)
-					.VAlign(VAlign_Center)
+					SNew(SImage)
+					.Image_Lambda([this]() -> const FSlateBrush *
+					{
+						static const FName HoveredBorderName("PropertyEditor.AssetThumbnailBorderHovered");
+						static const FName RegularBorderName("PropertyEditor.AssetThumbnailBorder");
+						
+						if (AssetThumbnailOverlay)
+						{
+							return AssetThumbnailOverlay->IsHovered() ? FAppStyle::Get().GetBrush(HoveredBorderName) : FAppStyle::Get().GetBrush(RegularBorderName);
+						}
+						
+						return nullptr;
+					})
+					.Visibility(EVisibility::SelfHitTestInvisible)
+				];
+			}
+			
+			// Picker
+			TSharedPtr<SObjectPropertyEntryBox> AssetPickerWidget;
+			if (UPoseSearchDatabase* Database = ViewModel->GetPoseSearchDatabase())
+			{
+				if (const FPoseSearchDatabaseAnimationAssetBase* DatabaseAnimationAsset = Database->GetAnimationAssetBase(SourceAssetIdx))
+				{
+					SAssignNew(AssetPickerWidget, SObjectPropertyEntryBox)
+					/*.CustomContentSlot()
 					[
-						SNew(SImage)
-						.Image(FAppStyle::Get().GetBrush("Graph.Node.Loop"))
-						.ColorAndOpacity(this, &SDatabaseAssetListItem::GetLoopingColorAndOpacity)
-						.ToolTipText(this, &SDatabaseAssetListItem::GetLoopingToolTip)
-					]
+						SNew(STextBlock)
+						.Text(this, &SDatabaseAssetListItem::GetName)
+						.ColorAndOpacity(this, &SDatabaseAssetListItem::GetNameTextColorAndOpacity)
+					]*/
+					.DisplayThumbnail(false)
+					.IsEnabled_Lambda([ViewModelPtr = EditorViewModel.Pin(), TreeNodePtr = WeakAssetTreeNode.Pin()]()
+					{
+						if (const UPoseSearchDatabase* Database = ViewModelPtr->GetPoseSearchDatabase())
+						{
+							if (Database->GetAnimationAssets().IsValidIndex(TreeNodePtr->SourceAssetIdx))
+							{
+								return ViewModelPtr->IsEnabled(TreeNodePtr->SourceAssetIdx);
+							}
+						}
 
-					// Root Motion
-					+ SHorizontalBox::Slot()
-					.Padding(1.0f, 2.0f)
-					.AutoWidth()
-					.HAlign(HAlign_Center)
-					.VAlign(VAlign_Center)
-					[
-						SNew(SImage)
-						.Image(FAppStyle::Get().GetBrush("AnimGraph.Attribute.RootMotionDelta.Icon"))
-						.DesiredSizeOverride(FVector2D{16.f, 16.f})
-						.ColorAndOpacity(this, &SDatabaseAssetListItem::GetRootMotionColorAndOpacity)
-						.ToolTipText(this, &SDatabaseAssetListItem::GetRootMotionOptionToolTip)
-					]
-					
-					// Mirror Type
-					+ SHorizontalBox::Slot()
-					.Padding(2.0f, 3.0f)
-					.AutoWidth()
-					.HAlign(HAlign_Center)
-					.VAlign(VAlign_Center)
-					[
-						SNew(SImage)
-						.Image(this, &SDatabaseAssetListItem::GetMirrorOptionSlateBrush)
-						.ToolTipText(this, &SDatabaseAssetListItem::GetMirrorOptionToolTip)
-						.OnMouseButtonDown(this, &SDatabaseAssetListItem::MirrorOptionOnMouseButtonDown)
-					]
+						return false;
+					})
+					.AllowedClass(DatabaseAnimationAsset->GetAnimationAssetStaticClass())
+					.ObjectPath_Lambda([Database, SourceAssetIdx]()
+					{
+						if (const FPoseSearchDatabaseAnimationAssetBase* DatabaseAnimationAsset = Database->GetAnimationAssetBase(SourceAssetIdx))
+						{
+							if (const UObject* AnimAsset = DatabaseAnimationAsset->GetAnimationAsset())
+							{
+								return AnimAsset->GetPathName();
+							}
+						}
+						
+						return FString("");
+					})
+					.OnObjectChanged_Lambda([ViewModel, SourceAssetIdx](const FAssetData& AssetData)
+					{
+						const FScopedTransaction Transaction(LOCTEXT("Edit Asset", "Edit Asset"));
 
-					// Disable Reselection
-					+ SHorizontalBox::Slot()
-					.Padding(4.0f, 1.0f)
-					.AutoWidth()
-					.HAlign(HAlign_Center)
-					.VAlign(VAlign_Center)
-					[
-						SNew(SCheckBox)
-						.IsChecked(this, &SDatabaseAssetListItem::GetDisableReselectionChecked)
-						.OnCheckStateChanged(const_cast<SDatabaseAssetListItem*>(this), &SDatabaseAssetListItem::OnDisableReselectionChanged)
-						.ToolTipText(this, &SDatabaseAssetListItem::GetDisableReselectionToolTip)
-						.Padding(FMargin(0.0f, 2.0f, 0.0f, 0.0f))
-						.CheckedImage(FAppStyle::Get().GetBrush("MotionMatchingEditor.DisablePoseReselection"))
-						.CheckedHoveredImage(FAppStyle::Get().GetBrush("MotionMatchingEditor.DisablePoseReselection"))
-						.CheckedPressedImage(FAppStyle::Get().GetBrush("MotionMatchingEditor.DisablePoseReselection"))
-						.UncheckedImage(FAppStyle::Get().GetBrush("MotionMatchingEditor.EnablePoseReselection"))
-						.UncheckedHoveredImage(FAppStyle::Get().GetBrush("MotionMatchingEditor.EnablePoseReselection"))
-						.UncheckedPressedImage(FAppStyle::Get().GetBrush("MotionMatchingEditor.EnablePoseReselection"))
-					]
+						if (UObject* AnimAsset = AssetData.GetAsset())
+						{
+							ViewModel->SetAnimationAsset(SourceAssetIdx, AnimAsset);
+						}
+					});
+				}
+			}
+
+			// Info icons
+			TSharedPtr<SHorizontalBox> InfoIconsHorizontalBox;
+			{
+				SAssignNew(InfoIconsHorizontalBox, SHorizontalBox)
+				+ SHorizontalBox::Slot()
+				.AutoWidth()
+				.Padding(4.0f, 1.0f)
+				.HAlign(HAlign_Center)
+				.VAlign(VAlign_Center)
+				[
+					SNew(SImage)
+					.Image(FAppStyle::Get().GetBrush("Graph.Node.Loop"))
+					.ColorAndOpacity(this, &SDatabaseAssetListItem::GetLoopingColorAndOpacity)
+					.ToolTipText(this, &SDatabaseAssetListItem::GetLoopingToolTip)
 				]
-			]
-			
-			+SHorizontalBox::Slot()
-			.AutoWidth()
-			.HAlign(HAlign_Right)
-			.VAlign(VAlign_Center)
-			[
-				SNew(SHorizontalBox)
+
+				// Root Motion
 				+ SHorizontalBox::Slot()
-				.MaxWidth(18)
-				.Padding(4.0f, 0.0f, 0.0f, 0.0f)
+				.Padding(1.0f, 2.0f)
 				.AutoWidth()
-				.HAlign(HAlign_Right)
+				.HAlign(HAlign_Center)
 				.VAlign(VAlign_Center)
 				[
 					SNew(SImage)
-					.Image(FAppStyle::Get().GetBrush("Icons.EyeDropper"))
-					.Visibility_Raw(this, &SDatabaseAssetListItem::GetSelectedActorIconVisbility)
+					.Image(FAppStyle::Get().GetBrush("AnimGraph.Attribute.RootMotionDelta.Icon"))
+					.DesiredSizeOverride(FVector2D{16.f, 16.f})
+					.ColorAndOpacity(this, &SDatabaseAssetListItem::GetRootMotionColorAndOpacity)
+					.ToolTipText(this, &SDatabaseAssetListItem::GetRootMotionOptionToolTip)
+				]
+				
+				// Mirror Type
+				+ SHorizontalBox::Slot()
+				.Padding(2.0f, 3.0f)
+				.AutoWidth()
+				.HAlign(HAlign_Center)
+				.VAlign(VAlign_Center)
+				[
+					SNew(SImage)
+					.Image(this, &SDatabaseAssetListItem::GetMirrorOptionSlateBrush)
+					.ToolTipText(this, &SDatabaseAssetListItem::GetMirrorOptionToolTip)
+					.OnMouseButtonDown(this, &SDatabaseAssetListItem::MirrorOptionOnMouseButtonDown)
 				]
+
+				// Disable Reselection
+				+ SHorizontalBox::Slot()
+				.Padding(4.0f, 1.0f)
+				.AutoWidth()
+				.HAlign(HAlign_Center)
+				.VAlign(VAlign_Center)
+				[
+					SNew(SCheckBox)
+					.IsChecked(this, &SDatabaseAssetListItem::GetDisableReselectionChecked)
+					.OnCheckStateChanged(const_cast<SDatabaseAssetListItem*>(this), &SDatabaseAssetListItem::OnDisableReselectionChanged)
+					.ToolTipText(this, &SDatabaseAssetListItem::GetDisableReselectionToolTip)
+					.Padding(FMargin(0.0f, 2.0f, 0.0f, 0.0f))
+					.CheckedImage(FAppStyle::Get().GetBrush("MotionMatchingEditor.DisablePoseReselection"))
+					.CheckedHoveredImage(FAppStyle::Get().GetBrush("MotionMatchingEditor.DisablePoseReselection"))
+					.CheckedPressedImage(FAppStyle::Get().GetBrush("MotionMatchingEditor.DisablePoseReselection"))
+					.UncheckedImage(FAppStyle::Get().GetBrush("MotionMatchingEditor.EnablePoseReselection"))
+					.UncheckedHoveredImage(FAppStyle::Get().GetBrush("MotionMatchingEditor.EnablePoseReselection"))
+					.UncheckedPressedImage(FAppStyle::Get().GetBrush("MotionMatchingEditor.EnablePoseReselection"))
+				]
+				
+				// Disable/Enable
 				+ SHorizontalBox::Slot()
 				.MaxWidth(16)
 				.Padding(4.0f, 0.0f)
 				.AutoWidth()
-				.HAlign(HAlign_Right)
+				.HAlign(HAlign_Center)
 				.VAlign(VAlign_Center)
 				[
 					SNew(SCheckBox)
@@ -370,6 +432,68 @@ namespace UE::PoseSearch
 					.UncheckedHoveredImage(FAppStyle::Get().GetBrush("Icons.Hidden"))
 					.UncheckedPressedImage(FAppStyle::Get().GetBrush("Icons.Hidden"))
 				]
+
+				// Is this the picked item?
+				+ SHorizontalBox::Slot()
+				.MaxWidth(18)
+				.Padding(4.0f, 0.0f, 4.0f, 0.0f)
+				.AutoWidth()
+				.HAlign(HAlign_Center)
+				.VAlign(VAlign_Center)
+				[
+					SNew(SImage)
+					.Image(FAppStyle::Get().GetBrush("Icons.EyeDropper"))
+					.Visibility_Raw(this, &SDatabaseAssetListItem::GetSelectedActorIconVisbility)
+				];
+			}
+			
+			// Setup table row to display database item
+			SAssignNew(ItemWidget, SHorizontalBox)
+			.Clipping(EWidgetClipping::ClipToBounds)
+			+ SHorizontalBox::Slot()
+			.Padding(0, 0.0, 0.0, 0.0)
+			.FillWidth(1.0f)
+			[
+				SNew(SSplitter)
+				.Style(FAppStyle::Get(), "FoliageEditMode.Splitter")
+				.PhysicalSplitterHandleSize(0.0f)
+				.HitDetectionSplitterHandleSize(0.0f)
+				.MinimumSlotHeight(0.5f)
+					
+				// Asset Name with type icon
+				+ SSplitter::Slot()
+				.SizeRule(SSplitter::FractionOfParent)
+				[
+					SNew(SBorder)
+					.HAlign(HAlign_Left)
+					.VAlign(VAlign_Fill)
+					.BorderImage(FStyleDefaults::GetNoBrush())
+					[
+						SNew(SHorizontalBox)
+						.Clipping(EWidgetClipping::ClipToBounds)
+						+ SHorizontalBox::Slot()
+						.AutoWidth()
+						.Padding(0.0f, 0.0f, 10.0f, 0.0f)
+						.HAlign(HAlign_Left)
+						.VAlign(VAlign_Center)
+						[
+							AssetThumbnailOverlay.ToSharedRef()
+						]
+						+ SHorizontalBox::Slot()
+						.FillWidth(1.0f)
+						.VAlign(VAlign_Center)
+						[
+							AssetPickerWidget.ToSharedRef()
+						]
+					]
+				]
+					
+				// Display information via icons
+				+SSplitter::Slot()
+				.SizeRule(SSplitter::SizeToContent)
+				[
+					InfoIconsHorizontalBox.ToSharedRef()
+				]
 			];
 		}
 
diff --git a/Engine/Plugins/Animation/PoseSearch/Source/Editor/Private/PoseSearchDatabaseAssetTree.cpp b/Engine/Plugins/Animation/PoseSearch/Source/Editor/Private/PoseSearchDatabaseAssetTree.cpp
--- a/Engine/Plugins/Animation/PoseSearch/Source/Editor/Private/PoseSearchDatabaseAssetTree.cpp
+++ b/Engine/Plugins/Animation/PoseSearch/Source/Editor/Private/PoseSearchDatabaseAssetTree.cpp
@@ -40,8 +40,6 @@ namespace UE::PoseSearch
 		TSharedRef<FDatabaseViewModel> InEditorViewModel)
 	{
 		EditorViewModel = InEditorViewModel;
-		
-		ColumnSizeData.SetValueColumnWidth(0.6f);
 
 		CreateCommandList();
 
diff --git a/Engine/Plugins/Animation/PoseSearch/Source/Editor/Private/PoseSearchDatabaseViewModel.cpp b/Engine/Plugins/Animation/PoseSearch/Source/Editor/Private/PoseSearchDatabaseViewModel.cpp
--- a/Engine/Plugins/Animation/PoseSearch/Source/Editor/Private/PoseSearchDatabaseViewModel.cpp
+++ b/Engine/Plugins/Animation/PoseSearch/Source/Editor/Private/PoseSearchDatabaseViewModel.cpp
@@ -27,6 +27,7 @@
 #include "PoseSearchDatabasePreviewScene.h"
 #include "PoseSearchEditor.h"
 #include "PropertyEditorModule.h"
+#include "PoseSearch/PoseSearchMultiSequence.h"
 
 namespace UE::PoseSearch
 {
@@ -677,6 +678,52 @@ bool FDatabaseViewModel::IsEnabled(int32 AnimationAssetIndex) const
 	return false;
 }
 
+bool FDatabaseViewModel::SetAnimationAsset(int32 AnimationAssetIndex, UObject* AnimAsset)
+{
+	if (UPoseSearchDatabase* Database = GetPoseSearchDatabase())
+	{
+		if (FPoseSearchDatabaseAnimationAssetBase* DatabaseAnimationAsset = Database->GetMutableAnimationAssetBase(AnimationAssetIndex))
+		{
+			Database->Modify();
+
+			// Ensure that our target database item matches the input object's class.
+			UClass* AssetClass = AnimAsset->GetClass();
+			if (AssetClass == DatabaseAnimationAsset->GetAnimationAssetStaticClass())
+			{
+				if (AssetClass->IsChildOf( UAnimSequence::StaticClass()))
+				{
+					FPoseSearchDatabaseSequence* DatabaseSequenceAsset = static_cast<FPoseSearchDatabaseSequence*>(DatabaseAnimationAsset);
+					DatabaseSequenceAsset->Sequence = Cast<UAnimSequence>(AnimAsset);
+				}
+				else if (AssetClass->IsChildOf(UAnimComposite::StaticClass()))
+				{
+					FPoseSearchDatabaseAnimComposite* DatabaseCompositeAsset = static_cast<FPoseSearchDatabaseAnimComposite*>(DatabaseAnimationAsset);
+					DatabaseCompositeAsset->AnimComposite = Cast<UAnimComposite>(AnimAsset);
+				}
+				else if (AssetClass->IsChildOf(UAnimMontage::StaticClass()))
+				{
+					FPoseSearchDatabaseAnimMontage* DatabaseMontageAsset = static_cast<FPoseSearchDatabaseAnimMontage*>(DatabaseAnimationAsset);
+					DatabaseMontageAsset->AnimMontage = Cast<UAnimMontage>(AnimAsset);
+				}
+				else if (AssetClass->IsChildOf(UBlendSpace::StaticClass()))
+				{
+					FPoseSearchDatabaseBlendSpace* DatabaseBlendSpaceAsset = static_cast<FPoseSearchDatabaseBlendSpace*>(DatabaseAnimationAsset);
+					DatabaseBlendSpaceAsset->BlendSpace = Cast<UBlendSpace>(AnimAsset);
+				}
+				else if (AssetClass->IsChildOf(UPoseSearchMultiSequence::StaticClass()))
+				{
+					FPoseSearchDatabaseMultiSequence* DatabaseMultiSequenceAsset = static_cast<FPoseSearchDatabaseMultiSequence*>(DatabaseAnimationAsset);
+					DatabaseMultiSequenceAsset->MultiSequence = Cast<UPoseSearchMultiSequence>(AnimAsset);
+				}
+
+				return true;
+			}
+		}
+	}
+
+	return false;
+}
+
 int32 FDatabaseViewModel::SetSelectedNode(int32 PoseIdx, bool bClearSelection, bool bDrawQuery, TConstArrayView<float> InQueryVector)
 {
 	int32 SelectedSourceAssetIdx = INDEX_NONE;
diff --git a/Engine/Plugins/Animation/PoseSearch/Source/Editor/Public/PoseSearchDatabaseAssetListItem.h b/Engine/Plugins/Animation/PoseSearch/Source/Editor/Public/PoseSearchDatabaseAssetListItem.h
--- a/Engine/Plugins/Animation/PoseSearch/Source/Editor/Public/PoseSearchDatabaseAssetListItem.h
+++ b/Engine/Plugins/Animation/PoseSearch/Source/Editor/Public/PoseSearchDatabaseAssetListItem.h
@@ -65,6 +65,9 @@ namespace UE::PoseSearch
 		TWeakPtr<FDatabaseAssetTreeNode> WeakAssetTreeNode;
 		TWeakPtr<FDatabaseViewModel> EditorViewModel;
 		TWeakPtr<SDatabaseAssetTree> SkeletonView;
+		
+		FColor AssetTypeColor;
+		TSharedPtr<SOverlay> AssetThumbnailOverlay;
 	};
 }
 
diff --git a/Engine/Plugins/Animation/PoseSearch/Source/Editor/Public/PoseSearchDatabaseAssetTree.h b/Engine/Plugins/Animation/PoseSearch/Source/Editor/Public/PoseSearchDatabaseAssetTree.h
--- a/Engine/Plugins/Animation/PoseSearch/Source/Editor/Public/PoseSearchDatabaseAssetTree.h
+++ b/Engine/Plugins/Animation/PoseSearch/Source/Editor/Public/PoseSearchDatabaseAssetTree.h
@@ -40,7 +40,6 @@ namespace UE::PoseSearch
 		void FinalizeTreeChanges(bool bRecoverSelection = false, bool bRefreshView = true);
 
 		void SetSelectedItem(int32 SourceAssetIdx, bool bClearSelection);
-		FDetailColumnSizeData& GetColumnSizeData() { return ColumnSizeData; }
 		
 	protected:
 		TWeakPtr<FDatabaseViewModel> EditorViewModel;
@@ -117,8 +116,6 @@ namespace UE::PoseSearch
 		void OnConvertToBranchIn();
 
 		friend SDatabaseAssetListItem;
-
-		FDetailColumnSizeData ColumnSizeData;
 		
 	protected:
 		// Called when an item is selected/deselected
diff --git a/Engine/Plugins/Animation/PoseSearch/Source/Editor/Public/PoseSearchDatabaseViewModel.h b/Engine/Plugins/Animation/PoseSearch/Source/Editor/Public/PoseSearchDatabaseViewModel.h
--- a/Engine/Plugins/Animation/PoseSearch/Source/Editor/Public/PoseSearchDatabaseViewModel.h
+++ b/Engine/Plugins/Animation/PoseSearch/Source/Editor/Public/PoseSearchDatabaseViewModel.h
@@ -117,6 +117,8 @@ namespace UE::PoseSearch
 		void SetIsEnabled(int32 AnimationAssetIndex, bool bEnabled);
 		bool IsEnabled(int32 AnimationAssetIndex) const;
 
+		bool SetAnimationAsset(int32 AnimationAssetIndex, UObject* AnimAsset);
+		
 		int32 SetSelectedNode(int32 PoseIdx, bool bClearSelection, bool bDrawQuery, TConstArrayView<float> InQueryVector);
 		void SetSelectedNodes(const TArrayView<TSharedPtr<FDatabaseAssetTreeNode>>& InSelectedNodes);
 		void ProcessSelectedActor(AActor* Actor);
