Motion Matching Editor - Now you can set a database item's options via the right click menu.

[RN] minor

diff --git a/Engine/Plugins/Animation/PoseSearch/Source/Editor/Private/PoseSearchDatabaseAssetListItem.cpp b/Engine/Plugins/Animation/PoseSearch/Source/Editor/Private/PoseSearchDatabaseAssetListItem.cpp
--- a/Engine/Plugins/Animation/PoseSearch/Source/Editor/Private/PoseSearchDatabaseAssetListItem.cpp
+++ b/Engine/Plugins/Animation/PoseSearch/Source/Editor/Private/PoseSearchDatabaseAssetListItem.cpp
@@ -727,26 +727,18 @@ namespace UE::PoseSearch
 
 			if (InMouseEvent.GetEffectingButton() == EKeys::LeftMouseButton)
 			{
-				if (UPoseSearchDatabase* Database = ViewModel->GetPoseSearchDatabase())
-				{
-					if (FPoseSearchDatabaseAnimationAssetBase* DatabaseAnimationAsset = Database->GetMutableAnimationAssetBase(AssetTreeNode->SourceAssetIdx))
-					{
-						const FScopedTransaction Transaction(LOCTEXT("OnClickEditMirrorOptionPoseSearchDatabase", "Edit Mirror Option"));
-
-						// Get next mirror option
-						static const TArray<EPoseSearchMirrorOption> OptionArray = { EPoseSearchMirrorOption::UnmirroredOnly, EPoseSearchMirrorOption::MirroredOnly, EPoseSearchMirrorOption::UnmirroredAndMirrored };
-						const int32 NextOption = (static_cast<int32>(DatabaseAnimationAsset->MirrorOption) + 1) % OptionArray.Num();
-
-						// Modify asset (@todo: this should be done through the viewmodel).
-						Database->Modify();
-						DatabaseAnimationAsset->MirrorOption = OptionArray[NextOption];
-
-						SkeletonView.Pin()->RefreshTreeView(false, true);
-						ViewModel->BuildSearchIndex();
+				const FScopedTransaction Transaction(LOCTEXT("OnClickEditMirrorOptionPoseSearchDatabase", "Edit Mirror Option"));
+				
+				// Get next mirror option
+				static const TArray<EPoseSearchMirrorOption> OptionArray = { EPoseSearchMirrorOption::UnmirroredOnly, EPoseSearchMirrorOption::MirroredOnly, EPoseSearchMirrorOption::UnmirroredAndMirrored };
+				const int32 NextOption = (static_cast<int32>(ViewModel->GetMirrorOption(AssetTreeNode->SourceAssetIdx)) + 1) % OptionArray.Num();
+				
+				ViewModel->SetMirrorOption(AssetTreeNode->SourceAssetIdx, OptionArray[NextOption]);
+				
+				SkeletonView.Pin()->RefreshTreeView(false, true);
+				ViewModel->BuildSearchIndex();
 
-						return FReply::Handled();
-					}
-				}
+				return FReply::Handled();
 			}
 		}
 		return FReply::Unhandled();
diff --git a/Engine/Plugins/Animation/PoseSearch/Source/Editor/Private/PoseSearchDatabaseAssetTree.cpp b/Engine/Plugins/Animation/PoseSearch/Source/Editor/Private/PoseSearchDatabaseAssetTree.cpp
--- a/Engine/Plugins/Animation/PoseSearch/Source/Editor/Private/PoseSearchDatabaseAssetTree.cpp
+++ b/Engine/Plugins/Animation/PoseSearch/Source/Editor/Private/PoseSearchDatabaseAssetTree.cpp
@@ -497,22 +497,142 @@ namespace UE::PoseSearch
 		{
 			MenuBuilder.BeginSection("SelectedAssetsEdit", LOCTEXT("SelectedAssetEdit", "Asset Actions"));
 			{
-				MenuBuilder.AddMenuEntry(
-					LOCTEXT("Enable", "Enable selected assets"),
-					LOCTEXT("EnableTooltip", "Sets Assets Enabled."),
-					FSlateIcon(FAppStyle::GetAppStyleSetName(), "Icons.Visible"),
-					FUIAction(FExecuteAction::CreateSP(this, &SDatabaseAssetTree::OnEnableNodes)),
+				const TSharedPtr<FDatabaseViewModel> ViewModel = EditorViewModel.Pin();
+
+				const EPoseSearchMirrorOption MirrorOption = ViewModel->GetMirrorOption(SelectedNodes[0]->SourceAssetIdx);
+				FName IconToUseForSubMenu = NAME_None;
+
+				if (SelectedNodes.Num() == 1)
+				{
+					switch (MirrorOption)
+					{
+						case EPoseSearchMirrorOption::UnmirroredOnly: IconToUseForSubMenu = "GraphEditor.AlignNodesRight"; break;
+						case EPoseSearchMirrorOption::MirroredOnly: IconToUseForSubMenu = "GraphEditor.AlignNodesLeft"; break;
+						case EPoseSearchMirrorOption::UnmirroredAndMirrored: IconToUseForSubMenu = "GraphEditor.AlignNodesCenter"; break;
+					}
+				}
+				
+				const FText LabelToUseForSubMenu = SelectedNodes.Num() == 1 ? LOCTEXT("SetMirrorOption", "Set Mirror Option") : LOCTEXT("SetMirrorOptionInSelectedAssets", "Set Mirror Option on selected assets");
+				const FText TooltipToUseForSubMenu = LOCTEXT("SetMirrorOptionTooltip", "Set the mirror option in the selected asset(s)");
+					
+				MenuBuilder.AddSubMenu(
+				LabelToUseForSubMenu, 
+				TooltipToUseForSubMenu,
+				FNewMenuDelegate::CreateLambda([this](FMenuBuilder& MenuBuilder)
+				{
+						
+					MenuBuilder.AddMenuEntry(
+						LOCTEXT("OriginalOnly", "Original Only"),
+						LOCTEXT("OriginalOnlyTooltip", "Mirror Option: Original Only"),
+						FSlateIcon(FAppStyle::GetAppStyleSetName(), "GraphEditor.AlignNodesRight"),
+						FUIAction(FExecuteAction::CreateSP(this, &SDatabaseAssetTree::OnSetMirrorOptionForSelectedNodes, EPoseSearchMirrorOption::UnmirroredOnly)),
+						NAME_None,
+						EUserInterfaceActionType::Button
+						);
+
+					MenuBuilder.AddMenuEntry(
+						LOCTEXT("MirrorOnly", "Mirrored Only"),
+						LOCTEXT("MirrorOnlyTooltip", "Mirror Option: Mirrored Only"),
+						FSlateIcon(FAppStyle::GetAppStyleSetName(), "GraphEditor.AlignNodesLeft"),
+						FUIAction(FExecuteAction::CreateSP(this, &SDatabaseAssetTree::OnSetMirrorOptionForSelectedNodes, EPoseSearchMirrorOption::MirroredOnly)),
+						NAME_None,
+						EUserInterfaceActionType::Button
+						);
+					
+					MenuBuilder.AddMenuEntry(
+						LOCTEXT("OriginalAndMirrorOnly", "Original and Mirrored"),
+						LOCTEXT("OriginalAndMirrorOnlyTooltip", "Mirror Option: Original and Mirrored"),
+						FSlateIcon(FAppStyle::GetAppStyleSetName(), "GraphEditor.AlignNodesCenter"),
+						FUIAction(FExecuteAction::CreateSP(this, &SDatabaseAssetTree::OnSetMirrorOptionForSelectedNodes, EPoseSearchMirrorOption::UnmirroredAndMirrored)),
+						NAME_None,
+						EUserInterfaceActionType::Button
+						);
+				}),
+				false,
+				FSlateIcon(FAppStyle::GetAppStyleSetName(), IconToUseForSubMenu));
+
+				if (SelectedNodes.Num() > 1)
+				{
+					MenuBuilder.AddMenuEntry(
+					LOCTEXT("EnableReselection", "Enable pose reselection in selected assets"),
+					LOCTEXT("EnableTooltipReselection", "Enable reselection of poses in the same asset"),
+					FSlateIcon(FAppStyle::GetAppStyleSetName(), "MotionMatchingEditor.EnablePoseReselection"),
+					FUIAction(FExecuteAction::CreateSP(this, &SDatabaseAssetTree::OnSetPoseReselectionForSelectedNodes, true)),
 					NAME_None,
 					EUserInterfaceActionType::Button);
 
-				MenuBuilder.AddMenuEntry(
-					LOCTEXT("Disable", "Disable selected assets"),
-					LOCTEXT("DisableToolTip", "Sets Assets Disabled."),
-					FSlateIcon(FAppStyle::GetAppStyleSetName(), "Icons.Hidden"),
-					FUIAction(FExecuteAction::CreateSP(this, &SDatabaseAssetTree::OnDisableNodes)),
+					MenuBuilder.AddMenuEntry(
+						LOCTEXT("DisableReselection", "Disable pose reselection in selected assets"),
+						LOCTEXT("DisableToolTipReselection", "Disable reselection of poses in the same asset"),
+						FSlateIcon(FAppStyle::GetAppStyleSetName(), "MotionMatchingEditor.DisablePoseReselection"),
+						FUIAction(FExecuteAction::CreateSP(this, &SDatabaseAssetTree::OnSetPoseReselectionForSelectedNodes, false)),
+						NAME_None,
+						EUserInterfaceActionType::Button);
+				
+				}
+				else
+				{
+					const bool IsReselectionDisabled = ViewModel->IsDisableReselection(SelectedNodes[0]->SourceAssetIdx);
+					const FName IconToUse = IsReselectionDisabled ? "MotionMatchingEditor.EnablePoseReselection" : "MotionMatchingEditor.DisablePoseReselection";
+					const FText LabelToUse = IsReselectionDisabled ? LOCTEXT("EnablePoseReselection", "Enable pose reselection") : LOCTEXT("DisablePoseReselection", "Disable pose reselection");
+					const FText TooltipToUse = IsReselectionDisabled ? LOCTEXT("EnablePoseReselectionTooltip", "Enable pose reselection in the same asset") : LOCTEXT("DisablePoseReselectionTooltip", "Disable pose reselection in the same asset");
+					
+					MenuBuilder.AddMenuEntry(
+						LabelToUse,
+						TooltipToUse,
+						FSlateIcon(FAppStyle::GetAppStyleSetName(), IconToUse),
+						FUIAction(FExecuteAction::CreateLambda([this, SourceAssetIndex = SelectedNodes[0]->SourceAssetIdx, ViewModel]()
+						{
+							bool IsPoseReselectionDisabled = ViewModel->IsDisableReselection(SourceAssetIndex);
+							ViewModel->SetDisableReselection(SourceAssetIndex, !IsPoseReselectionDisabled);
+							
+							RefreshTreeView(false, true);
+						})),
+						NAME_None,
+						EUserInterfaceActionType::Button);
+				}
+
+				if (SelectedNodes.Num() > 1)
+				{
+					MenuBuilder.AddMenuEntry(
+					LOCTEXT("EnableSelectedAssets", "Enable selected assets"),
+					LOCTEXT("EnableSelectedAssetsToolTip", "Sets Assets Enabled."),
+					FSlateIcon(FAppStyle::GetAppStyleSetName(), "Icons.Visible"),
+					FUIAction(FExecuteAction::CreateSP(this, &SDatabaseAssetTree::OnEnableNodes)),
 					NAME_None,
 					EUserInterfaceActionType::Button);
 
+					MenuBuilder.AddMenuEntry(
+						LOCTEXT("DisableSelectedAssets", "Disable selected assets"),
+						LOCTEXT("DisableSelectedAssetsToolTip", "Sets Assets Disabled."),
+						FSlateIcon(FAppStyle::GetAppStyleSetName(), "Icons.Hidden"),
+						FUIAction(FExecuteAction::CreateSP(this, &SDatabaseAssetTree::OnDisableNodes)),
+						NAME_None,
+						EUserInterfaceActionType::Button);
+				}
+				else
+				{
+					const bool IsEnabled = ViewModel->IsEnabled(SelectedNodes[0]->SourceAssetIdx);
+					const FName IconToUse = !IsEnabled ? "Icons.Visible" : "Icons.Hidden";
+					const FText LabelToUse = !IsEnabled ? LOCTEXT("EnableAsset", "Enable asset") : LOCTEXT("DisableAsset", "Disable asset");
+					const FText TooltipToUse = !IsEnabled ? LOCTEXT("EnableAssetTooltip", "Inlcude asset in query when pose matching / motion matching.") : LOCTEXT("DisableAssetTooltip", "Exclude asset in query when pose matching / motion matching.");
+					
+					MenuBuilder.AddMenuEntry(
+						LabelToUse,
+						TooltipToUse,
+						FSlateIcon(FAppStyle::GetAppStyleSetName(), IconToUse),
+						FUIAction(FExecuteAction::CreateLambda([this, SourceAssetIndex = SelectedNodes[0]->SourceAssetIdx, ViewModel]()
+						{
+							const bool IsEnabled = ViewModel->IsEnabled(SourceAssetIndex);
+							ViewModel->SetIsEnabled(SourceAssetIndex, !IsEnabled);
+							
+							FinalizeTreeChanges();
+						})),
+						NAME_None,
+						EUserInterfaceActionType::Button);
+				}
+				
+
 				MenuBuilder.AddMenuEntry(
 					LOCTEXT("ConvertToBranchIn", "Convert selected assets to sample via BranchIn notify"),
 					LOCTEXT("ConvertToBranchInToolTip", "Creates PoseSearchBranchIn notify state for the asset sampling range"),
@@ -886,6 +1006,50 @@ namespace UE::PoseSearch
 		}
 	}
 
+	void SDatabaseAssetTree::OnSetMirrorOptionForSelectedNodes(EPoseSearchMirrorOption InMirrorOption)
+	{
+		const TSharedPtr<FDatabaseViewModel> ViewModel = EditorViewModel.Pin();
+		if (UPoseSearchDatabase* PoseSearchDatabase = ViewModel->GetPoseSearchDatabase())
+		{
+			TArray<TSharedPtr<FDatabaseAssetTreeNode>> SelectedNodes = TreeView->GetSelectedItems();
+			if (!SelectedNodes.IsEmpty())
+			{
+				const FScopedTransaction Transaction(LOCTEXT("OnClickEditMirrorOptionPoseSearchDatabase", "Edit Mirror Option on selected items"));
+				
+				PoseSearchDatabase->Modify();
+
+				for (const TSharedPtr<FDatabaseAssetTreeNode>& SelectedNode : SelectedNodes)
+				{
+					ViewModel->SetMirrorOption(SelectedNode->SourceAssetIdx, InMirrorOption);
+				}
+
+				FinalizeTreeChanges();
+			}
+		}
+	}
+
+	void SDatabaseAssetTree::OnSetPoseReselectionForSelectedNodes(bool bIsEnabled)
+	{
+		const TSharedPtr<FDatabaseViewModel> ViewModel = EditorViewModel.Pin();
+		if (UPoseSearchDatabase* PoseSearchDatabase = ViewModel->GetPoseSearchDatabase())
+		{
+			TArray<TSharedPtr<FDatabaseAssetTreeNode>> SelectedNodes = TreeView->GetSelectedItems();
+			if (!SelectedNodes.IsEmpty())
+			{
+				const FScopedTransaction Transaction(LOCTEXT("OnClickEditPoseReselection", "Set pose reselection for selected items"));
+				
+				PoseSearchDatabase->Modify();
+
+				for (const TSharedPtr<FDatabaseAssetTreeNode>& SelectedNode : SelectedNodes)
+				{
+					ViewModel->SetDisableReselection(SelectedNode->SourceAssetIdx, !bIsEnabled);
+				}
+
+				RefreshTreeView(false, true);
+			}
+		}
+	}
+
 	void SDatabaseAssetTree::OnConvertToBranchIn()
 	{
 		const TSharedPtr<FDatabaseViewModel> ViewModel = EditorViewModel.Pin();
diff --git a/Engine/Plugins/Animation/PoseSearch/Source/Editor/Private/PoseSearchDatabaseViewModel.cpp b/Engine/Plugins/Animation/PoseSearch/Source/Editor/Private/PoseSearchDatabaseViewModel.cpp
--- a/Engine/Plugins/Animation/PoseSearch/Source/Editor/Private/PoseSearchDatabaseViewModel.cpp
+++ b/Engine/Plugins/Animation/PoseSearch/Source/Editor/Private/PoseSearchDatabaseViewModel.cpp
@@ -724,6 +724,32 @@ bool FDatabaseViewModel::SetAnimationAsset(int32 AnimationAssetIndex, UObject* A
 	return false;
 }
 
+void FDatabaseViewModel::SetMirrorOption(int32 AnimationAssetIndex, EPoseSearchMirrorOption InMirrorOption)
+{
+	if (UPoseSearchDatabase* Database = GetPoseSearchDatabase())
+	{
+		if (FPoseSearchDatabaseAnimationAssetBase* DatabaseAnimationAsset = Database->GetMutableAnimationAssetBase(AnimationAssetIndex))
+		{
+			Database->Modify();
+
+			DatabaseAnimationAsset->MirrorOption = InMirrorOption;
+		}
+	}
+}
+
+EPoseSearchMirrorOption FDatabaseViewModel::GetMirrorOption(int32 AnimationAssetIndex)
+{
+	if (const UPoseSearchDatabase* Database = GetPoseSearchDatabase())
+	{
+		if (const FPoseSearchDatabaseAnimationAssetBase* DatabaseAnimationAsset = Database->GetAnimationAssetBase(AnimationAssetIndex))
+		{
+			return DatabaseAnimationAsset->MirrorOption;
+		}
+	}
+
+	return EPoseSearchMirrorOption::MirroredOnly;
+}
+
 int32 FDatabaseViewModel::SetSelectedNode(int32 PoseIdx, bool bClearSelection, bool bDrawQuery, TConstArrayView<float> InQueryVector)
 {
 	int32 SelectedSourceAssetIdx = INDEX_NONE;
diff --git a/Engine/Plugins/Animation/PoseSearch/Source/Editor/Public/PoseSearchDatabaseAssetTree.h b/Engine/Plugins/Animation/PoseSearch/Source/Editor/Public/PoseSearchDatabaseAssetTree.h
--- a/Engine/Plugins/Animation/PoseSearch/Source/Editor/Public/PoseSearchDatabaseAssetTree.h
+++ b/Engine/Plugins/Animation/PoseSearch/Source/Editor/Public/PoseSearchDatabaseAssetTree.h
@@ -112,7 +112,9 @@ namespace UE::PoseSearch
 		void EnableSelectedNodes(bool bIsEnabled);
 		void OnEnableNodes() { EnableSelectedNodes(true); }
 		void OnDisableNodes() { EnableSelectedNodes(false); }
-
+		void OnSetMirrorOptionForSelectedNodes(EPoseSearchMirrorOption InMirrorOption);
+		void OnSetPoseReselectionForSelectedNodes(bool bIsEnabled);
+		
 		void OnConvertToBranchIn();
 
 		friend SDatabaseAssetListItem;
diff --git a/Engine/Plugins/Animation/PoseSearch/Source/Editor/Public/PoseSearchDatabaseViewModel.h b/Engine/Plugins/Animation/PoseSearch/Source/Editor/Public/PoseSearchDatabaseViewModel.h
--- a/Engine/Plugins/Animation/PoseSearch/Source/Editor/Public/PoseSearchDatabaseViewModel.h
+++ b/Engine/Plugins/Animation/PoseSearch/Source/Editor/Public/PoseSearchDatabaseViewModel.h
@@ -10,6 +10,7 @@
 #include "PoseSearch/PoseSearchRole.h"
 #include "PoseSearch/PoseSearchTrajectoryTypes.h"
 #include "PoseSearchDatabasePreviewScene.h"
+#include "PoseSearch/PoseSearchDatabase.h"
 #include "UObject/GCObject.h"
 
 class UWorld;
@@ -118,6 +119,9 @@ namespace UE::PoseSearch
 		bool IsEnabled(int32 AnimationAssetIndex) const;
 
 		bool SetAnimationAsset(int32 AnimationAssetIndex, UObject* AnimAsset);
+
+		void SetMirrorOption(int32 AnimationAssetIndex, EPoseSearchMirrorOption InMirrorOption);
+		EPoseSearchMirrorOption GetMirrorOption(int32 AnimationAssetIndex);
 		
 		int32 SetSelectedNode(int32 PoseIdx, bool bClearSelection, bool bDrawQuery, TConstArrayView<float> InQueryVector);
 		void SetSelectedNodes(const TArrayView<TSharedPtr<FDatabaseAssetTreeNode>>& InSelectedNodes);
