motion matching - fixed discrepancy between FPoseSearchTrajectoryData::StepCharacterMovementGroundPrediction and UCharacterMovementComponent::ApplyVelocityBraking

diff --git a/Engine/Plugins/Animation/PoseSearch/Source/Runtime/Private/PoseSearchTrajectoryLibrary.cpp b/Engine/Plugins/Animation/PoseSearch/Source/Runtime/Private/PoseSearchTrajectoryLibrary.cpp
--- a/Engine/Plugins/Animation/PoseSearch/Source/Runtime/Private/PoseSearchTrajectoryLibrary.cpp
+++ b/Engine/Plugins/Animation/PoseSearch/Source/Runtime/Private/PoseSearchTrajectoryLibrary.cpp
@@ -37,6 +37,7 @@ void FPoseSearchTrajectoryData::UpdateData(
 	{
 		TrajectoryDataDerived.MaxSpeed = FMath::Max(MoveComp->GetMaxSpeed() * MoveComp->GetAnalogInputModifier(), MoveComp->GetMinAnalogSpeed());
 		TrajectoryDataDerived.BrakingDeceleration = FMath::Max(0.f, MoveComp->GetMaxBrakingDeceleration());
+		TrajectoryDataDerived.BrakingSubStepTime = MoveComp->BrakingSubStepTime;
 		TrajectoryDataDerived.bOrientRotationToMovement = MoveComp->bOrientRotationToMovement;
 
 		TrajectoryDataDerived.Velocity = MoveComp->Velocity;
@@ -104,8 +105,8 @@ FVector FPoseSearchTrajectoryData::StepCharacterMovementGroundPrediction(
 			return InVelocity;
 		}
 
-		static const float MaxTimeStep = 1.f / 60.f;
 		float RemainingTime = DeltaTime;
+		const float MaxTimeStep = FMath::Clamp(TrajectoryDataDerived.BrakingSubStepTime, 1.0f / 75.0f, 1.0f / 20.0f);
 
 		const FVector PrevLinearVelocity = OutVelocity;
 		const FVector RevAccel = (bZeroBraking ? FVector::ZeroVector : (-TrajectoryDataDerived.BrakingDeceleration * OutVelocity.GetSafeNormal()));
diff --git a/Engine/Plugins/Animation/PoseSearch/Source/Runtime/Public/PoseSearch/PoseSearchTrajectoryLibrary.h b/Engine/Plugins/Animation/PoseSearch/Source/Runtime/Public/PoseSearch/PoseSearchTrajectoryLibrary.h
--- a/Engine/Plugins/Animation/PoseSearch/Source/Runtime/Public/PoseSearch/PoseSearchTrajectoryLibrary.h
+++ b/Engine/Plugins/Animation/PoseSearch/Source/Runtime/Public/PoseSearch/PoseSearchTrajectoryLibrary.h
@@ -26,6 +26,7 @@ public:
 
 		float MaxSpeed = 0.f;
 		float BrakingDeceleration = 0.f;
+		float BrakingSubStepTime = 0.f;
 		float Friction = 0.f;
 
 		FVector Velocity = FVector::ZeroVector;
