Motion Matching Editor - Added warnings for an empty database item and a asset with an incompatible skeleton.

Other changes:
- Updated asset picker widget to use functions instead of lambdas for readability.
- Addressed `FDatabaseViewModel::SetAnimationAsset` not checking in the input `FAssetData` is null.

[RN] minor

diff --git a/Engine/Plugins/Animation/PoseSearch/Source/Editor/Private/PoseSearchDatabaseAssetListItem.cpp b/Engine/Plugins/Animation/PoseSearch/Source/Editor/Private/PoseSearchDatabaseAssetListItem.cpp
--- a/Engine/Plugins/Animation/PoseSearch/Source/Editor/Private/PoseSearchDatabaseAssetListItem.cpp
+++ b/Engine/Plugins/Animation/PoseSearch/Source/Editor/Private/PoseSearchDatabaseAssetListItem.cpp
@@ -306,47 +306,23 @@ namespace UE::PoseSearch
 				if (const FPoseSearchDatabaseAnimationAssetBase* DatabaseAnimationAsset = Database->GetAnimationAssetBase(SourceAssetIdx))
 				{
 					SAssignNew(AssetPickerWidget, SObjectPropertyEntryBox)
-					/*.CustomContentSlot()
+					.AllowClear(false)
+					.AllowedClass(DatabaseAnimationAsset->GetAnimationAssetStaticClass())
+					.DisplayThumbnail(false)
+					.IsEnabled(this, &SDatabaseAssetListItem::GetAssetPickerIsEnabled)
+					.ObjectPath(this, &SDatabaseAssetListItem::GetAssetPickerObjectPath)
+					.OnObjectChanged(this, &SDatabaseAssetListItem::OnAssetPickerObjectChanged)
+					.CustomContentSlot()
 					[
+						// Display warning below picked asset.
 						SNew(STextBlock)
-						.Text(this, &SDatabaseAssetListItem::GetName)
-						.ColorAndOpacity(this, &SDatabaseAssetListItem::GetNameTextColorAndOpacity)
-					]*/
-					.DisplayThumbnail(false)
-					.IsEnabled_Lambda([ViewModelPtr = EditorViewModel.Pin(), TreeNodePtr = WeakAssetTreeNode.Pin()]()
-					{
-						if (const UPoseSearchDatabase* Database = ViewModelPtr->GetPoseSearchDatabase())
-						{
-							if (Database->GetAnimationAssets().IsValidIndex(TreeNodePtr->SourceAssetIdx))
-							{
-								return ViewModelPtr->IsEnabled(TreeNodePtr->SourceAssetIdx);
-							}
-						}
-
-						return false;
-					})
-					.AllowedClass(DatabaseAnimationAsset->GetAnimationAssetStaticClass())
-					.ObjectPath_Lambda([Database, SourceAssetIdx]()
-					{
-						if (const FPoseSearchDatabaseAnimationAssetBase* DatabaseAnimationAsset = Database->GetAnimationAssetBase(SourceAssetIdx))
-						{
-							if (const UObject* AnimAsset = DatabaseAnimationAsset->GetAnimationAsset())
-							{
-								return AnimAsset->GetPathName();
-							}
-						}
-						
-						return FString("");
-					})
-					.OnObjectChanged_Lambda([ViewModel, SourceAssetIdx](const FAssetData& AssetData)
-					{
-						const FScopedTransaction Transaction(LOCTEXT("Edit Asset", "Edit Asset"));
-
-						if (UObject* AnimAsset = AssetData.GetAsset())
-						{
-							ViewModel->SetAnimationAsset(SourceAssetIdx, AnimAsset);
-						}
-					});
+						.Margin(FMargin(2,0))
+						.Justification(ETextJustify::Left)
+						.Font(FSlateFontInfo(FCoreStyle::GetDefaultFont(),8, "Regular"))
+						.Text(this, &SDatabaseAssetListItem::GetAssetPickerText)
+						.ColorAndOpacity(this, &SDatabaseAssetListItem::GetAssetPickerCustomContentSlotTextColor)
+						.Visibility(this, &SDatabaseAssetListItem::GetAssetPickerCustomContentSlotVisibility)
+					];
 				}
 			}
 
@@ -443,7 +419,7 @@ namespace UE::PoseSearch
 				[
 					SNew(SImage)
 					.Image(FAppStyle::Get().GetBrush("Icons.EyeDropper"))
-					.Visibility_Raw(this, &SDatabaseAssetListItem::GetSelectedActorIconVisbility)
+					.Visibility_Raw(this, &SDatabaseAssetListItem::GetSelectedActorIconVisibility)
 				];
 			}
 			
@@ -512,7 +488,7 @@ namespace UE::PoseSearch
 		}
 	}
 
-	EVisibility SDatabaseAssetListItem::GetSelectedActorIconVisbility() const
+	EVisibility SDatabaseAssetListItem::GetSelectedActorIconVisibility() const
 	{
 		TSharedPtr<FDatabaseViewModel> ViewModelPtr = EditorViewModel.Pin();
 		if (TSharedPtr<FDatabaseAssetTreeNode> AssetTreeNode = WeakAssetTreeNode.Pin())
@@ -528,6 +504,109 @@ namespace UE::PoseSearch
 		return EVisibility::Hidden;
 	}
 
+	void SDatabaseAssetListItem::OnAssetPickerObjectChanged(const FAssetData& AssetData)
+	{
+		const TSharedPtr<FDatabaseViewModel> ViewModelPtr = EditorViewModel.Pin();
+
+		if (const TSharedPtr<FDatabaseAssetTreeNode> AssetTreeNode = WeakAssetTreeNode.Pin())
+		{
+			const FScopedTransaction Transaction(LOCTEXT("Edit Asset", "Edit Asset"));
+
+			ViewModelPtr->SetAnimationAsset(AssetTreeNode->SourceAssetIdx, AssetData.GetAsset());
+		}
+	}
+
+	FString SDatabaseAssetListItem::GetAssetPickerObjectPath() const
+	{
+		const TSharedPtr<FDatabaseViewModel> ViewModelPtr = EditorViewModel.Pin();
+
+		if (const TSharedPtr<FDatabaseAssetTreeNode> AssetTreeNode = WeakAssetTreeNode.Pin())
+		{
+			if (const UPoseSearchDatabase* Database = ViewModelPtr->GetPoseSearchDatabase())
+			{
+				if (const FPoseSearchDatabaseAnimationAssetBase* DatabaseAnimationAsset = Database->GetAnimationAssetBase(AssetTreeNode->SourceAssetIdx))
+				{
+					if (const UObject* AnimAsset = DatabaseAnimationAsset->GetAnimationAsset())
+					{
+						return AnimAsset->GetPathName();
+					}
+				}
+			}
+		}
+		
+		return FString("");
+	}
+
+	bool SDatabaseAssetListItem::GetAssetPickerIsEnabled() const
+	{
+		const TSharedPtr<FDatabaseViewModel> ViewModelPtr = EditorViewModel.Pin();
+
+		if (const TSharedPtr<FDatabaseAssetTreeNode> AssetTreeNode = WeakAssetTreeNode.Pin())
+		{
+			if (const UPoseSearchDatabase* Database = ViewModelPtr->GetPoseSearchDatabase())
+			{
+				if (Database->GetAnimationAssets().IsValidIndex(AssetTreeNode->SourceAssetIdx))
+				{
+					return ViewModelPtr->IsEnabled(AssetTreeNode->SourceAssetIdx);
+				}
+			}
+		}
+		
+		return false;
+	}
+
+	EVisibility SDatabaseAssetListItem::GetAssetPickerCustomContentSlotVisibility() const
+	{
+		const TSharedPtr<FDatabaseViewModel> ViewModelPtr = EditorViewModel.Pin();
+							
+		if (const UPoseSearchDatabase* Database = ViewModelPtr->GetPoseSearchDatabase())
+		{
+			if (const TSharedPtr<FDatabaseAssetTreeNode> AssetTreeNode = WeakAssetTreeNode.Pin())
+			{
+				if (const FPoseSearchDatabaseAnimationAssetBase* DatabaseAnimationAssetBase = Database->GetAnimationAssetBase(AssetTreeNode->SourceAssetIdx))
+				{
+					if (DatabaseAnimationAssetBase->IsEnabled())
+					{
+						if (DatabaseAnimationAssetBase->GetAnimationAsset() == nullptr || !DatabaseAnimationAssetBase->IsSkeletonCompatible(Database->Schema))
+						{
+							return EVisibility::Visible;
+						}
+					}
+				}
+			}
+		}
+							
+		return EVisibility::Collapsed;
+	}
+
+	FText SDatabaseAssetListItem::GetAssetPickerText() const
+	{
+		const TSharedPtr<FDatabaseViewModel> ViewModelPtr = EditorViewModel.Pin();
+							
+		if (const UPoseSearchDatabase* Database = ViewModelPtr->GetPoseSearchDatabase())
+		{
+			if (const TSharedPtr<FDatabaseAssetTreeNode> AssetTreeNode = WeakAssetTreeNode.Pin())
+			{
+				if (const FPoseSearchDatabaseAnimationAssetBase* DatabaseAnimationAssetBase = Database->GetAnimationAssetBase(AssetTreeNode->SourceAssetIdx))
+				{
+					if (DatabaseAnimationAssetBase->IsEnabled())
+					{
+						if (DatabaseAnimationAssetBase->GetAnimationAsset() == nullptr)
+						{
+							return LOCTEXT("ErrorNoAsset", "No asset has been selected.");
+						}
+						else if (!DatabaseAnimationAssetBase->IsSkeletonCompatible(Database->Schema))
+						{
+							return LOCTEXT("ErrorIncompatibleSkeleton", "This asset's skeleton is not compatible with the schema's skeleton(s).");
+						}
+					}
+				}
+			}
+		}
+
+		return FText::GetEmpty();
+	}
+
 	FText SDatabaseAssetListItem::GetDisableReselectionToolTip() const
 	{
 		if (GetDisableReselectionChecked() == ECheckBoxState::Checked)
@@ -611,32 +690,27 @@ namespace UE::PoseSearch
 		}
 	}
 
-	FSlateColor SDatabaseAssetListItem::GetNameTextColorAndOpacity() const
+	FSlateColor SDatabaseAssetListItem::GetAssetPickerCustomContentSlotTextColor() const
 	{
-		TSharedPtr<FDatabaseViewModel> ViewModelPtr = EditorViewModel.Pin();
+		const TSharedPtr<FDatabaseViewModel> ViewModelPtr = EditorViewModel.Pin();
+
 		if (const UPoseSearchDatabase* Database = ViewModelPtr->GetPoseSearchDatabase())
 		{
-			if (TSharedPtr<FDatabaseAssetTreeNode> AssetTreeNode = WeakAssetTreeNode.Pin())
+			if (const TSharedPtr<FDatabaseAssetTreeNode> AssetTreeNode = WeakAssetTreeNode.Pin())
 			{
 				if (const FPoseSearchDatabaseAnimationAssetBase* DatabaseAnimationAssetBase = Database->GetAnimationAssetBase(AssetTreeNode->SourceAssetIdx))
 				{
 					if (DatabaseAnimationAssetBase->IsEnabled())
 					{
-						if (DatabaseAnimationAssetBase->bSynchronizeWithExternalDependency)
-						{
-							return FColor::Turquoise;
-						}
-
-						if (!DatabaseAnimationAssetBase->IsSkeletonCompatible(Database->Schema))
+						if (DatabaseAnimationAssetBase->GetAnimationAsset() == nullptr || !DatabaseAnimationAssetBase->IsSkeletonCompatible(Database->Schema))
 						{
 							return FColor::Red;
 						}
-						
-						return FLinearColor::White;
 					}
 				}
 			}
 		}
+		
 		return DisabledColor;
 	}
 
diff --git a/Engine/Plugins/Animation/PoseSearch/Source/Editor/Private/PoseSearchDatabaseViewModel.cpp b/Engine/Plugins/Animation/PoseSearch/Source/Editor/Private/PoseSearchDatabaseViewModel.cpp
--- a/Engine/Plugins/Animation/PoseSearch/Source/Editor/Private/PoseSearchDatabaseViewModel.cpp
+++ b/Engine/Plugins/Animation/PoseSearch/Source/Editor/Private/PoseSearchDatabaseViewModel.cpp
@@ -680,47 +680,50 @@ bool FDatabaseViewModel::IsEnabled(int32 AnimationAssetIndex) const
 
 bool FDatabaseViewModel::SetAnimationAsset(int32 AnimationAssetIndex, UObject* AnimAsset)
 {
-	if (UPoseSearchDatabase* Database = GetPoseSearchDatabase())
+	if (AnimAsset)
 	{
-		if (FPoseSearchDatabaseAnimationAssetBase* DatabaseAnimationAsset = Database->GetMutableAnimationAssetBase(AnimationAssetIndex))
+		if (UPoseSearchDatabase* Database = GetPoseSearchDatabase())
 		{
-			Database->Modify();
-
-			// Ensure that our target database item matches the input object's class.
-			UClass* AssetClass = AnimAsset->GetClass();
-			if (AssetClass == DatabaseAnimationAsset->GetAnimationAssetStaticClass())
+			if (FPoseSearchDatabaseAnimationAssetBase* DatabaseAnimationAsset = Database->GetMutableAnimationAssetBase(AnimationAssetIndex))
 			{
-				if (AssetClass->IsChildOf( UAnimSequence::StaticClass()))
-				{
-					FPoseSearchDatabaseSequence* DatabaseSequenceAsset = static_cast<FPoseSearchDatabaseSequence*>(DatabaseAnimationAsset);
-					DatabaseSequenceAsset->Sequence = Cast<UAnimSequence>(AnimAsset);
-				}
-				else if (AssetClass->IsChildOf(UAnimComposite::StaticClass()))
-				{
-					FPoseSearchDatabaseAnimComposite* DatabaseCompositeAsset = static_cast<FPoseSearchDatabaseAnimComposite*>(DatabaseAnimationAsset);
-					DatabaseCompositeAsset->AnimComposite = Cast<UAnimComposite>(AnimAsset);
-				}
-				else if (AssetClass->IsChildOf(UAnimMontage::StaticClass()))
-				{
-					FPoseSearchDatabaseAnimMontage* DatabaseMontageAsset = static_cast<FPoseSearchDatabaseAnimMontage*>(DatabaseAnimationAsset);
-					DatabaseMontageAsset->AnimMontage = Cast<UAnimMontage>(AnimAsset);
-				}
-				else if (AssetClass->IsChildOf(UBlendSpace::StaticClass()))
-				{
-					FPoseSearchDatabaseBlendSpace* DatabaseBlendSpaceAsset = static_cast<FPoseSearchDatabaseBlendSpace*>(DatabaseAnimationAsset);
-					DatabaseBlendSpaceAsset->BlendSpace = Cast<UBlendSpace>(AnimAsset);
-				}
-				else if (AssetClass->IsChildOf(UPoseSearchMultiSequence::StaticClass()))
+				Database->Modify();
+
+				// Ensure that our target database item matches the input object's class.
+				const UClass* AssetClass = AnimAsset->GetClass();
+				if (AssetClass == DatabaseAnimationAsset->GetAnimationAssetStaticClass())
 				{
-					FPoseSearchDatabaseMultiSequence* DatabaseMultiSequenceAsset = static_cast<FPoseSearchDatabaseMultiSequence*>(DatabaseAnimationAsset);
-					DatabaseMultiSequenceAsset->MultiSequence = Cast<UPoseSearchMultiSequence>(AnimAsset);
-				}
+					if (AssetClass->IsChildOf( UAnimSequence::StaticClass()))
+					{
+						FPoseSearchDatabaseSequence* DatabaseSequenceAsset = static_cast<FPoseSearchDatabaseSequence*>(DatabaseAnimationAsset);
+						DatabaseSequenceAsset->Sequence = Cast<UAnimSequence>(AnimAsset);
+					}
+					else if (AssetClass->IsChildOf(UAnimComposite::StaticClass()))
+					{
+						FPoseSearchDatabaseAnimComposite* DatabaseCompositeAsset = static_cast<FPoseSearchDatabaseAnimComposite*>(DatabaseAnimationAsset);
+						DatabaseCompositeAsset->AnimComposite = Cast<UAnimComposite>(AnimAsset);
+					}
+					else if (AssetClass->IsChildOf(UAnimMontage::StaticClass()))
+					{
+						FPoseSearchDatabaseAnimMontage* DatabaseMontageAsset = static_cast<FPoseSearchDatabaseAnimMontage*>(DatabaseAnimationAsset);
+						DatabaseMontageAsset->AnimMontage = Cast<UAnimMontage>(AnimAsset);
+					}
+					else if (AssetClass->IsChildOf(UBlendSpace::StaticClass()))
+					{
+						FPoseSearchDatabaseBlendSpace* DatabaseBlendSpaceAsset = static_cast<FPoseSearchDatabaseBlendSpace*>(DatabaseAnimationAsset);
+						DatabaseBlendSpaceAsset->BlendSpace = Cast<UBlendSpace>(AnimAsset);
+					}
+					else if (AssetClass->IsChildOf(UPoseSearchMultiSequence::StaticClass()))
+					{
+						FPoseSearchDatabaseMultiSequence* DatabaseMultiSequenceAsset = static_cast<FPoseSearchDatabaseMultiSequence*>(DatabaseAnimationAsset);
+						DatabaseMultiSequenceAsset->MultiSequence = Cast<UPoseSearchMultiSequence>(AnimAsset);
+					}
 
-				return true;
+					return true;
+				}
 			}
 		}
 	}
-
+	
 	return false;
 }
 
diff --git a/Engine/Plugins/Animation/PoseSearch/Source/Editor/Public/PoseSearchDatabaseAssetListItem.h b/Engine/Plugins/Animation/PoseSearch/Source/Editor/Public/PoseSearchDatabaseAssetListItem.h
--- a/Engine/Plugins/Animation/PoseSearch/Source/Editor/Public/PoseSearchDatabaseAssetListItem.h
+++ b/Engine/Plugins/Animation/PoseSearch/Source/Editor/Public/PoseSearchDatabaseAssetListItem.h
@@ -40,9 +40,14 @@ namespace UE::PoseSearch
 		virtual FReply OnMouseButtonDoubleClick(const FGeometry& InMyGeometry, const FPointerEvent& InMouseEvent) override;
 
 		
-		EVisibility GetSelectedActorIconVisbility() const;
+		EVisibility GetSelectedActorIconVisibility() const;
 
-		FSlateColor GetNameTextColorAndOpacity() const;
+		void OnAssetPickerObjectChanged(const FAssetData& AssetData);
+		FString GetAssetPickerObjectPath() const;
+		bool GetAssetPickerIsEnabled() const;
+		EVisibility GetAssetPickerCustomContentSlotVisibility() const;
+		FText GetAssetPickerText() const;
+		FSlateColor GetAssetPickerCustomContentSlotTextColor() const;
 		
 		FSlateColor GetLoopingColorAndOpacity() const;
 		FText GetLoopingToolTip() const;
