motion matching -  adding separate ratios for rotation and translation to recover to the skeleton reference the future root bone transform

diff --git a/Engine/Plugins/Animation/PoseSearch/Source/Runtime/Private/AnimNode_PoseSearchHistoryCollector.cpp b/Engine/Plugins/Animation/PoseSearch/Source/Runtime/Private/AnimNode_PoseSearchHistoryCollector.cpp
--- a/Engine/Plugins/Animation/PoseSearch/Source/Runtime/Private/AnimNode_PoseSearchHistoryCollector.cpp
+++ b/Engine/Plugins/Animation/PoseSearch/Source/Runtime/Private/AnimNode_PoseSearchHistoryCollector.cpp
@@ -49,7 +49,8 @@ void FAnimNode_PoseSearchHistoryCollector_Base::Initialize_AnyThread(const FAnim
 		FMemMark Mark(FMemStack::Get());
 		FCSPose<FCompactPose> ComponentSpacePose;
 		ComponentSpacePose.InitPose(&Context.AnimInstanceProxy->GetRequiredBones());
-		PoseHistory.EvaluateComponentSpace_AnyThread(0.f, ComponentSpacePose, bStoreScales, RootBoneRecoveryTime, true, true, GetRequiredBones(Context.AnimInstanceProxy));
+		PoseHistory.EvaluateComponentSpace_AnyThread(0.f, ComponentSpacePose, bStoreScales, 
+			RootBoneRecoveryTime, RootBoneTranslationRecoveryRatio, RootBoneRotationRecoveryRatio, true, true, GetRequiredBones(Context.AnimInstanceProxy));
 	}
 }
 
@@ -158,8 +159,8 @@ void FAnimNode_PoseSearchHistoryCollector::Evaluate_AnyThread(FPoseContext& Outp
 		RequiredBones = GetRequiredBones(Output.AnimInstanceProxy);
 	}
 
-	PoseHistory.EvaluateComponentSpace_AnyThread(Output.AnimInstanceProxy->GetDeltaSeconds(), ComponentSpacePose,
-		bStoreScales, RootBoneRecoveryTime, bNeedsReset, bCacheBones, RequiredBones);
+	PoseHistory.EvaluateComponentSpace_AnyThread(Output.AnimInstanceProxy->GetDeltaSeconds(), ComponentSpacePose, bStoreScales,
+		RootBoneRecoveryTime, RootBoneTranslationRecoveryRatio, RootBoneRotationRecoveryRatio, bNeedsReset, bCacheBones, RequiredBones);
 
 	bCacheBones = false;
 
@@ -223,8 +224,8 @@ void FAnimNode_PoseSearchComponentSpaceHistoryCollector::EvaluateComponentSpace_
 		RequiredBones = GetRequiredBones(Output.AnimInstanceProxy);
 	}
 
-	PoseHistory.EvaluateComponentSpace_AnyThread(Output.AnimInstanceProxy->GetDeltaSeconds(), Output.Pose,
-		bStoreScales, RootBoneRecoveryTime, bNeedsReset, bCacheBones, RequiredBones);
+	PoseHistory.EvaluateComponentSpace_AnyThread(Output.AnimInstanceProxy->GetDeltaSeconds(), Output.Pose, bStoreScales, 
+		RootBoneRecoveryTime, RootBoneTranslationRecoveryRatio, RootBoneRotationRecoveryRatio, bNeedsReset, bCacheBones, RequiredBones);
 	
 	bCacheBones = false;
 
diff --git a/Engine/Plugins/Animation/PoseSearch/Source/Runtime/Private/PoseSearchHistory.cpp b/Engine/Plugins/Animation/PoseSearch/Source/Runtime/Private/PoseSearchHistory.cpp
--- a/Engine/Plugins/Animation/PoseSearch/Source/Runtime/Private/PoseSearchHistory.cpp
+++ b/Engine/Plugins/Animation/PoseSearch/Source/Runtime/Private/PoseSearchHistory.cpp
@@ -470,6 +470,11 @@ void FArchivedPoseHistory::DebugDraw(const UWorld* World, FColor Color) const
 					}
 				}
 
+				if (i == 0)
+				{
+					DrawDebugLine(World, GlobalTransforms.GetTranslation(), RootTransform.GetTranslation(), FColor::Purple, false, -1.f, SDPG_Foreground);
+				}
+
 				PrevGlobalTransforms[i] = GlobalTransforms;
 			}
 		}
@@ -694,7 +699,8 @@ void FPoseHistory::SetTrajectory(const FPoseSearchQueryTrajectory& InTrajectory,
 }
 
 void FPoseHistory::EvaluateComponentSpace_AnyThread(float DeltaTime, FCSPose<FCompactPose>& ComponentSpacePose, bool bStoreScales,
-	float RootBoneRecoveryTime, bool bNeedsReset, bool bCacheBones, const TArray<FBoneIndexType>& RequiredBones)
+	float RootBoneRecoveryTime, float RootBoneTranslationRecoveryRatio, float RootBoneRotationRecoveryRatio,
+	bool bNeedsReset, bool bCacheBones, const TArray<FBoneIndexType>& RequiredBones)
 {
 	CheckThreadSafetyWrite(WritePoseDataThreadSafeCounter);
 	CheckThreadSafetyRead(ReadPoseDataThreadSafeCounter);
@@ -782,17 +788,28 @@ void FPoseHistory::EvaluateComponentSpace_AnyThread(float DeltaTime, FCSPose<FCo
 	}
 
 	// Regardless of the retention policy, we always update the most recent Entry
-	WritePoseData.Entries.Last().Update(0.f, ComponentSpacePose, WritePoseData.BoneToTransformMap, bStoreScales);
+	FPoseHistoryEntry& MostRecentEntry = WritePoseData.Entries.Last();
+	MostRecentEntry.Update(0.f, ComponentSpacePose, WritePoseData.BoneToTransformMap, bStoreScales);
 
-	if (RootBoneRecoveryTime > 0.f)
+	if (RootBoneRecoveryTime > 0.f && !Trajectory.Samples.IsEmpty())
 	{
 		// adding the updated "future" root bone Entry
-		const FReferenceSkeleton& RefSkeleton = Skeleton->GetReferenceSkeleton();
-		const TArray<FTransform>& RefBonePose = RefSkeleton.GetRefBonePose();
-		const FTransform& RefRootBone = RefBonePose[RootBoneIndexType];
+		const FTransform& RefRootBone = Skeleton->GetReferenceSkeleton().GetRefBonePose()[RootBoneIndexType];
+		const FQuat RootBoneRotationAtRecoveryTime = FMath::Lerp(FQuat(MostRecentEntry.ComponentSpaceRotations[RootBoneIndexType]), RefRootBone.GetRotation(), RootBoneRotationRecoveryRatio);
+
+		FVector RootBoneDeltaTranslationAtRecoveryTime = FVector::ZeroVector;
+		if (RootBoneTranslationRecoveryRatio > 0.f)
+		{
+			const FTransform WorldRootAtCurrentTime = Trajectory.GetSampleAtTime(0.f).GetTransform();
+			const FTransform WorldRootBoneAtCurrentTime = MostRecentEntry.GetComponentSpaceTransform(RootBoneIndexType) * WorldRootAtCurrentTime;
+			const FVector WorldRootBoneDeltaTranslationAtCurrentTime = (WorldRootBoneAtCurrentTime.GetTranslation() - WorldRootAtCurrentTime.GetTranslation()) * RootBoneTranslationRecoveryRatio;
+			const FTransform WorldRootAtRecoveryTime = Trajectory.GetSampleAtTime(RootBoneRecoveryTime).GetTransform();
+			RootBoneDeltaTranslationAtRecoveryTime = WorldRootAtRecoveryTime.InverseTransformVector(WorldRootBoneDeltaTranslationAtCurrentTime);
+		}
 
+		const FTransform RootBoneTransformAtRecoveryTime(RootBoneRotationAtRecoveryTime, RootBoneDeltaTranslationAtRecoveryTime, RefRootBone.GetScale3D());
 		FutureEntryTemp.SetNum(1, bStoreScales);
-		FutureEntryTemp.SetComponentSpaceTransform(RootBoneIndexType, RefRootBone);
+		FutureEntryTemp.SetComponentSpaceTransform(RootBoneIndexType, RootBoneTransformAtRecoveryTime);
 		FutureEntryTemp.AccumulatedSeconds = RootBoneRecoveryTime;
 		WritePoseData.Entries.Emplace(MoveTemp(FutureEntryTemp));
 	}
diff --git a/Engine/Plugins/Animation/PoseSearch/Source/Runtime/Public/PoseSearch/AnimNode_PoseSearchHistoryCollector.h b/Engine/Plugins/Animation/PoseSearch/Source/Runtime/Public/PoseSearch/AnimNode_PoseSearchHistoryCollector.h
--- a/Engine/Plugins/Animation/PoseSearch/Source/Runtime/Public/PoseSearch/AnimNode_PoseSearchHistoryCollector.h
+++ b/Engine/Plugins/Animation/PoseSearch/Source/Runtime/Public/PoseSearch/AnimNode_PoseSearchHistoryCollector.h
@@ -37,10 +37,19 @@ public:
 	UPROPERTY(EditAnywhere, Category = Settings)
 	bool bStoreScales = false;
 
-	// time in seconds to recover to the reference skeleton root bone from any eventual root bone modification. if zero the behaviour will be disabled (Experimental)
+	// time in seconds to recover to the reference skeleton root bone transform by RootBoneTranslationRecoveryRatio and RootBoneRotationRecoveryRatio
+	// from any eventual root bone modification. if zero the behaviour will be disabled (Experimental)
 	UPROPERTY(EditAnywhere, BlueprintReadWrite, Category = Experimental, meta = (ClampMin="0"))
 	float RootBoneRecoveryTime = 0.f;
 
+	// ratio to recover to the reference skeleton root bone translation from any eventual root bone modification. zero for no recovery, 1 for full recovery
+	UPROPERTY(EditAnywhere, BlueprintReadWrite, Category = Experimental, meta = (ClampMin="0", ClampMax="1", EditCondition = "RootBoneRecoveryTime > 0", EditConditionHides))
+	float RootBoneTranslationRecoveryRatio = 1.f;
+
+	// ratio to recover to the reference skeleton root bone rotation from any eventual root bone modification. zero for no recovery, 1 for full recovery
+	UPROPERTY(EditAnywhere, BlueprintReadWrite, Category = Experimental, meta = (ClampMin="0", ClampMax="1", EditCondition = "RootBoneRecoveryTime > 0", EditConditionHides))
+	float RootBoneRotationRecoveryRatio = 1.f;
+
 	// Update Counter for detecting being relevant
 	FGraphTraversalCounter UpdateCounter;
 
diff --git a/Engine/Plugins/Animation/PoseSearch/Source/Runtime/Public/PoseSearch/PoseSearchHistory.h b/Engine/Plugins/Animation/PoseSearch/Source/Runtime/Public/PoseSearch/PoseSearchHistory.h
--- a/Engine/Plugins/Animation/PoseSearch/Source/Runtime/Public/PoseSearch/PoseSearchHistory.h
+++ b/Engine/Plugins/Animation/PoseSearch/Source/Runtime/Public/PoseSearch/PoseSearchHistory.h
@@ -110,7 +110,9 @@ struct FPoseHistory : public IPoseHistory
 	void PreUpdate();
 
 	void Initialize_AnyThread(int32 InNumPoses, float InSamplingInterval);
-	void EvaluateComponentSpace_AnyThread(float DeltaTime, FCSPose<FCompactPose>& ComponentSpacePose, bool bStoreScales, float RootBoneRecoveryTime, bool bNeedsReset, bool bCacheBones, const TArray<FBoneIndexType>& RequiredBones);
+	void EvaluateComponentSpace_AnyThread(float DeltaTime, FCSPose<FCompactPose>& ComponentSpacePose, bool bStoreScales,
+		float RootBoneRecoveryTime, float RootBoneTranslationRecoveryRatio, float RootBoneRotationRecoveryRatio,
+		bool bNeedsReset, bool bCacheBones, const TArray<FBoneIndexType>& RequiredBones);
 
 	// IPoseHistory interface
 	virtual bool GetTransformAtTime(float Time, FTransform& OutBoneTransform, const USkeleton* BoneIndexSkeleton = nullptr, FBoneIndexType BoneIndexType = RootBoneIndexType, FBoneIndexType ReferenceBoneIndexType = ComponentSpaceIndexType, bool bExtrapolate = false) const override;
