Motion Matching Editor - Support for MultiSequence drag & drop.

[RN] minor

diff --git a/Engine/Plugins/Animation/PoseSearch/Source/Editor/Private/PoseSearchDatabaseAssetTree.cpp b/Engine/Plugins/Animation/PoseSearch/Source/Editor/Private/PoseSearchDatabaseAssetTree.cpp
--- a/Engine/Plugins/Animation/PoseSearch/Source/Editor/Private/PoseSearchDatabaseAssetTree.cpp
+++ b/Engine/Plugins/Animation/PoseSearch/Source/Editor/Private/PoseSearchDatabaseAssetTree.cpp
@@ -22,6 +22,7 @@
 #include "SPositiveActionButton.h"
 #include "Styling/AppStyle.h"
 #include "ScopedTransaction.h"
+#include "PoseSearch/PoseSearchMultiSequence.h"
 #include "Styling/StyleColors.h"
 #include "Widgets/Text/SRichTextBlock.h"
 #include "Widgets/Input/SCheckBox.h"
@@ -139,7 +140,8 @@ namespace UE::PoseSearch
 							if (AssetClass->IsChildOf(UAnimSequence::StaticClass()) ||
 								AssetClass->IsChildOf(UAnimComposite::StaticClass()) ||
 								AssetClass->IsChildOf(UBlendSpace::StaticClass()) ||
-								AssetClass->IsChildOf(UAnimMontage::StaticClass()))
+								AssetClass->IsChildOf(UAnimMontage::StaticClass()) ||
+								AssetClass->IsChildOf(UPoseSearchMultiSequence::StaticClass()))
 							{
 								Reply = FReply::Handled();
 								break;
@@ -342,7 +344,8 @@ namespace UE::PoseSearch
 					if (AssetClass->IsChildOf(UAnimSequence::StaticClass()) ||
 						AssetClass->IsChildOf(UAnimComposite::StaticClass()) ||
 						AssetClass->IsChildOf(UBlendSpace::StaticClass()) ||
-						AssetClass->IsChildOf(UAnimMontage::StaticClass()))
+						AssetClass->IsChildOf(UAnimMontage::StaticClass()) ||
+						AssetClass->IsChildOf(UPoseSearchMultiSequence::StaticClass()))
 					{
 						ReturnedDropZone = EItemDropZone::OntoItem;
 						break;
@@ -422,6 +425,11 @@ namespace UE::PoseSearch
 					ViewModel->AddAnimMontageToDatabase(Cast<UAnimMontage>(Asset));
 					++AddedAssets;
 				}
+				else if (AssetClass->IsChildOf(UPoseSearchMultiSequence::StaticClass()))
+				{
+					ViewModel->AddMultiSequenceToDatabase(Cast<UPoseSearchMultiSequence>(Asset));
+					++AddedAssets;
+				}
 			}
 			
 			GWarn->EndSlowTask();
@@ -754,7 +762,7 @@ namespace UE::PoseSearch
 	{
 		FScopedTransaction Transaction(LOCTEXT("AddMultiSequenceTransaction", "Add Multi Sequence"));
 
-		EditorViewModel.Pin()->AddMultiSequenceToDatabase();
+		EditorViewModel.Pin()->AddMultiSequenceToDatabase(nullptr);
 
 		if (bFinalizeChanges)
 		{
diff --git a/Engine/Plugins/Animation/PoseSearch/Source/Editor/Private/PoseSearchDatabaseViewModel.cpp b/Engine/Plugins/Animation/PoseSearch/Source/Editor/Private/PoseSearchDatabaseViewModel.cpp
--- a/Engine/Plugins/Animation/PoseSearch/Source/Editor/Private/PoseSearchDatabaseViewModel.cpp
+++ b/Engine/Plugins/Animation/PoseSearch/Source/Editor/Private/PoseSearchDatabaseViewModel.cpp
@@ -571,13 +571,14 @@ void FDatabaseViewModel::AddAnimMontageToDatabase(UAnimMontage* AnimMontage)
 	}
 }
 
-void FDatabaseViewModel::AddMultiSequenceToDatabase()
+void FDatabaseViewModel::AddMultiSequenceToDatabase(UPoseSearchMultiSequence* MultiSequence)
 {
 	if (UPoseSearchDatabase* Database = GetPoseSearchDatabase())
 	{
 		Database->Modify();
 		
 		FPoseSearchDatabaseMultiSequence NewAsset;
+		NewAsset.MultiSequence = MultiSequence;
 		Database->AddAnimationAsset(FInstancedStruct::Make(NewAsset));
 	}
 }
diff --git a/Engine/Plugins/Animation/PoseSearch/Source/Editor/Public/PoseSearchDatabaseViewModel.h b/Engine/Plugins/Animation/PoseSearch/Source/Editor/Public/PoseSearchDatabaseViewModel.h
--- a/Engine/Plugins/Animation/PoseSearch/Source/Editor/Public/PoseSearchDatabaseViewModel.h
+++ b/Engine/Plugins/Animation/PoseSearch/Source/Editor/Public/PoseSearchDatabaseViewModel.h
@@ -21,6 +21,7 @@ class UAnimComposite;
 class UAnimSequence;
 class UBlendSpace;
 class UMirrorDataTable;
+class UPoseSearchMultiSequence;
 
 namespace UE::PoseSearch
 {
@@ -108,7 +109,7 @@ namespace UE::PoseSearch
 		void AddBlendSpaceToDatabase(UBlendSpace* BlendSpace);
 		void AddAnimCompositeToDatabase(UAnimComposite* AnimComposite);
 		void AddAnimMontageToDatabase(UAnimMontage* AnimMontage);
-		void AddMultiSequenceToDatabase();
+		void AddMultiSequenceToDatabase(UPoseSearchMultiSequence* MultiSequence);
 
 		bool DeleteFromDatabase(int32 AnimationAssetIndex);
 
