Motion Matching Editor - Make mirror option match the visual style of the other editable options of a database item entry

Other fixes:
- Clicking the mirror option button fast enough does not trigger the asset to be opened

[RN] minor

diff --git a/Engine/Plugins/Animation/PoseSearch/Source/Editor/Private/PoseSearchDatabaseAssetListItem.cpp b/Engine/Plugins/Animation/PoseSearch/Source/Editor/Private/PoseSearchDatabaseAssetListItem.cpp
--- a/Engine/Plugins/Animation/PoseSearch/Source/Editor/Private/PoseSearchDatabaseAssetListItem.cpp
+++ b/Engine/Plugins/Animation/PoseSearch/Source/Editor/Private/PoseSearchDatabaseAssetListItem.cpp
@@ -37,6 +37,111 @@
 namespace UE::PoseSearch
 {
 	static constexpr FLinearColor DisabledColor = FLinearColor(1.f, 1.f, 1.f, 0.25f);
+
+	/* We need a custom widget to be able to consume the "DoubleClick" event so we can cycle through the mirror options but not open the asset. */
+	class SMirrorTypeWidget : public SCompoundWidget
+	{
+		SLATE_BEGIN_ARGS(SMirrorTypeWidget){}
+		SLATE_END_ARGS()
+	public:
+		
+		void Construct(const FArguments& InArgs, const TWeakPtr<FDatabaseAssetTreeNode>& InAssetTreeNode, const TWeakPtr<SDatabaseAssetTree>& InAssetTree, const TWeakPtr<FDatabaseViewModel>& InViewModel)
+		{
+			WeakAssetTreeNode = InAssetTreeNode;
+			SkeletonView = InAssetTree;
+			EditorViewModel = InViewModel;
+			
+			ChildSlot
+			[
+				SNew(SOverlay)
+				+SOverlay::Slot()
+				[
+					SNew(SImage)
+					.Image(this, &SMirrorTypeWidget::GetBackgroundImage)
+				]
+				+SOverlay::Slot()
+				[
+					SNew(SImage)
+					.Image(this, &SMirrorTypeWidget::GetMirrorOptionSlateBrush)
+					.ToolTipText(this, &SMirrorTypeWidget::GetMirrorOptionToolTip)
+				]
+			];
+		}
+
+		virtual FReply OnMouseButtonDoubleClick(const FGeometry& InMyGeometry, const FPointerEvent& InMouseEvent) override
+		{
+			return OnMouseButtonDown(InMyGeometry, InMouseEvent);
+		}
+
+		virtual FReply OnMouseButtonDown(const FGeometry& MyGeometry, const FPointerEvent& InMouseEvent) override
+		{
+			if (const TSharedPtr<FDatabaseAssetTreeNode> AssetTreeNode = WeakAssetTreeNode.Pin())
+			{
+				const TSharedPtr<FDatabaseViewModel> ViewModel = EditorViewModel.Pin();
+
+				if (InMouseEvent.GetEffectingButton() == EKeys::LeftMouseButton)
+				{
+					const FScopedTransaction Transaction(LOCTEXT("OnClickEditMirrorOptionPoseSearchDatabase", "Edit Mirror Option"));
+				
+					// Get next mirror option
+					static const TArray<EPoseSearchMirrorOption> OptionArray = { EPoseSearchMirrorOption::UnmirroredOnly, EPoseSearchMirrorOption::MirroredOnly, EPoseSearchMirrorOption::UnmirroredAndMirrored };
+					const int32 NextOption = (static_cast<int32>(ViewModel->GetMirrorOption(AssetTreeNode->SourceAssetIdx)) + 1) % OptionArray.Num();
+				
+					ViewModel->SetMirrorOption(AssetTreeNode->SourceAssetIdx, OptionArray[NextOption]);
+				
+					SkeletonView.Pin()->RefreshTreeView(false, true);
+					ViewModel->BuildSearchIndex();
+
+					return FReply::Handled();
+				}
+			}
+			
+			return FReply::Unhandled();
+		}
+
+		const FSlateBrush * GetBackgroundImage() const
+		{
+			const FCheckBoxStyle& Style = FCoreStyle::Get().GetWidgetStyle<FCheckBoxStyle>("Checkbox");
+							
+			return IsHovered() ? &Style.BackgroundHoveredImage : &Style.BackgroundImage; 
+		}
+		
+		FText GetMirrorOptionToolTip() const
+		{
+			const TSharedPtr<FDatabaseAssetTreeNode> AssetTreeNode = WeakAssetTreeNode.Pin();
+
+			FString TooltipString;
+						
+			TooltipString.Append(LOCTEXT("ToolTipMirrorOption", "Mirror Option: ").ToString());
+			TooltipString.Append(AssetTreeNode ? UEnum::GetDisplayValueAsText(AssetTreeNode->GetMirrorOption()).ToString() : LOCTEXT("ToolTipMirrorOption_Invalid", "Invalid").ToString());
+						
+			return FText::FromString(TooltipString);
+		}
+		
+		const FSlateBrush* GetMirrorOptionSlateBrush() const
+		{
+			if (const TSharedPtr<FDatabaseAssetTreeNode> AssetTreeNode = WeakAssetTreeNode.Pin())
+			{
+				switch (AssetTreeNode->GetMirrorOption())
+				{
+				case EPoseSearchMirrorOption::UnmirroredOnly:
+					return FAppStyle::Get().GetBrush("GraphEditor.AlignNodesRight");
+
+				case EPoseSearchMirrorOption::MirroredOnly:
+					return FAppStyle::Get().GetBrush("GraphEditor.AlignNodesLeft");
+
+				case EPoseSearchMirrorOption::UnmirroredAndMirrored:
+					return FAppStyle::Get().GetBrush("GraphEditor.AlignNodesCenter");
+				}
+			}
+		
+			return nullptr;
+		}
+		
+		TWeakPtr<FDatabaseAssetTreeNode> WeakAssetTreeNode;
+		TWeakPtr<FDatabaseViewModel> EditorViewModel;
+		TWeakPtr<SDatabaseAssetTree> SkeletonView;
+	};
 	
 	void SDatabaseAssetListItem::Construct(
 		const FArguments& InArgs,
@@ -355,18 +460,15 @@ namespace UE::PoseSearch
 					.ColorAndOpacity(this, &SDatabaseAssetListItem::GetRootMotionColorAndOpacity)
 					.ToolTipText(this, &SDatabaseAssetListItem::GetRootMotionOptionToolTip)
 				]
-				
-				// Mirror Type
+
+				// Mirror type
 				+ SHorizontalBox::Slot()
-				.Padding(2.0f, 3.0f)
 				.AutoWidth()
-				.HAlign(HAlign_Center)
 				.VAlign(VAlign_Center)
+				.HAlign(HAlign_Center)
+				.Padding(4.0f, 0.0f, 4.0f, 0.0f)
 				[
-					SNew(SImage)
-					.Image(this, &SDatabaseAssetListItem::GetMirrorOptionSlateBrush)
-					.ToolTipText(this, &SDatabaseAssetListItem::GetMirrorOptionToolTip)
-					.OnMouseButtonDown(this, &SDatabaseAssetListItem::MirrorOptionOnMouseButtonDown)
+					SNew(SMirrorTypeWidget, WeakAssetTreeNode, SkeletonView, EditorViewModel)
 				]
 
 				// Disable Reselection
@@ -766,58 +868,6 @@ namespace UE::PoseSearch
 		return LOCTEXT("NodeRootMotionDisabledToolTip", "No root motion enabled (Read only)");
 	}
 
-	const FSlateBrush* SDatabaseAssetListItem::GetMirrorOptionSlateBrush() const
-	{
-		if (const TSharedPtr<FDatabaseAssetTreeNode> AssetTreeNode = WeakAssetTreeNode.Pin())
-		{
-			// TODO: Update icons when appropriate assets become available.
-			switch (AssetTreeNode->GetMirrorOption())
-			{
-			case EPoseSearchMirrorOption::UnmirroredOnly:
-				return FAppStyle::Get().GetBrush("GraphEditor.AlignNodesRight");
-
-			case EPoseSearchMirrorOption::MirroredOnly:
-				return FAppStyle::Get().GetBrush("GraphEditor.AlignNodesLeft");
-
-			case EPoseSearchMirrorOption::UnmirroredAndMirrored:
-				return FAppStyle::Get().GetBrush("GraphEditor.AlignNodesCenter");
-			}
-		}
-		
-		return nullptr;
-	}
-
-	FText SDatabaseAssetListItem::GetMirrorOptionToolTip() const
-	{
-		const TSharedPtr<FDatabaseAssetTreeNode> AssetTreeNode = WeakAssetTreeNode.Pin();
-		return FText::FromString(LOCTEXT("ToolTipMirrorOption", "Mirror Option: ").ToString() + (AssetTreeNode ? UEnum::GetDisplayValueAsText(AssetTreeNode->GetMirrorOption()).ToString() : LOCTEXT("ToolTipMirrorOption_Invalid", "Invalid").ToString()));
-	}
-
-	FReply SDatabaseAssetListItem::MirrorOptionOnMouseButtonDown(const FGeometry& InMyGeometry, const FPointerEvent& InMouseEvent)
-	{
-		if (const TSharedPtr<FDatabaseAssetTreeNode> AssetTreeNode = WeakAssetTreeNode.Pin())
-		{
-			const TSharedPtr<FDatabaseViewModel> ViewModel = EditorViewModel.Pin();
-
-			if (InMouseEvent.GetEffectingButton() == EKeys::LeftMouseButton)
-			{
-				const FScopedTransaction Transaction(LOCTEXT("OnClickEditMirrorOptionPoseSearchDatabase", "Edit Mirror Option"));
-				
-				// Get next mirror option
-				static const TArray<EPoseSearchMirrorOption> OptionArray = { EPoseSearchMirrorOption::UnmirroredOnly, EPoseSearchMirrorOption::MirroredOnly, EPoseSearchMirrorOption::UnmirroredAndMirrored };
-				const int32 NextOption = (static_cast<int32>(ViewModel->GetMirrorOption(AssetTreeNode->SourceAssetIdx)) + 1) % OptionArray.Num();
-				
-				ViewModel->SetMirrorOption(AssetTreeNode->SourceAssetIdx, OptionArray[NextOption]);
-				
-				SkeletonView.Pin()->RefreshTreeView(false, true);
-				ViewModel->BuildSearchIndex();
-
-				return FReply::Handled();
-			}
-		}
-		return FReply::Unhandled();
-	}
-
 	FText SDatabaseAssetListItem::GetAssetEnabledToolTip() const
 	{
 		if (GetAssetEnabledChecked() == ECheckBoxState::Checked)
diff --git a/Engine/Plugins/Animation/PoseSearch/Source/Editor/Public/PoseSearchDatabaseAssetListItem.h b/Engine/Plugins/Animation/PoseSearch/Source/Editor/Public/PoseSearchDatabaseAssetListItem.h
--- a/Engine/Plugins/Animation/PoseSearch/Source/Editor/Public/PoseSearchDatabaseAssetListItem.h
+++ b/Engine/Plugins/Animation/PoseSearch/Source/Editor/Public/PoseSearchDatabaseAssetListItem.h
@@ -55,10 +55,6 @@ namespace UE::PoseSearch
 		FSlateColor GetRootMotionColorAndOpacity() const;
 		FText GetRootMotionOptionToolTip() const;
 		
-		const FSlateBrush* GetMirrorOptionSlateBrush() const;
-		FText GetMirrorOptionToolTip() const;
-		FReply MirrorOptionOnMouseButtonDown(const FGeometry& InMyGeometry, const FPointerEvent& InMouseEvent);
-		
 		FText GetDisableReselectionToolTip() const;
 		ECheckBoxState GetDisableReselectionChecked() const;
 		void OnDisableReselectionChanged(ECheckBoxState NewCheckboxState);
