Motion Matching Editor - Asset browser refreshes its asset list when the target database's schema changes.

[RN] UE-212916

diff --git a/Engine/Plugins/Animation/PoseSearch/Source/Editor/Private/PoseSearchDatabaseAssetBrowser.cpp b/Engine/Plugins/Animation/PoseSearch/Source/Editor/Private/PoseSearchDatabaseAssetBrowser.cpp
--- a/Engine/Plugins/Animation/PoseSearch/Source/Editor/Private/PoseSearchDatabaseAssetBrowser.cpp
+++ b/Engine/Plugins/Animation/PoseSearch/Source/Editor/Private/PoseSearchDatabaseAssetBrowser.cpp
@@ -33,6 +33,10 @@ void SPoseSearchDatabaseAssetBrowser::Construct(const FArguments& InArgs, TShare
 			]
 		];
 
+	// Register to be notified when properties are edited. We leverage this to refresh the browser in case the target schema changes.
+	const FCoreUObjectDelegates::FOnObjectPropertyChanged::FDelegate OnPropertyChangedDelegate = FCoreUObjectDelegates::FOnObjectPropertyChanged::FDelegate::CreateRaw(this, &SPoseSearchDatabaseAssetBrowser::OnObjectPropertyChanged);
+	OnPropertyChangedHandle = FCoreUObjectDelegates::OnObjectPropertyChanged.Add(OnPropertyChangedDelegate);
+
 	RefreshView();
 }
 
@@ -51,7 +55,7 @@ void SPoseSearchDatabaseAssetBrowser::RefreshView()
 	AssetPickerConfig.bShowPathInColumnView = true;
 	AssetPickerConfig.bShowTypeInColumnView = false;
 	AssetPickerConfig.InitialThumbnailSize = EThumbnailSize::Small;
-	
+	AssetPickerConfig.RefreshAssetViewDelegates.Add(&RefreshAssetViewDelegate);
 	AssetPickerConfig.OnShouldFilterAsset = FOnShouldFilterAsset::CreateSP(this, &SPoseSearchDatabaseAssetBrowser::OnShouldFilterAsset);
 	AssetPickerConfig.OnAssetDoubleClicked = FOnAssetSelected::CreateSP(this, &SPoseSearchDatabaseAssetBrowser::OnAssetDoubleClicked);
 
@@ -73,6 +77,12 @@ void SPoseSearchDatabaseAssetBrowser::RefreshView()
 	AssetBrowserBox->SetContent(ContentBrowserModule.Get().CreateAssetPicker(AssetPickerConfig));
 }
 
+SPoseSearchDatabaseAssetBrowser::~SPoseSearchDatabaseAssetBrowser()
+{
+	// Unregister the property modification handler
+	FCoreUObjectDelegates::OnObjectPropertyChanged.Remove(OnPropertyChangedHandle);
+}
+
 void SPoseSearchDatabaseAssetBrowser::OnAssetDoubleClicked(const FAssetData& AssetData)
 {
 	if (!AssetData.GetAsset())
@@ -136,5 +146,19 @@ bool SPoseSearchDatabaseAssetBrowser::OnShouldFilterAsset(const FAssetData& Asse
 	
 	return true;
 }
-	
+
+void SPoseSearchDatabaseAssetBrowser::OnObjectPropertyChanged(UObject* InObject, FPropertyChangedEvent& InPropertyChangedEvent) const
+{
+	if (DatabaseViewModel)
+	{
+		if (const UPoseSearchDatabase* Database = DatabaseViewModel->GetPoseSearchDatabase())
+		{
+			// Only refresh asset browser if our target schema has changed.
+			if (Database->Schema && Database->Schema == InObject)
+			{
+				RefreshAssetViewDelegate.ExecuteIfBound(true);	
+			}
+		}
+	}
+}
 }
diff --git a/Engine/Plugins/Animation/PoseSearch/Source/Editor/Private/PoseSearchDatabaseAssetBrowser.h b/Engine/Plugins/Animation/PoseSearch/Source/Editor/Private/PoseSearchDatabaseAssetBrowser.h
--- a/Engine/Plugins/Animation/PoseSearch/Source/Editor/Private/PoseSearchDatabaseAssetBrowser.h
+++ b/Engine/Plugins/Animation/PoseSearch/Source/Editor/Private/PoseSearchDatabaseAssetBrowser.h
@@ -3,6 +3,7 @@
 #pragma once
 
 #include "CoreMinimal.h"
+#include "ContentBrowserDelegates.h"
 #include "Widgets/Layout/SBox.h"
 
 struct FAssetData;
@@ -22,6 +23,7 @@ namespace UE::PoseSearch
 		
 		void RefreshView();
 		
+		virtual ~SPoseSearchDatabaseAssetBrowser() override;
 	private:
 
 		TSharedPtr<SBox> AssetBrowserBox;
@@ -29,5 +31,13 @@ namespace UE::PoseSearch
 		
 		void OnAssetDoubleClicked(const FAssetData& AssetData);
 		bool OnShouldFilterAsset(const FAssetData& AssetData);
+
+		/* We need to be able to refresh the asset list if requested (i.e. Schema changes) */
+		FRefreshAssetViewDelegate RefreshAssetViewDelegate;
+
+		/* Keep track of delegate used to listen for pose search schema changes */
+		FDelegateHandle OnPropertyChangedHandle;
+		
+		void OnObjectPropertyChanged(UObject* InObject, FPropertyChangedEvent& InPropertyChangedEvent) const;
 	};
 }
\ No newline at end of file
