Motion Matching Editor - BranchIn icon and asset picker spacing adjustments.

#rnx

diff --git a/Engine/Plugins/Animation/PoseSearch/Source/Editor/Private/PoseSearchDatabaseAssetListItem.cpp b/Engine/Plugins/Animation/PoseSearch/Source/Editor/Private/PoseSearchDatabaseAssetListItem.cpp
--- a/Engine/Plugins/Animation/PoseSearch/Source/Editor/Private/PoseSearchDatabaseAssetListItem.cpp
+++ b/Engine/Plugins/Animation/PoseSearch/Source/Editor/Private/PoseSearchDatabaseAssetListItem.cpp
@@ -333,6 +333,38 @@ namespace UE::PoseSearch
 		{
 			TSharedPtr<FDatabaseViewModel> ViewModel = EditorViewModel.Pin();
 
+			// Branch in
+			TSharedPtr<SImage> BranchInIconWidget;
+			{
+				SAssignNew(BranchInIconWidget, SImage)
+				.Image(FAppStyle::Get().GetBrush("Icons.ArrowRight"))
+				.ColorAndOpacity(FColor::Turquoise)
+				.ToolTipText(LOCTEXT("NodeBranchInTooltip", "This database item is synchronize with an external depedency and is sampled via a BranchIn notify."))
+				.Visibility_Lambda([this]()
+				{
+				   const TSharedPtr<FDatabaseViewModel> ViewModelPtr = EditorViewModel.Pin();
+    
+				   if (const UPoseSearchDatabase* Database = ViewModelPtr->GetPoseSearchDatabase())
+				   {
+					  if (const TSharedPtr<FDatabaseAssetTreeNode> AssetTreeNode = WeakAssetTreeNode.Pin())
+					  {
+						 if (const FPoseSearchDatabaseAnimationAssetBase* DatabaseAnimationAssetBase = Database->GetAnimationAssetBase(AssetTreeNode->SourceAssetIdx))
+						 {
+							if (DatabaseAnimationAssetBase->bSynchronizeWithExternalDependency)
+							{
+							   return EVisibility::Visible;
+							}
+						 }
+					  }
+				   }
+                      
+				   return EVisibility::Hidden;
+				})
+				// @note: Works under the assumption there are not hierarchy in databases, done this way to avoid having to change the TreeView to a ListView in case its needed in the future.
+				.RenderTransform(FSlateRenderTransform(1.0f, FVector2d(-8.0f, 0.0f))) 
+				.Clipping(EWidgetClipping::ClipToBoundsWithoutIntersecting);
+			}
+			
 			// Item Thumbnail
 			{
 				// Get item Icon
@@ -543,13 +575,21 @@ namespace UE::PoseSearch
 				.SizeRule(SSplitter::FractionOfParent)
 				[
 					SNew(SBorder)
-					.HAlign(HAlign_Left)
+					.HAlign(HAlign_Fill)
 					.VAlign(VAlign_Fill)
 					.BorderImage(FStyleDefaults::GetNoBrush())
 					[
 						SNew(SHorizontalBox)
 						.Clipping(EWidgetClipping::ClipToBounds)
 						+ SHorizontalBox::Slot()
+						.HAlign(HAlign_Center)
+						.VAlign(VAlign_Center)
+						.Padding(0.0f, 0.0f, 0.0f, 0.0f)
+						.AutoWidth()
+						[
+						   BranchInIconWidget.ToSharedRef()
+						]
+						+ SHorizontalBox::Slot()
 						.AutoWidth()
 						.Padding(0.0f, 0.0f, 10.0f, 0.0f)
 						.HAlign(HAlign_Left)
@@ -558,8 +598,10 @@ namespace UE::PoseSearch
 							AssetThumbnailOverlay.ToSharedRef()
 						]
 						+ SHorizontalBox::Slot()
-						.FillWidth(1.0f)
-						.VAlign(VAlign_Center)
+						.AutoWidth()
+						.FillWidth(0.7f)
+						.Padding(0.0f, 0.0f, 30.0f, 0.0f)
+						.HAlign(HAlign_Fill)
 						[
 							AssetPickerWidget.ToSharedRef()
 						]
