Motion Matching Editor - Database item asset picker now filter based on if the asset is compatible with the database's schema.

[RN] minor

diff --git a/Engine/Plugins/Animation/PoseSearch/Source/Editor/Private/PoseSearchDatabaseAssetBrowser.cpp b/Engine/Plugins/Animation/PoseSearch/Source/Editor/Private/PoseSearchDatabaseAssetBrowser.cpp
--- a/Engine/Plugins/Animation/PoseSearch/Source/Editor/Private/PoseSearchDatabaseAssetBrowser.cpp
+++ b/Engine/Plugins/Animation/PoseSearch/Source/Editor/Private/PoseSearchDatabaseAssetBrowser.cpp
@@ -16,6 +16,10 @@
 #include "PoseSearch/PoseSearchSchema.h"
 #include "Subsystems/AssetEditorSubsystem.h"
 #include "Editor.h"
+#include "PoseSearchDatabaseEditorUtils.h"
+#include "Filters/SBasicFilterBar.h"
+
+#define LOCTEXT_NAMESPACE "PoseSearchDatabaseAssetBrowser"
 
 namespace UE::PoseSearch
 {
@@ -58,6 +62,8 @@ void SPoseSearchDatabaseAssetBrowser::RefreshView()
 	AssetPickerConfig.RefreshAssetViewDelegates.Add(&RefreshAssetViewDelegate);
 	AssetPickerConfig.OnShouldFilterAsset = FOnShouldFilterAsset::CreateSP(this, &SPoseSearchDatabaseAssetBrowser::OnShouldFilterAsset);
 	AssetPickerConfig.OnAssetDoubleClicked = FOnAssetSelected::CreateSP(this, &SPoseSearchDatabaseAssetBrowser::OnAssetDoubleClicked);
+	AssetPickerConfig.AssetShowWarningText = LOCTEXT("NoAssets_Warning", "No Assets found. No compatible assets with the database's schema where found. Ensure your assets' skeleton matches a skeleton from the database's schema.");
+	AssetPickerConfig.bCanShowDevelopersFolder = true;
 
 	// Hide all asset registry columns by default (we only really want the name and path)
 	const UObject* AnimSequenceDefaultObject = UAnimSequence::StaticClass()->GetDefaultObject();
@@ -106,22 +112,7 @@ bool SPoseSearchDatabaseAssetBrowser::OnShouldFilterAsset(const FAssetData& Asse
 
 	if (DatabaseViewModel)
 	{
-		if (const UPoseSearchDatabase* Database = DatabaseViewModel->GetPoseSearchDatabase())
-		{
-			if (Database->Schema)
-			{
-				TArray<FPoseSearchRoledSkeleton> RoledSkeletons = Database->Schema->GetRoledSkeletons();
-
-				for (const FPoseSearchRoledSkeleton& RoledSkeleton : RoledSkeletons)
-				{
-					if (RoledSkeleton.Skeleton && RoledSkeleton.Skeleton->IsCompatibleForEditor(AssetData))
-					{
-						// We found a compatible skeleton in the schema.
-						bAssetHasCompatibleSkeleton = true;
-					}
-				}
-			}
-		}
+		bAssetHasCompatibleSkeleton = FPoseSearchEditorUtils::IsAssetCompatibleWithDatabase(DatabaseViewModel->GetPoseSearchDatabase(), AssetData);
 	}
 	
 	if (AssetData.GetClass()->IsChildOf(UAnimSequence::StaticClass()) && bAssetHasCompatibleSkeleton)
@@ -146,7 +137,7 @@ bool SPoseSearchDatabaseAssetBrowser::OnShouldFilterAsset(const FAssetData& Asse
 	
 	return true;
 }
-
+	
 void SPoseSearchDatabaseAssetBrowser::OnObjectPropertyChanged(UObject* InObject, FPropertyChangedEvent& InPropertyChangedEvent) const
 {
 	if (DatabaseViewModel)
@@ -162,3 +153,5 @@ void SPoseSearchDatabaseAssetBrowser::OnObjectPropertyChanged(UObject* InObject,
 	}
 }
 }
+
+#undef LOCTEXT_NAMESPACE
\ No newline at end of file
diff --git a/Engine/Plugins/Animation/PoseSearch/Source/Editor/Private/PoseSearchDatabaseAssetListItem.cpp b/Engine/Plugins/Animation/PoseSearch/Source/Editor/Private/PoseSearchDatabaseAssetListItem.cpp
--- a/Engine/Plugins/Animation/PoseSearch/Source/Editor/Private/PoseSearchDatabaseAssetListItem.cpp
+++ b/Engine/Plugins/Animation/PoseSearch/Source/Editor/Private/PoseSearchDatabaseAssetListItem.cpp
@@ -21,10 +21,12 @@
 #include "Misc/TransactionObjectEvent.h"
 #include "PoseSearch/PoseSearchDatabase.h"
 #include "PoseSearchDatabaseAssetTree.h"
+#include "PoseSearchDatabaseEditorUtils.h"
 #include "PoseSearchDatabaseViewModel.h"
 #include "PropertyCustomizationHelpers.h"
 #include "ScopedTransaction.h"
 #include "SPositiveActionButton.h"
+#include "PoseSearch/PoseSearchSchema.h"
 #include "Styling/AppStyle.h"
 #include "Styling/StyleColors.h"
 #include "Subsystems/AssetEditorSubsystem.h"
@@ -449,6 +451,15 @@ namespace UE::PoseSearch
 					.IsEnabled(this, &SDatabaseAssetListItem::GetAssetPickerIsEnabled)
 					.ObjectPath(this, &SDatabaseAssetListItem::GetAssetPickerObjectPath)
 					.OnObjectChanged(this, &SDatabaseAssetListItem::OnAssetPickerObjectChanged)
+					.OnShouldFilterAsset_Lambda([this](const FAssetData& InAssetData)
+					{
+						if (EditorViewModel.IsValid())
+						{
+							return !FPoseSearchEditorUtils::IsAssetCompatibleWithDatabase(EditorViewModel.Pin()->GetPoseSearchDatabase(), InAssetData);
+						}
+						
+						return true;
+					})
 					.CustomContentSlot()
 					[
 						// Display warning below picked asset.
diff --git a/Engine/Plugins/Animation/PoseSearch/Source/Editor/Private/PoseSearchDatabaseEditorUtils.cpp b/Engine/Plugins/Animation/PoseSearch/Source/Editor/Private/PoseSearchDatabaseEditorUtils.cpp
new file mode 100644
--- /dev/null
+++ b/Engine/Plugins/Animation/PoseSearch/Source/Editor/Private/PoseSearchDatabaseEditorUtils.cpp
@@ -0,0 +1,31 @@
+﻿// Copyright Epic Games, Inc. All Rights Reserved.
+
+#include "PoseSearchDatabaseEditorUtils.h"
+
+#include "PoseSearch/PoseSearchDatabase.h"
+#include "PoseSearch/PoseSearchSchema.h"
+
+namespace UE::PoseSearch
+{
+	bool FPoseSearchEditorUtils::IsAssetCompatibleWithDatabase(const UPoseSearchDatabase* InDatabase, const FAssetData& InAssetData)
+	{
+		if (InDatabase)
+		{
+			if (InDatabase->Schema)
+			{
+				TArray<FPoseSearchRoledSkeleton> RoledSkeletons = InDatabase->Schema->GetRoledSkeletons();
+
+				for (const FPoseSearchRoledSkeleton& RoledSkeleton : RoledSkeletons)
+				{
+					if (RoledSkeleton.Skeleton && RoledSkeleton.Skeleton->IsCompatibleForEditor(InAssetData))
+					{
+						// We found a compatible skeleton in the schema.
+						return true;
+					}
+				}
+			}
+		}
+		
+		return false;
+	}
+}
diff --git a/Engine/Plugins/Animation/PoseSearch/Source/Editor/Private/PoseSearchDatabaseEditorUtils.h b/Engine/Plugins/Animation/PoseSearch/Source/Editor/Private/PoseSearchDatabaseEditorUtils.h
new file mode 100644
--- /dev/null
+++ b/Engine/Plugins/Animation/PoseSearch/Source/Editor/Private/PoseSearchDatabaseEditorUtils.h
@@ -0,0 +1,16 @@
+﻿// Copyright Epic Games, Inc. All Rights Reserved.
+
+#pragma once
+
+#include "CoreMinimal.h"
+#include "AssetRegistry/AssetData.h"
+
+class UPoseSearchDatabase;
+
+namespace UE::PoseSearch
+{
+	struct FPoseSearchEditorUtils
+	{
+		static bool IsAssetCompatibleWithDatabase(const UPoseSearchDatabase* InDatabase, const FAssetData & InAssetData);
+	};
+}
