motion matching - exposing UPoseSearchDatabase GetNumAnimationAssets(), GetAnimationAsset(int32 Index) to blueprint

diff --git a/Engine/Plugins/Animation/PoseSearch/Source/Runtime/Private/PoseSearchDatabase.cpp b/Engine/Plugins/Animation/PoseSearch/Source/Runtime/Private/PoseSearchDatabase.cpp
--- a/Engine/Plugins/Animation/PoseSearch/Source/Runtime/Private/PoseSearchDatabase.cpp
+++ b/Engine/Plugins/Animation/PoseSearch/Source/Runtime/Private/PoseSearchDatabase.cpp
@@ -796,6 +796,23 @@ bool UPoseSearchDatabase::Contains(const UObject* Object) const
 	return false;
 }
 
+int32 UPoseSearchDatabase::GetNumAnimationAssets() const
+{
+	return AnimationAssets.Num();
+}
+
+UObject* UPoseSearchDatabase::GetAnimationAsset(int32 Index) const
+{
+	if (AnimationAssets.IsValidIndex(Index))
+	{
+		if (const FPoseSearchDatabaseAnimationAssetBase* AnimationAssetBase = AnimationAssets[Index].GetPtr<FPoseSearchDatabaseAnimationAssetBase>())
+		{
+			return AnimationAssetBase->GetAnimationAsset();
+		}
+	}
+	return nullptr;
+}
+
 #if WITH_EDITOR
 void UPoseSearchDatabase::SynchronizeWithExternalDependencies()
 {
diff --git a/Engine/Plugins/Animation/PoseSearch/Source/Runtime/Public/PoseSearch/PoseSearchDatabase.h b/Engine/Plugins/Animation/PoseSearch/Source/Runtime/Public/PoseSearch/PoseSearchDatabase.h
--- a/Engine/Plugins/Animation/PoseSearch/Source/Runtime/Public/PoseSearch/PoseSearchDatabase.h
+++ b/Engine/Plugins/Animation/PoseSearch/Source/Runtime/Public/PoseSearch/PoseSearchDatabase.h
@@ -409,6 +409,12 @@ public:
 
 	bool Contains(const UObject* Object) const;
 
+	UFUNCTION(BlueprintPure, Category = "Animation|Pose Search|Experimental", meta = (BlueprintThreadSafe))
+	int32 GetNumAnimationAssets() const;
+
+	UFUNCTION(BlueprintPure, Category = "Animation|Pose Search|Experimental", meta = (BlueprintThreadSafe))
+	UObject* GetAnimationAsset(int32 Index) const;
+
 #if WITH_EDITOR
 	int32 GetNumberOfPrincipalComponents() const;
 
