blendstack - stitches support

diff --git a/Engine/Plugins/Animation/BlendStack/Source/Runtime/Private/AnimNode_BlendStack.cpp b/Engine/Plugins/Animation/BlendStack/Source/Runtime/Private/AnimNode_BlendStack.cpp
--- a/Engine/Plugins/Animation/BlendStack/Source/Runtime/Private/AnimNode_BlendStack.cpp
+++ b/Engine/Plugins/Animation/BlendStack/Source/Runtime/Private/AnimNode_BlendStack.cpp
@@ -6,6 +6,7 @@
 #include "Animation/AnimInstanceProxy.h"
 #include "Animation/AnimMontage.h"
 #include "Animation/AnimNode_Inertialization.h"
+#include "Animation/AnimPoseSearchProvider.h"
 #include "Animation/AnimSequence.h"
 #include "Animation/BlendSpace.h"
 #include "BlendStack/AnimNode_BlendStackInput.h"
@@ -56,6 +57,17 @@ void FBlendStackAnimPlayer::Initialize(const FAnimationInitializeContext& Contex
 		}
 	}
 
+	// if a negative AccumulatedTime is requested, we just delay the playback of the animation asset by -AccumulatedTime
+	// this is useful when we want to land exacly at the selected frame for example during "stitches": 
+	// consider the case where you want to play an animation A at a time n, A[n], after a blend of b seconds,
+	// eventually provided by a stitch search. That'd mean to start blending in into the animation A at time n-b, 
+	// that can be negative, and since the pose A[n-b] would be invalid we just use A[0] delayed by n-b seconds	
+	if (AccumulatedTime < 0.f)
+	{
+		ActivationDelay = FMath::Max(ActivationDelay - AccumulatedTime, 0.f);
+		AccumulatedTime = 0.f;
+	}
+
 	BlendOption = InBlendOption;
 
 	TotalBlendInTime = BlendTime;
@@ -384,6 +396,10 @@ float FBlendStackAnimPlayer::GetBlendInPercentage() const
 {
 	if (FMath::IsNearlyZero(TotalBlendInTime))
 	{
+		if (TimeToActivation > 0.f)
+		{
+			return 0.f;
+		}
 		return 1.f;
 	}
 
@@ -796,6 +812,70 @@ static void RequestInertialBlend(const FAnimationUpdateContext& Context, float B
 void FAnimNode_BlendStack_Standalone::BlendTo(const FAnimationUpdateContext& Context, UAnimationAsset* AnimationAsset, float AccumulatedTime, bool bLoop,
 	bool bMirrored, UMirrorDataTable* MirrorDataTable, float BlendTime, const UBlendProfile* BlendProfile, EAlphaBlendOption BlendOption, bool bUseInertialBlend,
 	const FVector& BlendParameters, float PlayRate, float ActivationDelay, FName GroupName, EAnimGroupRole::Type GroupRole, EAnimSyncMethod GroupMethod)
+{
+	using namespace UE::Anim;
+
+	bool bNeedToBlendTo = true;
+	if (StitchDatabase)
+	{
+		if (IPoseSearchProvider* PoseSearchProvider = IPoseSearchProvider::Get())
+		{
+			// looking for an animation stitch from the StitchDatabase that will connect, in BlendTime seconds, 
+			// the currently playing animation pose to the pose from AnimationAsset at AccumulatedTime + BlendTime
+			const UObject* AssetToSearch = StitchDatabase.Get();
+
+			IPoseSearchProvider::FSearchPlayingAsset PlayingAsset;
+			PlayingAsset.Asset = GetAnimAsset();
+			PlayingAsset.AccumulatedTime = GetAccumulatedTime();
+
+			IPoseSearchProvider::FSearchFutureAsset FutureAsset;
+			FutureAsset.Asset = AnimationAsset;
+			FutureAsset.AccumulatedTime = AccumulatedTime + BlendTime;
+			FutureAsset.IntervalTime = BlendTime;
+
+			const IPoseSearchProvider::FSearchResult SearchResult = PoseSearchProvider->Search(Context, MakeArrayView(&AssetToSearch, 1), PlayingAsset, FutureAsset);
+			if (UAnimationAsset* StitchAnimationAsset = Cast<UAnimationAsset>(SearchResult.SelectedAsset))
+			{
+				if (SearchResult.Dissimilarity <= StitchBlendMaxCost)
+				{
+					// blend to the selected animation stitch
+					InternalBlendTo(Context, StitchAnimationAsset, SearchResult.TimeOffsetSeconds, false, SearchResult.bMirrored, MirrorDataTable,
+						StitchBlendTime, BlendProfile, BlendOption, bUseInertialBlend, BlendParameters, SearchResult.WantedPlayRate, ActivationDelay, GroupName, GroupRole, GroupMethod);
+
+					// blend with an ActivationDelay of BlendTime - StitchBlendTime + ActivationDelay seconds
+					// to the AnimationAsset at AccumulatedTime + BlendTime - StitchBlendTime seconds in the future,
+					// so at BlendTime seconds ahead the AnimationAsset is playing the fully blended in pose at AccumulatedTime + BlendTime
+					InternalBlendTo(Context, AnimationAsset, AccumulatedTime + BlendTime - StitchBlendTime, bLoop, bMirrored, MirrorDataTable,
+						StitchBlendTime, BlendProfile, BlendOption, bUseInertialBlend, BlendParameters, PlayRate, BlendTime - StitchBlendTime + ActivationDelay, GroupName, GroupRole, GroupMethod);
+
+					bNeedToBlendTo = false;
+				}
+				else
+				{
+					UE_LOG(LogBlendStack, Display, TEXT("FAnimNode_BlendStack_Standalone::BlendTo StitchDatabase '%s' search cost is %f, above StitchBlendMaxCost %f. Defaulting to regular blend"), *GetNameSafe(StitchDatabase), SearchResult.Dissimilarity, StitchBlendMaxCost);
+				}
+			}
+			else
+			{
+				UE_LOG(LogBlendStack, Error, TEXT("FAnimNode_BlendStack_Standalone::BlendTo cannot use StitchDatabase '%s', because of missing IPoseSearchProvider::Search couldn't select a StitchAnimationAsset. Defaulting to regular blend"), *GetNameSafe(StitchDatabase));
+			}
+		}
+		else
+		{
+			UE_LOG(LogBlendStack, Error, TEXT("FAnimNode_BlendStack_Standalone::BlendTo cannot use StitchDatabase '%s', because of missing IPoseSearchProvider (is PoseSearch plugin enabled?). Defaulting to regular blend"), *GetNameSafe(StitchDatabase));
+		}
+	}
+
+	if (bNeedToBlendTo)
+	{
+		InternalBlendTo(Context, AnimationAsset, AccumulatedTime, bLoop, bMirrored, MirrorDataTable,
+			BlendTime, BlendProfile, BlendOption, bUseInertialBlend, BlendParameters, PlayRate, ActivationDelay, GroupName, GroupRole, GroupMethod);
+	}
+}
+
+void FAnimNode_BlendStack_Standalone::InternalBlendTo(const FAnimationUpdateContext& Context, UAnimationAsset* AnimationAsset, float AccumulatedTime, bool bLoop,
+	bool bMirrored, UMirrorDataTable* MirrorDataTable, float BlendTime, const UBlendProfile* BlendProfile, EAlphaBlendOption BlendOption, bool bUseInertialBlend,
+	const FVector& BlendParameters, float PlayRate, float ActivationDelay, FName GroupName, EAnimGroupRole::Type GroupRole, EAnimSyncMethod GroupMethod)
 {
 	const bool bBlendStackIsEmpty = AnimPlayers.IsEmpty();
 
@@ -813,7 +893,7 @@ void FAnimNode_BlendStack_Standalone::BlendTo(const FAnimationUpdateContext& Con
 
 	// If we don't add a new player, re-use the same graph...
 	int32 NewSamplePoseLinkIndex = CurrentSamplePoseLink;
-	if (!bBlendStackIsEmpty &&  AnimPlayers[0].GetBlendInPercentage() < 1.0f &&
+	if (!bBlendStackIsEmpty && AnimPlayers[0].GetBlendInPercentage() < 1.0f &&
 		AnimPlayers[0].GetCurrentBlendInTime() < MaxBlendInTimeToOverrideAnimation &&
 		FMath::IsNearlyEqual(AnimPlayers[0].GetTimeToActivation(), ActivationDelay))
 	{
diff --git a/Engine/Plugins/Animation/BlendStack/Source/Runtime/Public/BlendStack/AnimNode_BlendStack.h b/Engine/Plugins/Animation/BlendStack/Source/Runtime/Public/BlendStack/AnimNode_BlendStack.h
--- a/Engine/Plugins/Animation/BlendStack/Source/Runtime/Public/BlendStack/AnimNode_BlendStack.h
+++ b/Engine/Plugins/Animation/BlendStack/Source/Runtime/Public/BlendStack/AnimNode_BlendStack.h
@@ -119,6 +119,19 @@ struct BLENDSTACK_API FAnimNode_BlendStack_Standalone : public FAnimNode_AssetPl
 	UPROPERTY(EditAnywhere, Category = Settings)
 	bool bShouldFilterNotifies = false;
 	
+	// database used to search for an animation stitch to use as blend
+	UPROPERTY(EditAnywhere, BlueprintReadWrite, Category = "Stitch|Experimental", meta = (PinHiddenByDefault))
+	TObjectPtr<UObject> StitchDatabase;
+
+	// blend time in seconds used to blend into and out from a stitch animation
+	UPROPERTY(EditAnywhere, BlueprintReadWrite, Category = "Stitch|Experimental", meta = (PinHiddenByDefault))
+	float StitchBlendTime = 0.1f;
+
+	// if the cost from searching StitchDatabase is above StitchBlendMaxCost, clend stack will perform a regular blend,
+	// and not using the returned stitch animation as blend
+	UPROPERTY(EditAnywhere, BlueprintReadWrite, Category = "Stitch|Experimental", meta = (PinHiddenByDefault))
+	float StitchBlendMaxCost = 100.f;
+
 	// FAnimNode_Base interface
 	virtual void Evaluate_AnyThread(FPoseContext& Output) override;
 	virtual void Initialize_AnyThread(const FAnimationInitializeContext& Context) override;
@@ -145,6 +158,12 @@ struct BLENDSTACK_API FAnimNode_BlendStack_Standalone : public FAnimNode_AssetPl
 	int32 GetMaxActiveBlends() const { return MaxActiveBlends; }
 
 protected:
+	void InternalBlendTo(const FAnimationUpdateContext& Context, UAnimationAsset* AnimationAsset, float AccumulatedTime, bool bLoop, 
+		bool bMirrored, UMirrorDataTable* MirrorDataTable, float BlendTime,
+		const UBlendProfile* BlendProfile, EAlphaBlendOption BlendOption, 
+		bool bUseInertialBlend, const FVector& BlendParameters, float PlayRate, float ActivationDelay,
+		FName GroupName, EAnimGroupRole::Type GroupRole, EAnimSyncMethod Method);
+
 	static void BlendWithPose(FAnimationPoseData& InOutPoseData, const FAnimationPoseData& OtherPoseData, const float InOutPoseWeight);
 	static void BlendWithPosePerBone(FAnimationPoseData& InOutPoseData, const FAnimationPoseData& OtherPoseData, TConstArrayView<float> OtherPoseWeights);
 
diff --git a/Engine/Plugins/Animation/PoseSearch/Source/Runtime/Private/PoseSearchLibrary.cpp b/Engine/Plugins/Animation/PoseSearch/Source/Runtime/Private/PoseSearchLibrary.cpp
--- a/Engine/Plugins/Animation/PoseSearch/Source/Runtime/Private/PoseSearchLibrary.cpp
+++ b/Engine/Plugins/Animation/PoseSearch/Source/Runtime/Private/PoseSearchLibrary.cpp
@@ -589,7 +589,7 @@
 	TArray<UObject*> AssetsToSearch,
 	const FName PoseHistoryName,
 	const FPoseSearchContinuingProperties ContinuingProperties,
-	FPoseSearchFutureProperties Future,
+	const FPoseSearchFutureProperties Future,
 	FPoseSearchBlueprintResult& Result,
 	const int32 DebugSessionUniqueIdentifier)
 {
@@ -613,6 +613,7 @@
 	TArray<UObject*> AssetsToSearch,
 	const FName PoseHistoryName,
 	const FPoseSearchContinuingProperties ContinuingProperties,
+	const FPoseSearchFutureProperties Future,
 	FPoseSearchBlueprintResult& Result,
 	const int32 DebugSessionUniqueIdentifier)
 {
@@ -632,13 +633,13 @@
 	}
 
 	TArray<const UObject*>& AssetsToSearchConst = reinterpret_cast<TArray<const UObject*>&>(AssetsToSearch);
-	MotionMatch(AnimInstances, Roles, AssetsToSearchConst, PoseHistoryName, ContinuingProperties, FPoseSearchFutureProperties(), Result, DebugSessionUniqueIdentifier);
+	MotionMatch(AnimInstances, Roles, AssetsToSearchConst, PoseHistoryName, ContinuingProperties, Future, Result, DebugSessionUniqueIdentifier);
 }
 
 void UPoseSearchLibrary::MotionMatch(
-	TArrayView<UAnimInstance*> AnimInstances,
-	TArrayView<const UE::PoseSearch::FRole> Roles,
-	TArrayView<const UObject*> AssetsToSearch,
+	const TArrayView<UAnimInstance*> AnimInstances,
+	const TArrayView<const UE::PoseSearch::FRole> Roles,
+	const TArrayView<const UObject*> AssetsToSearch,
 	const FName PoseHistoryName,
 	const FPoseSearchContinuingProperties& ContinuingProperties,
 	const FPoseSearchFutureProperties& Future,
@@ -679,37 +680,82 @@
 
 	FMemMark Mark(FMemStack::Get());
 
-#if ENABLE_DRAW_DEBUG && ENABLE_ANIM_DEBUG
-	TArray<FColor, TInlineAllocator<PreallocatedRolesNum, TMemStackAllocator<>>> HistoryCollectorColors;
-#endif // ENABLE_DRAW_DEBUG && ENABLE_ANIM_DEBUG
-
-	// MemStackPoseHistories will hold future poses to match AssetSamplerBase (at FutureAnimationStartTime) TimeToFutureAnimationStart seconds in the future
-	TArray<FMemStackPoseHistory, TInlineAllocator<PreallocatedRolesNum, TMemStackAllocator<>>> MemStackPoseHistories;
+	TArray<const UE::PoseSearch::IPoseHistory*, TInlineAllocator<PreallocatedRolesNum, TMemStackAllocator<>>> PoseHistories;
 	for (UAnimInstance* AnimInstance : AnimInstances)
 	{
 		if (const FAnimNode_PoseSearchHistoryCollector_Base* PoseHistoryNode = FindPoseHistoryNode(PoseHistoryName, AnimInstance))
 		{
-			MemStackPoseHistories.AddDefaulted_GetRef().Init(&PoseHistoryNode->GetPoseHistory());
-
-#if ENABLE_DRAW_DEBUG && ENABLE_ANIM_DEBUG
-#if WITH_EDITORONLY_DATA
-			HistoryCollectorColors.Emplace(PoseHistoryNode->DebugColor.ToFColor(true));
-#else // WITH_EDITORONLY_DATA
-			HistoryCollectorColors.Emplace(FColor::Red);
-#endif // WITH_EDITORONLY_DATA
-#endif // ENABLE_DRAW_DEBUG && ENABLE_ANIM_DEBUG && WITH_EDITORONLY_DATA
+			PoseHistories.Add(&PoseHistoryNode->GetPoseHistory());
 		}
 	}
 
-	if (MemStackPoseHistories.Num() != AnimInstances.Num())
+	if (PoseHistories.Num() != AnimInstances.Num())
 	{
 		UE_LOG(LogPoseSearch, Error, TEXT("UPoseSearchLibrary::MotionMatch - Couldn't find pose history with name '%s'"), *PoseHistoryName.ToString());
 		return;
 	}
 
+	const FSearchResult SearchResult = MotionMatch(AnimInstances, Roles, PoseHistories, AssetsToSearch, ContinuingProperties, Future, DebugSessionUniqueIdentifier);
+	if (SearchResult.IsValid())
+	{
+		const UPoseSearchDatabase* Database = SearchResult.Database.Get();
+		const FSearchIndexAsset* SearchIndexAsset = SearchResult.GetSearchIndexAsset();
+		if (const FPoseSearchDatabaseAnimationAssetBase* DatabaseAsset = Database->GetAnimationAssetBase(*SearchIndexAsset))
+		{
+			Result.SelectedAnimation = DatabaseAsset->GetAnimationAsset();
+			Result.SelectedTime = SearchResult.AssetTime;
+			Result.bLoop = SearchIndexAsset->IsLooping();
+			Result.bIsMirrored = SearchIndexAsset->IsMirrored();
+			Result.BlendParameters = SearchIndexAsset->GetBlendParameters();
+			Result.SelectedDatabase = Database;
+			Result.SearchCost = SearchResult.PoseCost.GetTotalCost();
+			
+			// figuring out the WantedPlayRate
+			Result.WantedPlayRate = 1.f;
+			if (Future.Animation && Future.IntervalTime > 0.f)
+			{
+				if (const UPoseSearchFeatureChannel_PermutationTime* PermutationTimeChannel = Database->Schema->FindFirstChannelOfType<UPoseSearchFeatureChannel_PermutationTime>())
+				{
+					const FSearchIndex& SearchIndex = Database->GetSearchIndex();
+					if (!SearchIndex.IsValuesEmpty())
+					{
+						TConstArrayView<float> ResultData = Database->GetSearchIndex().GetPoseValues(SearchResult.PoseIdx);
+						const float ActualIntervalTime = PermutationTimeChannel->GetPermutationTime(ResultData);
+						Result.WantedPlayRate = ActualIntervalTime / Future.IntervalTime;
+					}
+				}
+			}
+		}
+	}
+}
+	
+UE::PoseSearch::FSearchResult UPoseSearchLibrary::MotionMatch(
+	const TArrayView<UAnimInstance*> AnimInstances,
+	const TArrayView<const UE::PoseSearch::FRole> Roles,
+	const TArrayView<const UE::PoseSearch::IPoseHistory*> PoseHistories, 
+	const TArrayView<const UObject*> AssetsToSearch,
+	const FPoseSearchContinuingProperties& ContinuingProperties,
+	const FPoseSearchFutureProperties& Future,
+	const int32 DebugSessionUniqueIdentifier)
+{
+	check(!AnimInstances.IsEmpty() && AnimInstances.Num() == Roles.Num() && AnimInstances.Num() == PoseHistories.Num());
+
+	using namespace UE::PoseSearch;
+
+	FSearchResult SearchResult;
+
+	FMemMark Mark(FMemStack::Get());
+
+	TArray<const UE::PoseSearch::IPoseHistory*, TInlineAllocator<PreallocatedRolesNum, TMemStackAllocator<>>> InternalPoseHistories;
+	InternalPoseHistories = PoseHistories;
+
+	// MemStackPoseHistories will hold future poses to match AssetSamplerBase (at FutureAnimationStartTime) TimeToFutureAnimationStart seconds in the future
+	TArray<FMemStackPoseHistory, TInlineAllocator<PreallocatedRolesNum, TMemStackAllocator<>>> MemStackPoseHistories;
 	float FutureIntervalTime = Future.IntervalTime;
 	if (Future.Animation)
 	{
+		MemStackPoseHistories.SetNum(InternalPoseHistories.Num());
+
 		float FutureAnimationTime = Future.AnimationTime;
 		if (FutureAnimationTime < FiniteDelta)
 		{
@@ -726,10 +772,12 @@
 
 		for (int32 RoleIndex = 0; RoleIndex < Roles.Num(); ++RoleIndex)
 		{
+			MemStackPoseHistories[RoleIndex].Init(InternalPoseHistories[RoleIndex]);
+
 			// extracting 2 poses to be able to calculate velocities
 			FCSPose<FCompactPose> ComponentSpacePose;
 			FCompactPose Pose;
-			Pose.SetBoneContainer(&AnimInstances[0]->GetRequiredBonesOnAnyThread());
+			Pose.SetBoneContainer(&AnimInstances[RoleIndex]->GetRequiredBonesOnAnyThread());
 
 			// @todo: add input BlendParameters to support sampling FutureAnimation blendspaces and support for multi character
 			const UAnimationAsset* AnimationAsset = Cast<UAnimationAsset>(Future.Animation);
@@ -759,94 +807,21 @@
 #if ENABLE_DRAW_DEBUG && ENABLE_ANIM_DEBUG
 			if (FAnimInstanceProxy* AnimInstanceProxy = UAnimInstanceProxyProvider::GetAnimInstanceProxy(AnimInstances[RoleIndex]))
 			{
-				MemStackPoseHistories[RoleIndex].DebugDraw(*AnimInstanceProxy, HistoryCollectorColors[RoleIndex]);
+				MemStackPoseHistories[RoleIndex].DebugDraw(*AnimInstanceProxy, FColor::Orange);
 			}
 #endif // ENABLE_DRAW_DEBUG && ENABLE_ANIM_DEBUG
+
+			InternalPoseHistories[RoleIndex] = MemStackPoseHistories[RoleIndex].GetThisOrPoseHistory();
 		}
 	}
 
-	TArray<const UE::PoseSearch::IPoseHistory*, TInlineAllocator<PreallocatedRolesNum, TMemStackAllocator<>>> PoseHistories;
-	for (const FMemStackPoseHistory& MemStackPoseHistory : MemStackPoseHistories)
-	{
-		PoseHistories.Add(MemStackPoseHistory.GetThisOrPoseHistory());
-	}
-
-	const FSearchResult SearchResult = MotionMatch(AnimInstances, Roles, PoseHistories, AssetsToSearch, ContinuingProperties, DebugSessionUniqueIdentifier, FutureIntervalTime);
-	if (SearchResult.IsValid())
-	{
-		const UPoseSearchDatabase* Database = SearchResult.Database.Get();
-		const FSearchIndexAsset* SearchIndexAsset = SearchResult.GetSearchIndexAsset();
-		if (const FPoseSearchDatabaseAnimationAssetBase* DatabaseAsset = Database->GetAnimationAssetBase(*SearchIndexAsset))
-		{
-			Result.SelectedAnimation = DatabaseAsset->GetAnimationAsset();
-			Result.SelectedTime = SearchResult.AssetTime;
-			Result.bLoop = SearchIndexAsset->IsLooping();
-			Result.bIsMirrored = SearchIndexAsset->IsMirrored();
-			Result.BlendParameters = SearchIndexAsset->GetBlendParameters();
-			Result.SelectedDatabase = Database;
-			Result.SearchCost = SearchResult.PoseCost.GetTotalCost();
-			
-			// figuring out the WantedPlayRate
-			Result.WantedPlayRate = 1.f;
-			if (Future.Animation)
-			{
-				if (const UPoseSearchFeatureChannel_PermutationTime* PermutationTimeChannel = Database->Schema->FindFirstChannelOfType<UPoseSearchFeatureChannel_PermutationTime>())
-				{
-					const FSearchIndex& SearchIndex = Database->GetSearchIndex();
-					if (!SearchIndex.IsValuesEmpty())
-					{
-						TConstArrayView<float> ResultData = Database->GetSearchIndex().GetPoseValues(SearchResult.PoseIdx);
-						const float ActualIntervalTime = PermutationTimeChannel->GetPermutationTime(ResultData);
-						Result.WantedPlayRate = ActualIntervalTime / FutureIntervalTime;
-					}
-				}
-			}
-		}
-	}
-}
-
-UE::PoseSearch::FSearchResult UPoseSearchLibrary::MotionMatch(
-	const FAnimationBaseContext& Context,
-	TArrayView<const UObject*> AssetsToSearch,
-	const FPoseSearchContinuingProperties& ContinuingProperties)
-{
-	using namespace UE::PoseSearch;
-
-	const IPoseHistory* PoseHistory = nullptr;
-	if (IPoseHistoryProvider* PoseHistoryProvider = Context.GetMessage<IPoseHistoryProvider>())
-	{
-		PoseHistory = &PoseHistoryProvider->GetPoseHistory();
-	}
-
-	UAnimInstance* AnimInstance = Cast<UAnimInstance>(Context.AnimInstanceProxy->GetAnimInstanceObject());
-	check(AnimInstance);
-
-	return MotionMatch(MakeArrayView(&AnimInstance, 1), MakeArrayView(&DefaultRole, 1), MakeArrayView(&PoseHistory, 1), AssetsToSearch, ContinuingProperties, Context.GetCurrentNodeId());
-}
-	
-UE::PoseSearch::FSearchResult UPoseSearchLibrary::MotionMatch(
-	TArrayView<UAnimInstance*> AnimInstances,
-	TArrayView<const UE::PoseSearch::FRole> Roles,
-	TArrayView<const UE::PoseSearch::IPoseHistory*> PoseHistories, 
-	TArrayView<const UObject*> AssetsToSearch,
-	const FPoseSearchContinuingProperties& ContinuingProperties,
-	const int32 DebugSessionUniqueIdentifier,
-	float DesiredPermutationTimeOffset)
-{
-	check(!AnimInstances.IsEmpty() && AnimInstances.Num() == Roles.Num() && AnimInstances.Num() == PoseHistories.Num());
-
-	using namespace UE::PoseSearch;
-
-	FSearchResult SearchResult;
-
-	FMemMark Mark(FMemStack::Get());
 	FSearchResult ReconstructedPreviousSearchResult;
-	FSearchContext SearchContext(DesiredPermutationTimeOffset, nullptr, ReconstructedPreviousSearchResult);
+	FSearchContext SearchContext(FutureIntervalTime, nullptr, ReconstructedPreviousSearchResult);
 
 	// @todo: all assets in AssetsToSearch should have a consistent Roles requirements, or else the search will throw an error!
 	for (int32 RoleIndex = 0; RoleIndex < Roles.Num(); ++RoleIndex)
 	{
-		SearchContext.AddRole(Roles[RoleIndex], AnimInstances[RoleIndex], PoseHistories[RoleIndex]);
+		SearchContext.AddRole(Roles[RoleIndex], AnimInstances[RoleIndex], InternalPoseHistories[RoleIndex]);
 	}
 
 	// budgeting some stack allocations for simple use cases. bigger requests of AnimationAssets contining 
@@ -1040,4 +1015,37 @@
 	return SearchResult;
 }
 
+// deprecated signatures
+UE::PoseSearch::FSearchResult UPoseSearchLibrary::MotionMatch(
+	const FAnimationBaseContext& Context,
+	TArrayView<const UObject*> AssetsToSearch,
+	const FPoseSearchContinuingProperties& ContinuingProperties)
+{
+	using namespace UE::PoseSearch;
+
+	const IPoseHistory* PoseHistory = nullptr;
+	if (IPoseHistoryProvider* PoseHistoryProvider = Context.GetMessage<IPoseHistoryProvider>())
+	{
+		PoseHistory = &PoseHistoryProvider->GetPoseHistory();
+	}
+
+	UAnimInstance* AnimInstance = Cast<UAnimInstance>(Context.AnimInstanceProxy->GetAnimInstanceObject());
+	check(AnimInstance);
+
+	return MotionMatch(MakeArrayView(&AnimInstance, 1), MakeArrayView(&DefaultRole, 1), MakeArrayView(&PoseHistory, 1), 
+		AssetsToSearch, ContinuingProperties, FPoseSearchFutureProperties(), Context.GetCurrentNodeId());
+}
+		
+UE::PoseSearch::FSearchResult UPoseSearchLibrary::MotionMatch(
+	TArrayView<UAnimInstance*> AnimInstances,
+	TArrayView<const UE::PoseSearch::FRole> Roles,
+	TArrayView<const UE::PoseSearch::IPoseHistory*> PoseHistories,
+	TArrayView<const UObject*> AssetsToSearch,
+	const FPoseSearchContinuingProperties& ContinuingProperties,
+	const int32 DebugSessionUniqueIdentifier,
+	float DesiredPermutationTimeOffset)
+{
+	return MotionMatch(AnimInstances, Roles, PoseHistories, AssetsToSearch, ContinuingProperties, FPoseSearchFutureProperties(), DebugSessionUniqueIdentifier);
+}
+	
 #undef LOCTEXT_NAMESPACE
diff --git a/Engine/Plugins/Animation/PoseSearch/Source/Runtime/Private/PoseSearchModule.cpp b/Engine/Plugins/Animation/PoseSearch/Source/Runtime/Private/PoseSearchModule.cpp
--- a/Engine/Plugins/Animation/PoseSearch/Source/Runtime/Private/PoseSearchModule.cpp
+++ b/Engine/Plugins/Animation/PoseSearch/Source/Runtime/Private/PoseSearchModule.cpp
@@ -1,10 +1,14 @@
 // Copyright Epic Games, Inc. All Rights Reserved.
 
 #include "Animation/AnimPoseSearchProvider.h"
+#include "Animation/AnimInstanceProxy.h"
+#include "Animation/AnimNodeBase.h"
 #include "Features/IModularFeatures.h"
 #include "Modules/ModuleManager.h"
 #include "PoseSearch/PoseSearchDatabase.h"
 #include "PoseSearch/PoseSearchLibrary.h"
+#include "PoseSearch/PoseSearchSchema.h"
+#include "PoseSearchFeatureChannel_PermutationTime.h"
 
 class FPoseSearchModule final : public IModuleInterface, public UE::Anim::IPoseSearchProvider
 {
@@ -22,21 +26,59 @@ public:
 	}
 
 	// IPoseSearchProvider
-	virtual UE::Anim::IPoseSearchProvider::FSearchResult Search(const FAnimationBaseContext& GraphContext, TArrayView<const UObject*> AssetsToSearch, const UObject* PlayingAsset, float PlayingAssetAccumulatedTime) const override
+	virtual UE::Anim::IPoseSearchProvider::FSearchResult Search(const FAnimationBaseContext& GraphContext, TArrayView<const UObject*> AssetsToSearch,
+		const FSearchPlayingAsset& PlayingAsset, const FSearchFutureAsset& FutureAsset) const override
 	{
+		using namespace UE::PoseSearch;
+
 		FPoseSearchContinuingProperties ContinuingProperties;
-		ContinuingProperties.PlayingAsset = PlayingAsset;
-		ContinuingProperties.PlayingAssetAccumulatedTime = PlayingAssetAccumulatedTime;
-		const UE::PoseSearch::FSearchResult SearchResult = UPoseSearchLibrary::MotionMatch(GraphContext, AssetsToSearch, ContinuingProperties);
+		ContinuingProperties.PlayingAsset = PlayingAsset.Asset;
+		ContinuingProperties.PlayingAssetAccumulatedTime = PlayingAsset.AccumulatedTime;
+
+		FPoseSearchFutureProperties Future;
+		Future.Animation = FutureAsset.Asset;
+		Future.AnimationTime = FutureAsset.AccumulatedTime;
+		Future.IntervalTime = FutureAsset.IntervalTime;
+
+		const IPoseHistory* PoseHistory = nullptr;
+		if (IPoseHistoryProvider* PoseHistoryProvider = GraphContext.GetMessage<IPoseHistoryProvider>())
+		{
+			PoseHistory = &PoseHistoryProvider->GetPoseHistory();
+		}
+
+		UAnimInstance* AnimInstance = Cast<UAnimInstance>(GraphContext.AnimInstanceProxy->GetAnimInstanceObject());
+		check(AnimInstance);
+		
+		const UE::PoseSearch::FSearchResult SearchResult = UPoseSearchLibrary::MotionMatch(MakeArrayView(&AnimInstance, 1), MakeArrayView(&DefaultRole, 1),
+			MakeArrayView(&PoseHistory, 1), AssetsToSearch, ContinuingProperties, Future, GraphContext.GetCurrentNodeId());
+
 		UE::Anim::IPoseSearchProvider::FSearchResult ProviderResult;
 		if (const UE::PoseSearch::FSearchIndexAsset* SearchIndexAsset = SearchResult.GetSearchIndexAsset())
 		{
-			if (const FPoseSearchDatabaseAnimationAssetBase* DatabaseAnimationAssetBase = SearchResult.Database->GetAnimationAssetBase(*SearchIndexAsset))
+			const UPoseSearchDatabase* Database = SearchResult.Database.Get();
+			if (const FPoseSearchDatabaseAnimationAssetBase* DatabaseAnimationAssetBase = Database->GetAnimationAssetBase(*SearchIndexAsset))
 			{
 				ProviderResult.SelectedAsset = DatabaseAnimationAssetBase->GetAnimationAsset();
 				ProviderResult.Dissimilarity = SearchResult.PoseCost.GetTotalCost();
 				ProviderResult.TimeOffsetSeconds = SearchResult.AssetTime;
 				ProviderResult.bIsFromContinuingPlaying = SearchResult.bIsContinuingPoseSearch;
+				ProviderResult.bMirrored = SearchIndexAsset->IsMirrored();
+
+				// figuring out the WantedPlayRate
+				ProviderResult.WantedPlayRate = 1.f;
+				if (Future.Animation && Future.IntervalTime > 0.f)
+				{
+					if (const UPoseSearchFeatureChannel_PermutationTime* PermutationTimeChannel = Database->Schema->FindFirstChannelOfType<UPoseSearchFeatureChannel_PermutationTime>())
+					{
+						const FSearchIndex& SearchIndex = Database->GetSearchIndex();
+						if (!SearchIndex.IsValuesEmpty())
+						{
+							TConstArrayView<float> ResultData = Database->GetSearchIndex().GetPoseValues(SearchResult.PoseIdx);
+							const float ActualIntervalTime = PermutationTimeChannel->GetPermutationTime(ResultData);
+							ProviderResult.WantedPlayRate = ActualIntervalTime / Future.IntervalTime;
+						}
+					}
+				}
 			}
 		}
 
diff --git a/Engine/Plugins/Animation/PoseSearch/Source/Runtime/Public/PoseSearch/PoseSearchLibrary.h b/Engine/Plugins/Animation/PoseSearch/Source/Runtime/Public/PoseSearch/PoseSearchLibrary.h
--- a/Engine/Plugins/Animation/PoseSearch/Source/Runtime/Public/PoseSearch/PoseSearchLibrary.h
+++ b/Engine/Plugins/Animation/PoseSearch/Source/Runtime/Public/PoseSearch/PoseSearchLibrary.h
@@ -94,7 +94,7 @@ struct POSESEARCH_API FPoseSearchFutureProperties
 public:
 	// Animation to play (it'll start at AnimationTime seconds)
 	UPROPERTY(EditAnywhere, BlueprintReadWrite, Category=State)
-	TObjectPtr<UObject> Animation;
+	TObjectPtr<const UObject> Animation;
 
 	// Start time for Animation
 	UPROPERTY(EditAnywhere, BlueprintReadWrite, Category=State)
@@ -188,7 +188,7 @@ public:
 		TArray<UObject*> AssetsToSearch,
 		const FName PoseHistoryName,
 		const FPoseSearchContinuingProperties ContinuingProperties,
-		FPoseSearchFutureProperties Future,
+		const FPoseSearchFutureProperties Future,
 		FPoseSearchBlueprintResult& Result,
 		const int32 DebugSessionUniqueIdentifier = 6174);
 
@@ -209,24 +209,36 @@ public:
 		TArray<UObject*> AssetsToSearch,
 		const FName PoseHistoryName,
 		const FPoseSearchContinuingProperties ContinuingProperties,
+		const FPoseSearchFutureProperties Future,
 		FPoseSearchBlueprintResult& Result,
 		const int32 DebugSessionUniqueIdentifier = 6174);
 
 	static void MotionMatch(
-		TArrayView<UAnimInstance*> AnimInstances,
-		TArrayView<const UE::PoseSearch::FRole> Roles,
-		TArrayView<const UObject*> AssetsToSearch,
+		const TArrayView<UAnimInstance*> AnimInstances,
+		const TArrayView<const UE::PoseSearch::FRole> Roles,
+		const TArrayView<const UObject*> AssetsToSearch,
 		const FName PoseHistoryName,
 		const FPoseSearchContinuingProperties& ContinuingProperties,
 		const FPoseSearchFutureProperties& Future,
 		FPoseSearchBlueprintResult& Result,
 		const int32 DebugSessionUniqueIdentifier);
 
+	static UE::PoseSearch::FSearchResult MotionMatch(
+		const TArrayView<UAnimInstance*> AnimInstances,
+		const TArrayView<const UE::PoseSearch::FRole> Roles,
+		const TArrayView<const UE::PoseSearch::IPoseHistory*> PoseHistories, 
+		const TArrayView<const UObject*> AssetsToSearch,
+		const FPoseSearchContinuingProperties& ContinuingProperties,
+		const FPoseSearchFutureProperties& Future,
+		const int32 DebugSessionUniqueIdentifier);
+
+	UE_DEPRECATED(5.4, "Use other MotionMatch signatures instead")
 	static UE::PoseSearch::FSearchResult MotionMatch(
 		const FAnimationBaseContext& Context,
 		TArrayView<const UObject*> AssetsToSearch,
 		const FPoseSearchContinuingProperties& ContinuingProperties);
-
+		
+	UE_DEPRECATED(5.4, "Use other MotionMatch signatures instead")
 	static UE::PoseSearch::FSearchResult MotionMatch(
 		TArrayView<UAnimInstance*> AnimInstances,
 		TArrayView<const UE::PoseSearch::FRole> Roles,
diff --git a/Engine/Plugins/Experimental/Chooser/Source/Chooser/Private/AnimNode_ChooserPlayer.cpp b/Engine/Plugins/Experimental/Chooser/Source/Chooser/Private/AnimNode_ChooserPlayer.cpp
--- a/Engine/Plugins/Experimental/Chooser/Source/Chooser/Private/AnimNode_ChooserPlayer.cpp
+++ b/Engine/Plugins/Experimental/Chooser/Source/Chooser/Private/AnimNode_ChooserPlayer.cpp
@@ -25,6 +25,8 @@ FAnimNode_ChooserPlayer::FAnimNode_ChooserPlayer()
 
 UAnimationAsset* FAnimNode_ChooserPlayer::ChooseAsset(const FAnimationUpdateContext& Context)
 {
+	using namespace UE::Anim;
+
 	if (Chooser.IsValid())
 	{
 		// reset settings to default
@@ -33,7 +35,7 @@ UAnimationAsset* FAnimNode_ChooserPlayer::ChooseAsset(const FAnimationUpdateCont
 		const FObjectChooserBase& ChooserBase = Chooser.Get<FObjectChooserBase>();
 		if (bStartFromMatchingPose)
 		{
-			if (UE::Anim::IPoseSearchProvider* PoseSearchProvider = UE::Anim::IPoseSearchProvider::Get())
+			if (IPoseSearchProvider* PoseSearchProvider = IPoseSearchProvider::Get())
 			{
 				TArray<const UObject*, TInlineAllocator<128>> ChosenAssets;
 				ChooserBase.ChooseMulti(ChooserContext, FObjectChooserBase::FObjectChooserIteratorCallback::CreateLambda([&ChosenAssets](UObject* InResult)
@@ -42,7 +44,12 @@ UAnimationAsset* FAnimNode_ChooserPlayer::ChooseAsset(const FAnimationUpdateCont
 						return FObjectChooserBase::EIteratorStatus::Continue;
 					}));
 
-				const UE::Anim::IPoseSearchProvider::FSearchResult SearchResult = PoseSearchProvider->Search(Context, ChosenAssets, GetAnimAsset(), GetAccumulatedTime());
+
+				IPoseSearchProvider::FSearchPlayingAsset PlayingAsset;
+				PlayingAsset.Asset = GetAnimAsset();
+				PlayingAsset.AccumulatedTime = GetAccumulatedTime();
+
+				const IPoseSearchProvider::FSearchResult SearchResult = PoseSearchProvider->Search(Context, ChosenAssets, PlayingAsset, IPoseSearchProvider::FSearchFutureAsset() );
 				if (UAnimationAsset* SelectedAnimationAsset = Cast<UAnimationAsset>(SearchResult.SelectedAsset))
 				{
 					if (!SearchResult.bIsFromContinuingPlaying)
diff --git a/Engine/Source/Runtime/Engine/Private/Animation/AnimNode_SequencePlayer.cpp b/Engine/Source/Runtime/Engine/Private/Animation/AnimNode_SequencePlayer.cpp
--- a/Engine/Source/Runtime/Engine/Private/Animation/AnimNode_SequencePlayer.cpp
+++ b/Engine/Source/Runtime/Engine/Private/Animation/AnimNode_SequencePlayer.cpp
@@ -158,13 +158,16 @@ float FAnimNode_SequencePlayerBase::GetTimeFromEnd(float CurrentNodeTime) const
 
 float FAnimNode_SequencePlayerBase::GetEffectiveStartPosition(const FAnimationBaseContext& Context) const
 {
+	using namespace UE::Anim;
+
 	// Override the start position if pose matching is enabled
 	const UObject* CurrentSequence = GetSequence();
 	if (CurrentSequence != nullptr && GetStartFromMatchingPose())
 	{
-		if (const UE::Anim::IPoseSearchProvider* PoseSearchProvider = UE::Anim::IPoseSearchProvider::Get())
+		if (const IPoseSearchProvider* PoseSearchProvider = IPoseSearchProvider::Get())
 		{
-			UE::Anim::IPoseSearchProvider::FSearchResult Result = PoseSearchProvider->Search(Context, MakeArrayView(&CurrentSequence, 1));
+			const IPoseSearchProvider::FSearchResult Result = PoseSearchProvider->Search(Context, MakeArrayView(&CurrentSequence, 1),
+				IPoseSearchProvider::FSearchPlayingAsset(), IPoseSearchProvider::FSearchFutureAsset());
 			if (Result.SelectedAsset != nullptr)
 			{
 				return Result.TimeOffsetSeconds;
diff --git a/Engine/Source/Runtime/Engine/Public/Animation/AnimPoseSearchProvider.h b/Engine/Source/Runtime/Engine/Public/Animation/AnimPoseSearchProvider.h
--- a/Engine/Source/Runtime/Engine/Public/Animation/AnimPoseSearchProvider.h
+++ b/Engine/Source/Runtime/Engine/Public/Animation/AnimPoseSearchProvider.h
@@ -28,20 +28,46 @@ public:
 		float TimeOffsetSeconds = 0.f;
 		float Dissimilarity = MAX_flt;
 		bool bIsFromContinuingPlaying = false;
+		bool bMirrored = false;
+		float WantedPlayRate = 1.f;
+	};
+
+	struct FSearchPlayingAsset
+	{
+		const UObject* Asset = nullptr;
+		float AccumulatedTime = 0.f;
+	};
+
+	struct FSearchFutureAsset : FSearchPlayingAsset
+	{
+		float IntervalTime = 0.f;
 	};
 
 	/**
 	* Finds a matching pose in the input Object given the current graph context
 	* 
-	* @param	GraphContext				Graph execution context used to construct a pose search query
-	* @param	AssetsToSearch				The assets to search for the pose query
-	* @param	PlayingAsset				The currently playing asset, used to bias the score of the eventually found continuing pose
-	* @param	PlayingAssetAccumulatedTime	The accumulated time of the currently playing asset
+	* @param	GraphContext					Graph execution context used to construct a pose search query
+	* @param	AssetsToSearch					The assets to search for the pose query
+	* @param	PlayingAsset.Asset				The currently playing asset, used to bias the score of the eventually found continuing pose
+	* @param	PlayingAsset.AccumulatedTime	The accumulated time of the currently playing asset
+	* @param	FutureAsset.Asset				The asset that will play in the future after FutureAssetIntervalTime seconds
+	* @param	FutureAsset.AccumulatedTime		The FutureAsset accumulated time in seconds when it'll start play
+	* @param	FutureAsset.IntervalTime		The requested time interval before FutureAsset will start playing at FutureAssetAccumulatedTime
 	* 
-	* @return	FSearchResult				The search result identifying the asset from AssetsToSearch or PlayingAsset that most closely matches the query
+	* @return	FSearchResult					The search result identifying the asset from AssetsToSearch or PlayingAsset that most closely matches the query
 	*/
+	virtual FSearchResult Search(const FAnimationBaseContext& GraphContext, TArrayView<const UObject*> AssetsToSearch, const FSearchPlayingAsset& PlayingAsset, const FSearchFutureAsset& FutureAsset) const = 0;
+	
+	UE_DEPRECATED(5.4, "Use the Search with FutureAsset instead")
 	virtual FSearchResult Search(const FAnimationBaseContext& GraphContext, TArrayView<const UObject*> AssetsToSearch,
-		const UObject* PlayingAsset = nullptr, float PlayingAssetAccumulatedTime = 0.f) const = 0;
+		const UObject* PlayingAsset = nullptr, float PlayingAssetAccumulatedTime = 0.f) const
+	{
+		FSearchPlayingAsset SearchPlayingAsset;
+		SearchPlayingAsset.Asset = PlayingAsset;
+		SearchPlayingAsset.AccumulatedTime = PlayingAssetAccumulatedTime;
+
+		return Search(GraphContext, AssetsToSearch, SearchPlayingAsset, FSearchFutureAsset());
+	}
 };
 
 } // namespace UE::Anim
\ No newline at end of file
