RESOLUTION: Conflict in FPoseSearchDatabaseAnimationAssetBase::GetEffectiveSamplingRange(). Took repo changes and applied the fix manually (handling cacse where minimum is greated than the maximum).

Fixed SamplingRange not being clamped to the assset's play length in the details panel. Additionally, fixed the GetEffectiveSamplingRange() method handle the case where the SamplingRange's minimum is greater than the maximum.

[RN] minor

diff --git a/Engine/Plugins/Animation/PoseSearch/Source/Editor/Private/PoseSearchDatabaseEditorReflection.cpp b/Engine/Plugins/Animation/PoseSearch/Source/Editor/Private/PoseSearchDatabaseEditorReflection.cpp
--- a/Engine/Plugins/Animation/PoseSearch/Source/Editor/Private/PoseSearchDatabaseEditorReflection.cpp
+++ b/Engine/Plugins/Animation/PoseSearch/Source/Editor/Private/PoseSearchDatabaseEditorReflection.cpp
@@ -38,7 +38,7 @@ void UPoseSearchDatabaseReflectionBase::PostEditChangeProperty(FPropertyChangedE
 	}
 }
 
-bool UPoseSearchDatabaseSequenceReflection::ApplyChanges() const 
+bool UPoseSearchDatabaseSequenceReflection::ApplyChanges() 
 {
 	if (const TSharedPtr<UE::PoseSearch::FDatabaseAssetTreeNode> AssetTreeNode = WeakAssetTreeNode.Pin())
 	{
@@ -50,6 +50,9 @@ bool UPoseSearchDatabaseSequenceReflection::ApplyChanges() const
 				FInstancedStruct& DatabaseAsset = Database->GetMutableAnimationAssetStruct(AssetTreeNode->SourceAssetIdx);
 				if (FPoseSearchDatabaseSequence* DatabaseSequence = DatabaseAsset.GetMutablePtr<FPoseSearchDatabaseSequence>())
 				{
+					Sequence.SamplingRange.Min = FMath::Clamp(Sequence.SamplingRange.Min, 0.0f, Sequence.GetPlayLength()); 
+					Sequence.SamplingRange.Max = FMath::Clamp(Sequence.SamplingRange.Max, 0.0f, Sequence.GetPlayLength()); 
+					
 					*DatabaseSequence = Sequence;
 					Database->MarkPackageDirty();
 
@@ -62,7 +65,7 @@ bool UPoseSearchDatabaseSequenceReflection::ApplyChanges() const
 	return false;
 }
 
-bool UPoseSearchDatabaseBlendSpaceReflection::ApplyChanges() const
+bool UPoseSearchDatabaseBlendSpaceReflection::ApplyChanges()
 {
 	if (const TSharedPtr<UE::PoseSearch::FDatabaseAssetTreeNode> AssetTreeNode = WeakAssetTreeNode.Pin())
 	{
@@ -86,7 +89,7 @@ bool UPoseSearchDatabaseBlendSpaceReflection::ApplyChanges() const
 	return false;
 }
 
-bool UPoseSearchDatabaseAnimCompositeReflection::ApplyChanges() const
+bool UPoseSearchDatabaseAnimCompositeReflection::ApplyChanges()
 {
 	if (const TSharedPtr<UE::PoseSearch::FDatabaseAssetTreeNode> AssetTreeNode = WeakAssetTreeNode.Pin())
 	{
@@ -98,7 +101,11 @@ bool UPoseSearchDatabaseAnimCompositeReflection::ApplyChanges() const
 				FInstancedStruct& DatabaseAsset = Database->GetMutableAnimationAssetStruct(AssetTreeNode->SourceAssetIdx);
 				if (FPoseSearchDatabaseAnimComposite* DatabaseAnimComposite = DatabaseAsset.GetMutablePtr<FPoseSearchDatabaseAnimComposite>())
 				{
+					AnimComposite.SamplingRange.Min = FMath::Clamp(AnimComposite.SamplingRange.Min, 0.0f, AnimComposite.GetPlayLength());
+					AnimComposite.SamplingRange.Max = FMath::Clamp(AnimComposite.SamplingRange.Max, 0.0f, AnimComposite.GetPlayLength());
+					
 					*DatabaseAnimComposite = AnimComposite;
+
 					Database->MarkPackageDirty();
 
 					return true;
@@ -110,7 +117,7 @@ bool UPoseSearchDatabaseAnimCompositeReflection::ApplyChanges() const
 	return false;
 }
 
-bool UPoseSearchDatabaseAnimMontageReflection::ApplyChanges() const
+bool UPoseSearchDatabaseAnimMontageReflection::ApplyChanges()
 {
 	if (const TSharedPtr<UE::PoseSearch::FDatabaseAssetTreeNode> AssetTreeNode = WeakAssetTreeNode.Pin())
 	{
@@ -122,7 +129,11 @@ bool UPoseSearchDatabaseAnimMontageReflection::ApplyChanges() const
 				FInstancedStruct& DatabaseAsset = Database->GetMutableAnimationAssetStruct(AssetTreeNode->SourceAssetIdx);
 				if (FPoseSearchDatabaseAnimMontage* DatabaseAnimMontage = DatabaseAsset.GetMutablePtr<FPoseSearchDatabaseAnimMontage>())
 				{
+					AnimMontage.SamplingRange.Min = FMath::Clamp(AnimMontage.SamplingRange.Min, 0.0f, AnimMontage.GetPlayLength()); 
+					AnimMontage.SamplingRange.Max = FMath::Clamp(AnimMontage.SamplingRange.Max, 0.0f, AnimMontage.GetPlayLength());
+					
 					*DatabaseAnimMontage = AnimMontage;
+
 					Database->MarkPackageDirty();
 
 					return true;
@@ -134,7 +145,7 @@ bool UPoseSearchDatabaseAnimMontageReflection::ApplyChanges() const
 	return false;
 }
 
-bool UPoseSearchDatabaseMultiSequenceReflection::ApplyChanges() const
+bool UPoseSearchDatabaseMultiSequenceReflection::ApplyChanges()
 {
 	if (const TSharedPtr<UE::PoseSearch::FDatabaseAssetTreeNode> AssetTreeNode = WeakAssetTreeNode.Pin())
 	{
diff --git a/Engine/Plugins/Animation/PoseSearch/Source/Editor/Public/PoseSearchDatabaseEditorReflection.h b/Engine/Plugins/Animation/PoseSearch/Source/Editor/Public/PoseSearchDatabaseEditorReflection.h
--- a/Engine/Plugins/Animation/PoseSearch/Source/Editor/Public/PoseSearchDatabaseEditorReflection.h
+++ b/Engine/Plugins/Animation/PoseSearch/Source/Editor/Public/PoseSearchDatabaseEditorReflection.h
@@ -24,7 +24,7 @@ public:
 		const TWeakPtr<UE::PoseSearch::FDatabaseAssetTreeNode>& InWeakAssetTreeNode,
 		const TSharedPtr<UE::PoseSearch::SDatabaseAssetTree>& InAssetTreeWidget);
 	
-	virtual bool ApplyChanges() const { return false; };
+	virtual bool ApplyChanges() { return false; };
 
 #if WITH_EDITOR
 	virtual void PostEditChangeProperty(FPropertyChangedEvent& PropertyChangedEvent) override;
@@ -59,7 +59,7 @@ public:
 	UPROPERTY(EditAnywhere, Category = "Selected Sequence")
 	FPoseSearchDatabaseSequenceEx Sequence;
 	
-	virtual bool ApplyChanges() const override;
+	virtual bool ApplyChanges() override;
 };
 
 USTRUCT()
@@ -86,7 +86,7 @@ public:
 	UPROPERTY(EditAnywhere, Category = "Selected Blend Space")
 	FPoseSearchDatabaseBlendSpaceEx BlendSpace;
 	
-	virtual bool ApplyChanges() const override;
+	virtual bool ApplyChanges() override;
 };
 
 USTRUCT()
@@ -113,7 +113,7 @@ public:
 	UPROPERTY(EditAnywhere, Category = "Selected Anim Composite")
 	FPoseSearchDatabaseAnimCompositeEx AnimComposite;
 	
-	virtual bool ApplyChanges() const override;
+	virtual bool ApplyChanges() override;
 };
 
 USTRUCT()
@@ -140,7 +140,7 @@ public:
 	UPROPERTY(EditAnywhere, Category = "Selected Anim Montage")
 	FPoseSearchDatabaseAnimMontageEx AnimMontage;
 	
-	virtual bool ApplyChanges() const override;
+	virtual bool ApplyChanges() override;
 };
 
 USTRUCT()
@@ -167,7 +167,7 @@ public:
 	UPROPERTY(EditAnywhere, Category = "Selected Sequence")
 	FPoseSearchDatabaseMultiSequenceEx MultiSequence;
 	
-	virtual bool ApplyChanges() const override;
+	virtual bool ApplyChanges() override;
 };
 
 USTRUCT()
diff --git a/Engine/Plugins/Animation/PoseSearch/Source/Runtime/Private/PoseSearchDatabase.cpp b/Engine/Plugins/Animation/PoseSearch/Source/Runtime/Private/PoseSearchDatabase.cpp
--- a/Engine/Plugins/Animation/PoseSearch/Source/Runtime/Private/PoseSearchDatabase.cpp
+++ b/Engine/Plugins/Animation/PoseSearch/Source/Runtime/Private/PoseSearchDatabase.cpp
@@ -399,6 +399,14 @@ FFloatInterval FPoseSearchDatabaseAnimationAssetBase::GetEffectiveSamplingRange(
 	FFloatInterval Range;
 	Range.Min = bSampleAll ? 0.0f : SamplingRange.Min;
 	Range.Max = bSampleAll ? PlayLength : FMath::Min(PlayLength, SamplingRange.Max);
+
+	if (Range.Min > Range.Max)
+	{
+		UE_LOG(LogPoseSearch, Warning, TEXT("Sampling range minimum (%f) is greated than max (%f). Setting min to be equal to max."), Range.Min, Range.Max)
+		
+		Range.Min = Range.Max;
+	}
+
 	return Range;
 }
 #endif // WITH_EDITORONLY_DATA
