motion matching - improvements for multicharacter support

diff --git a/Engine/Plugins/Animation/PoseSearch/Source/Runtime/Private/AnimNode_PoseSearchHistoryCollector.cpp b/Engine/Plugins/Animation/PoseSearch/Source/Runtime/Private/AnimNode_PoseSearchHistoryCollector.cpp
--- a/Engine/Plugins/Animation/PoseSearch/Source/Runtime/Private/AnimNode_PoseSearchHistoryCollector.cpp
+++ b/Engine/Plugins/Animation/PoseSearch/Source/Runtime/Private/AnimNode_PoseSearchHistoryCollector.cpp
@@ -10,30 +10,6 @@
 
 #define LOCTEXT_NAMESPACE "AnimNode_PoseSearchHistoryCollector"
 
-namespace UE::PoseSearch::Private
-{
-
-class FPoseHistoryProvider : public IPoseHistoryProvider
-{
-public:
-	FPoseHistoryProvider(const IPoseHistory& InPoseHistory)
-		: PoseHistory(InPoseHistory)
-	{
-	}
-
-	// IPoseHistoryProvider interface
-	virtual const IPoseHistory& GetPoseHistory() const override
-	{
-		return PoseHistory;
-	}
-
-private:
-	const IPoseHistory& PoseHistory;
-};
-
-} // namespace UE::PoseSearch::Private
-
-
 /////////////////////////////////////////////////////
 // FAnimNode_PoseSearchHistoryCollector_Base
 
@@ -108,6 +84,18 @@ void FAnimNode_PoseSearchHistoryCollector_Base::PreUpdate(const UAnimInstance* I
 	Super::PreUpdate(InAnimInstance);
 	
 	if (bGenerateTrajectory)
+	{
+		GenerateTrajectory(InAnimInstance);
+	}
+
+	PoseHistory.PreUpdate();
+
+	bIsTrajectoryGeneratedBeforePreUpdate = false;
+}
+
+void FAnimNode_PoseSearchHistoryCollector_Base::GenerateTrajectory(const UAnimInstance* InAnimInstance)
+{
+	if (!bIsTrajectoryGeneratedBeforePreUpdate)
 	{
 		FPoseSearchTrajectoryData::FSampling TrajectoryDataSampling;
 		TrajectoryDataSampling.NumHistorySamples = FMath::Max(PoseCount, TrajectoryHistoryCount);
@@ -116,11 +104,10 @@ void FAnimNode_PoseSearchHistoryCollector_Base::PreUpdate(const UAnimInstance* I
 		TrajectoryDataSampling.SecondsPerPredictionSample = PredictionSamplingInterval;
 
 		PoseHistory.GenerateTrajectory(InAnimInstance, InAnimInstance->GetDeltaSeconds(), TrajectoryData, TrajectoryDataSampling);
-	}
 
-	PoseHistory.PreUpdate();
+		bIsTrajectoryGeneratedBeforePreUpdate = true;
+	}
 }
-
 /////////////////////////////////////////////////////
 // FAnimNode_PoseSearchHistoryCollector
 
@@ -179,7 +166,7 @@ void FAnimNode_PoseSearchHistoryCollector::Update_AnyThread(const FAnimationUpda
 {
 	DECLARE_SCOPE_HIERARCHICAL_COUNTER_ANIMNODE(Update_AnyThread);
 	Super::Update_AnyThread(Context);
-	UE::Anim::TScopedGraphMessage<UE::PoseSearch::Private::FPoseHistoryProvider> ScopedMessage(Context, PoseHistory);
+	UE::Anim::TScopedGraphMessage<UE::PoseSearch::FPoseHistoryProvider> ScopedMessage(Context, PoseHistory);
 	Source.Update(Context);
 }
 
@@ -244,7 +231,7 @@ void FAnimNode_PoseSearchComponentSpaceHistoryCollector::Update_AnyThread(const
 {
 	DECLARE_SCOPE_HIERARCHICAL_COUNTER_ANIMNODE(Update_AnyThread);
 	Super::Update_AnyThread(Context);
-	UE::Anim::TScopedGraphMessage<UE::PoseSearch::Private::FPoseHistoryProvider> ScopedMessage(Context, PoseHistory);
+	UE::Anim::TScopedGraphMessage<UE::PoseSearch::FPoseHistoryProvider> ScopedMessage(Context, PoseHistory);
 	Source.Update(Context);
 }
 
diff --git a/Engine/Plugins/Animation/PoseSearch/Source/Runtime/Private/PoseSearchDatabase.cpp b/Engine/Plugins/Animation/PoseSearch/Source/Runtime/Private/PoseSearchDatabase.cpp
--- a/Engine/Plugins/Animation/PoseSearch/Source/Runtime/Private/PoseSearchDatabase.cpp
+++ b/Engine/Plugins/Animation/PoseSearch/Source/Runtime/Private/PoseSearchDatabase.cpp
@@ -1192,22 +1192,34 @@ UE::PoseSearch::FSearchResult UPoseSearchDatabase::SearchContinuingPose(UE::Pose
 	const FSearchIndexAsset& SearchIndexAsset = SearchIndex.GetAssetForPose(PoseIdx);
 	const FPoseSearchDatabaseAnimationAssetBase* DatabaseAnimationAssetBase = GetAnimationAssetStruct(SearchIndexAsset).GetPtr<FPoseSearchDatabaseAnimationAssetBase>();
 	check(DatabaseAnimationAssetBase);
-	const UAnimationAsset* AnimationAsset = CastChecked<UAnimationAsset>(DatabaseAnimationAssetBase->GetAnimationAsset());
-	
-	// sampler used only to extract the notify states. RootTransformOrigin can be set as Identity, since will not be relevant
-	const FAnimationAssetSampler SequenceBaseSampler(AnimationAsset, FTransform::Identity, SearchIndexAsset.GetBlendParameters());
-	const float SampleTime = GetRealAssetTime(PoseIdx);
 
 	float UpdatedContinuingPoseCostBias = ContinuingPoseCostBias;
-	SequenceBaseSampler.ExtractPoseSearchNotifyStates(SampleTime, [&UpdatedContinuingPoseCostBias](const UAnimNotifyState_PoseSearchBase* PoseSearchNotify)
+	const float SampleTime = GetRealAssetTime(PoseIdx);
+	for (int32 RoleIndex = 0; RoleIndex < DatabaseAnimationAssetBase->GetNumRoles(); ++RoleIndex)
+	{
+		if (const UAnimationAsset* AnimationAsset = DatabaseAnimationAssetBase->GetAnimationAssetForRole(DatabaseAnimationAssetBase->GetRole(RoleIndex)))
 		{
-			if (const UAnimNotifyState_PoseSearchOverrideContinuingPoseCostBias* ContinuingPoseCostBiasNotify = Cast<const UAnimNotifyState_PoseSearchOverrideContinuingPoseCostBias>(PoseSearchNotify))
+			// sampler used only to extract the notify states. RootTransformOrigin can be set as Identity, since will not be relevant
+			const FAnimationAssetSampler SequenceBaseSampler(AnimationAsset, FTransform::Identity, SearchIndexAsset.GetBlendParameters());
+
+			bool bDone = false;
+			SequenceBaseSampler.ExtractPoseSearchNotifyStates(SampleTime, [&UpdatedContinuingPoseCostBias, &bDone](const UAnimNotifyState_PoseSearchBase* PoseSearchNotify)
+				{
+					if (const UAnimNotifyState_PoseSearchOverrideContinuingPoseCostBias* ContinuingPoseCostBiasNotify = Cast<const UAnimNotifyState_PoseSearchOverrideContinuingPoseCostBias>(PoseSearchNotify))
+					{
+						UpdatedContinuingPoseCostBias = ContinuingPoseCostBiasNotify->CostAddend;
+						bDone = true;
+						return false;
+					}
+					return true;
+				});
+
+			if (bDone)
 			{
-				UpdatedContinuingPoseCostBias = ContinuingPoseCostBiasNotify->CostAddend;
-				return false;
+				break;
 			}
-			return true;
-		});
+		}
+	}
 
 	// since any PoseCost calculated here is at least SearchIndex.MinCostAddend + UpdatedContinuingPoseCostBias,
 	// there's no point in performing the search if CurrentBestTotalCost is already better than that
diff --git a/Engine/Plugins/Animation/PoseSearch/Source/Runtime/Private/PoseSearchHistory.cpp b/Engine/Plugins/Animation/PoseSearch/Source/Runtime/Private/PoseSearchHistory.cpp
--- a/Engine/Plugins/Animation/PoseSearch/Source/Runtime/Private/PoseSearchHistory.cpp
+++ b/Engine/Plugins/Animation/PoseSearch/Source/Runtime/Private/PoseSearchHistory.cpp
@@ -13,7 +13,7 @@
 #include "PoseSearch/PoseSearchDatabase.h"
 #include "PoseSearch/PoseSearchDefines.h"
 
-IMPLEMENT_ANIMGRAPH_MESSAGE(UE::PoseSearch::IPoseHistoryProvider);
+IMPLEMENT_ANIMGRAPH_MESSAGE(UE::PoseSearch::FPoseHistoryProvider);
 
 namespace UE::PoseSearch
 {
diff --git a/Engine/Plugins/Animation/PoseSearch/Source/Runtime/Private/PoseSearchLibrary.cpp b/Engine/Plugins/Animation/PoseSearch/Source/Runtime/Private/PoseSearchLibrary.cpp
--- a/Engine/Plugins/Animation/PoseSearch/Source/Runtime/Private/PoseSearchLibrary.cpp
+++ b/Engine/Plugins/Animation/PoseSearch/Source/Runtime/Private/PoseSearchLibrary.cpp
@@ -36,6 +36,63 @@ TAutoConsoleVariable<bool> CVarAnimMotionMatchDrawMatchEnable(TEXT("a.MotionMatc
 
 namespace UE::PoseSearch
 {
+	// budgeting some stack allocations for simple use cases. bigger requests of AnimationAssets containing UAnimNotifyState_PoseSearchBranchIn 
+	// referencing multiple databases will default to a slower TMemStackAllocator (that hides heap allocations)
+	enum { MAX_STACK_ALLOCATED_ANIMATIONS = 16 };
+	enum { MAX_STACK_ALLOCATED_SETS = 2 };
+	typedef	TArray<const UObject*, TInlineAllocator<MAX_STACK_ALLOCATED_ANIMATIONS, TMemStackAllocator<>>> TAssetsToSearch;
+	// an empty TAssetsToSearch associated to Database means we need to search ALL the assets
+	typedef TMap<const UPoseSearchDatabase*, TAssetsToSearch, TInlineSetAllocator<MAX_STACK_ALLOCATED_SETS, TMemStackSetAllocator<>>> TAssetsToSearchPerDatabaseMap;
+	typedef TPair<const UPoseSearchDatabase*, TAssetsToSearch> TAssetsToSearchPerDatabasePair;
+
+	static void AddToSearchForDatabase(TAssetsToSearchPerDatabaseMap& AssetsToSearchPerDatabaseMap, const UObject* AssetToSearch, const UPoseSearchDatabase* Database)
+	{
+		if (TAssetsToSearch* AssetsToSearch = AssetsToSearchPerDatabaseMap.Find(Database))
+		{
+			// an empty TAssetsToSearch associated to Database means we need to search ALL the assets, so we don't need to add this AssetToSearch
+			if (!AssetsToSearch->IsEmpty())
+			{
+				AssetsToSearch->AddUnique(AssetToSearch);
+			}
+		}
+		else
+		{
+			AssetsToSearchPerDatabaseMap.Add(Database).AddUnique(AssetToSearch);
+		}
+	}
+
+	static void AddToSearch(TAssetsToSearchPerDatabaseMap& AssetsToSearchPerDatabaseMap, const UObject* AssetToSearch)
+	{
+		if (const UAnimSequenceBase* SequenceBase = Cast<const UAnimSequenceBase>(AssetToSearch))
+		{
+			for (const FAnimNotifyEvent& NotifyEvent : SequenceBase->Notifies)
+			{
+				if (const UAnimNotifyState_PoseSearchBranchIn* PoseSearchBranchIn = Cast<UAnimNotifyState_PoseSearchBranchIn>(NotifyEvent.NotifyStateClass))
+				{
+					if (!PoseSearchBranchIn->Database)
+					{
+						UE_LOG(LogPoseSearch, Error, TEXT("improperly setup UAnimNotifyState_PoseSearchBranchIn with null Database in %s"), *SequenceBase->GetName());
+						continue;
+					}
+#if WITH_EDITOR
+					if (!PoseSearchBranchIn->Database->Contains(SequenceBase))
+					{
+						UE_LOG(LogPoseSearch, Error, TEXT("improperly setup UAnimSequenceBase. Database %s doesn't contain UAnimSequenceBase %s"), *PoseSearchBranchIn->Database->GetName(), *SequenceBase->GetName());
+						continue;
+					}
+#endif // WITH_EDITOR
+					
+					AddToSearchForDatabase(AssetsToSearchPerDatabaseMap, SequenceBase, PoseSearchBranchIn->Database);
+				}
+			}
+		}
+		else if (const UPoseSearchDatabase* Database = Cast<UPoseSearchDatabase>(AssetToSearch))
+		{
+			// an empty TAssetsToSearch associated to Database means we need to search ALL the assets
+			AssetsToSearchPerDatabaseMap.FindOrAdd(Database).Reset();
+		}
+	}
+
 	static bool IsForceInterrupt(EPoseSearchInterruptMode InterruptMode, const UPoseSearchDatabase* CurrentResultDatabase, const TArray<TObjectPtr<const UPoseSearchDatabase>>& Databases)
 	{
 		switch (InterruptMode)
@@ -108,47 +165,6 @@ namespace UE::PoseSearch
 		return false;
 	}
 
-	static const FAnimNode_PoseSearchHistoryCollector_Base* FindPoseHistoryNode(const FName PoseHistoryName, const UAnimInstance* AnimInstance)
-	{
-		if (AnimInstance)
-		{
-			TSet<const UAnimInstance*, DefaultKeyFuncs<const UAnimInstance*>, TInlineSetAllocator<128>> AlreadyVisited;
-			TArray<const UAnimInstance*, TInlineAllocator<128>> ToVisit;
-
-			ToVisit.Add(AnimInstance);
-			AlreadyVisited.Add(AnimInstance);
-
-			while (!ToVisit.IsEmpty())
-			{
-				const UAnimInstance* Visiting = ToVisit.Pop();
-
-				if (IAnimClassInterface* AnimBlueprintClass = IAnimClassInterface::GetFromClass(Visiting->GetClass()))
-				{
-					if (const FAnimSubsystem_Tag* TagSubsystem = AnimBlueprintClass->FindSubsystem<FAnimSubsystem_Tag>())
-					{
-						if (const FAnimNode_PoseSearchHistoryCollector_Base* HistoryCollector = TagSubsystem->FindNodeByTag<FAnimNode_PoseSearchHistoryCollector_Base>(PoseHistoryName, Visiting))
-						{
-							return HistoryCollector;
-						}
-					}
-				}
-
-				const USkeletalMeshComponent* SkeletalMeshComponent = Visiting->GetSkelMeshComponent();
-				const TArray<UAnimInstance*>& LinkedAnimInstances = SkeletalMeshComponent->GetLinkedAnimInstances();
-				for (const UAnimInstance* LinkedAnimInstance : LinkedAnimInstances)
-				{
-					bool bIsAlreadyInSet = false;
-					AlreadyVisited.Add(LinkedAnimInstance, &bIsAlreadyInSet);
-					if (!bIsAlreadyInSet)
-					{
-						ToVisit.Add(LinkedAnimInstance);
-					}
-				}
-			}
-		}
-		return nullptr;
-	}
-
 #if ENABLE_DRAW_DEBUG && ENABLE_ANIM_DEBUG
 	class UAnimInstanceProxyProvider : public UAnimInstance
 	{
@@ -416,9 +432,8 @@ void UPoseSearchLibrary::UpdateMotionMatchingState(
 
 	InOutMotionMatchingState.bJumpedToPose = false;
 
-	// used when YawFromAnimationBlendRate is greater than zero, by setting a future (YawFromAnimationTrajectoryBlendTime seconds ahead) root bone to the skeleton default
 	const IPoseHistory* PoseHistory = nullptr;
-	if (IPoseHistoryProvider* PoseHistoryProvider = Context.GetMessage<IPoseHistoryProvider>())
+	if (FPoseHistoryProvider* PoseHistoryProvider = Context.GetMessage<FPoseHistoryProvider>())
 	{
 		PoseHistory = &PoseHistoryProvider->GetPoseHistory();
 	}
@@ -728,7 +743,7 @@ void UPoseSearchLibrary::MotionMatch(
 		}
 	}
 }
-	
+
 UE::PoseSearch::FSearchResult UPoseSearchLibrary::MotionMatch(
 	const TArrayView<UAnimInstance*> AnimInstances,
 	const TArrayView<const UE::PoseSearch::FRole> Roles,
@@ -824,63 +839,11 @@ UE::PoseSearch::FSearchResult UPoseSearchLibrary::MotionMatch(
 		SearchContext.AddRole(Roles[RoleIndex], AnimInstances[RoleIndex], InternalPoseHistories[RoleIndex]);
 	}
 
-	// budgeting some stack allocations for simple use cases. bigger requests of AnimationAssets contining 
-	// UAnimNotifyState_PoseSearchBranchIn referencing multiple datbases will default to slower heap allocations
-	enum { MAX_STACK_ALLOCATED_ANIMATIONS = 16 };
-	enum { MAX_STACK_ALLOCATED_SETS = 2 };
-	typedef	TArray<const UObject*, TInlineAllocator<MAX_STACK_ALLOCATED_ANIMATIONS, TMemStackAllocator<>>> TAssetsToSearch;
-	// an empty TAssetsToSearch associated to Database means we need to search ALL the assets
-	typedef TMap<const UPoseSearchDatabase*, TAssetsToSearch, TInlineSetAllocator<MAX_STACK_ALLOCATED_SETS, TMemStackSetAllocator<>>> TAssetsToSearchPerDatabaseMap;
-	typedef TPair<const UPoseSearchDatabase*, TAssetsToSearch> TAssetsToSearchPerDatabasePair;
 	TAssetsToSearchPerDatabaseMap AssetsToSearchPerDatabaseMap;
 	
-	auto AddToSearch = [](TAssetsToSearchPerDatabaseMap& AssetsToSearchPerDatabaseMap, const UObject* AssetToSearch)
-	{
-		if (const UAnimSequenceBase* SequenceBase = Cast<const UAnimSequenceBase>(AssetToSearch))
-		{
-			for (const FAnimNotifyEvent& NotifyEvent : SequenceBase->Notifies)
-			{
-				if (const UAnimNotifyState_PoseSearchBranchIn* PoseSearchBranchIn = Cast<UAnimNotifyState_PoseSearchBranchIn>(NotifyEvent.NotifyStateClass))
-				{
-					if (PoseSearchBranchIn->Database)
-					{
-#if WITH_EDITOR
-						if (!PoseSearchBranchIn->Database->Contains(SequenceBase))
-						{
-							UE_LOG(LogPoseSearch, Error, TEXT("improperly setup UAnimSequenceBase. Database %s doesn't contain UAnimSequenceBase %s"), *PoseSearchBranchIn->Database->GetName(), *SequenceBase->GetName());
-						}
-						else 
-#endif // WITH_EDITOR
-						if (TAssetsToSearch* AssetsToSearch = AssetsToSearchPerDatabaseMap.Find(PoseSearchBranchIn->Database))
-						{
-							// an empty TAssetsToSearch associated to Database means we need to search ALL the assets, so we don't need to add this SequenceBase
-							if (!AssetsToSearch->IsEmpty())
-							{
-								AssetsToSearch->AddUnique(SequenceBase);
-							}
-						}
-						else
-						{
-							AssetsToSearchPerDatabaseMap.Add(PoseSearchBranchIn->Database).AddUnique(SequenceBase);
-						}
-					}
-					else
-					{
-						UE_LOG(LogPoseSearch, Error, TEXT("improperly setup UAnimNotifyState_PoseSearchBranchIn with null Database in %s"), *SequenceBase->GetName());
-					}
-				}
-			}
-		}
-		else if (const UPoseSearchDatabase* Database = Cast<UPoseSearchDatabase>(AssetToSearch))
-		{
-			// an empty TAssetsToSearch associated to Database means we need to search ALL the assets
-			AssetsToSearchPerDatabaseMap.FindOrAdd(Database).Reset();
-		}
-	};
-
 	// collecting all the possible continuing pose search (it could be multiple searches, but most likely only one)
 	const float DeltaSeconds = AnimInstances[0]->GetDeltaSeconds();
-	if (const UAnimationAsset* PlayingAnimationAsset = Cast<UAnimationAsset>(ContinuingProperties.PlayingAsset))
+	if (const UObject* PlayingAnimationAsset = ContinuingProperties.PlayingAsset.Get())
 	{
 		AddToSearch(AssetsToSearchPerDatabaseMap, PlayingAnimationAsset);
 		
@@ -890,7 +853,11 @@ UE::PoseSearch::FSearchResult UPoseSearchLibrary::MotionMatch(
 			{
 				if (Database->Contains(PlayingAnimationAsset))
 				{
-					AddToSearch(AssetsToSearchPerDatabaseMap, AssetToSearch);
+					// checking just in case later on we add support for databases containing other databases
+					check(Cast<const UPoseSearchDatabase>(PlayingAnimationAsset) == nullptr);
+
+					// since it cannot be a database we can directly add it to AssetsToSearchPerDatabaseMap
+					AddToSearchForDatabase(AssetsToSearchPerDatabaseMap, AssetToSearch, Database);
 				}
 			}
 		}
@@ -1024,7 +991,7 @@ UE::PoseSearch::FSearchResult UPoseSearchLibrary::MotionMatch(
 	using namespace UE::PoseSearch;
 
 	const IPoseHistory* PoseHistory = nullptr;
-	if (IPoseHistoryProvider* PoseHistoryProvider = Context.GetMessage<IPoseHistoryProvider>())
+	if (FPoseHistoryProvider* PoseHistoryProvider = Context.GetMessage<FPoseHistoryProvider>())
 	{
 		PoseHistory = &PoseHistoryProvider->GetPoseHistory();
 	}
@@ -1048,4 +1015,47 @@ UE::PoseSearch::FSearchResult UPoseSearchLibrary::MotionMatch(
 	return MotionMatch(AnimInstances, Roles, PoseHistories, AssetsToSearch, ContinuingProperties, FPoseSearchFutureProperties(), DebugSessionUniqueIdentifier);
 }
 	
+const FAnimNode_PoseSearchHistoryCollector_Base* UPoseSearchLibrary::FindPoseHistoryNode(
+	const FName PoseHistoryName,
+	const UAnimInstance* AnimInstance)
+{
+	if (AnimInstance)
+	{
+		TSet<const UAnimInstance*, DefaultKeyFuncs<const UAnimInstance*>, TInlineSetAllocator<128>> AlreadyVisited;
+		TArray<const UAnimInstance*, TInlineAllocator<128>> ToVisit;
+
+		ToVisit.Add(AnimInstance);
+		AlreadyVisited.Add(AnimInstance);
+
+		while (!ToVisit.IsEmpty())
+		{
+			const UAnimInstance* Visiting = ToVisit.Pop();
+
+			if (IAnimClassInterface* AnimBlueprintClass = IAnimClassInterface::GetFromClass(Visiting->GetClass()))
+			{
+				if (const FAnimSubsystem_Tag* TagSubsystem = AnimBlueprintClass->FindSubsystem<FAnimSubsystem_Tag>())
+				{
+					if (const FAnimNode_PoseSearchHistoryCollector_Base* HistoryCollector = TagSubsystem->FindNodeByTag<FAnimNode_PoseSearchHistoryCollector_Base>(PoseHistoryName, Visiting))
+					{
+						return HistoryCollector;
+					}
+				}
+			}
+
+			const USkeletalMeshComponent* SkeletalMeshComponent = Visiting->GetSkelMeshComponent();
+			const TArray<UAnimInstance*>& LinkedAnimInstances = SkeletalMeshComponent->GetLinkedAnimInstances();
+			for (const UAnimInstance* LinkedAnimInstance : LinkedAnimInstances)
+			{
+				bool bIsAlreadyInSet = false;
+				AlreadyVisited.Add(LinkedAnimInstance, &bIsAlreadyInSet);
+				if (!bIsAlreadyInSet)
+				{
+					ToVisit.Add(LinkedAnimInstance);
+				}
+			}
+		}
+	}
+	return nullptr;
+}
+
 #undef LOCTEXT_NAMESPACE
diff --git a/Engine/Plugins/Animation/PoseSearch/Source/Runtime/Private/PoseSearchModule.cpp b/Engine/Plugins/Animation/PoseSearch/Source/Runtime/Private/PoseSearchModule.cpp
--- a/Engine/Plugins/Animation/PoseSearch/Source/Runtime/Private/PoseSearchModule.cpp
+++ b/Engine/Plugins/Animation/PoseSearch/Source/Runtime/Private/PoseSearchModule.cpp
@@ -41,7 +41,7 @@ public:
 		Future.IntervalTime = FutureAsset.IntervalTime;
 
 		const IPoseHistory* PoseHistory = nullptr;
-		if (IPoseHistoryProvider* PoseHistoryProvider = GraphContext.GetMessage<IPoseHistoryProvider>())
+		if (FPoseHistoryProvider* PoseHistoryProvider = GraphContext.GetMessage<FPoseHistoryProvider>())
 		{
 			PoseHistory = &PoseHistoryProvider->GetPoseHistory();
 		}
diff --git a/Engine/Plugins/Animation/PoseSearch/Source/Runtime/Private/PoseSearchResult.cpp b/Engine/Plugins/Animation/PoseSearch/Source/Runtime/Private/PoseSearchResult.cpp
--- a/Engine/Plugins/Animation/PoseSearch/Source/Runtime/Private/PoseSearchResult.cpp
+++ b/Engine/Plugins/Animation/PoseSearch/Source/Runtime/Private/PoseSearchResult.cpp
@@ -47,11 +47,6 @@ void FSearchResult::Update(float NewAssetTime)
 	}
 }
 
-bool FSearchResult::IsValid() const
-{
-	return PoseIdx != INDEX_NONE && Database != nullptr;
-}
-
 void FSearchResult::Reset()
 {
 	PoseIdx = INDEX_NONE;
diff --git a/Engine/Plugins/Animation/PoseSearch/Source/Runtime/Public/PoseSearch/AnimNode_PoseSearchHistoryCollector.h b/Engine/Plugins/Animation/PoseSearch/Source/Runtime/Public/PoseSearch/AnimNode_PoseSearchHistoryCollector.h
--- a/Engine/Plugins/Animation/PoseSearch/Source/Runtime/Public/PoseSearch/AnimNode_PoseSearchHistoryCollector.h
+++ b/Engine/Plugins/Animation/PoseSearch/Source/Runtime/Public/PoseSearch/AnimNode_PoseSearchHistoryCollector.h
@@ -99,11 +99,13 @@ public:
 
 	const UE::PoseSearch::FPoseHistory& GetPoseHistory() const { return PoseHistory; }
 	UE::PoseSearch::FPoseHistory& GetPoseHistory() { return PoseHistory; }
+	void GenerateTrajectory(const UAnimInstance* InAnimInstance);
 
 protected:
 	TArray<FBoneIndexType> GetRequiredBones(const FAnimInstanceProxy* AnimInstanceProxy) const;
 
 	UE::PoseSearch::FPoseHistory PoseHistory;
+	bool bIsTrajectoryGeneratedBeforePreUpdate = false;
 };
 
 USTRUCT(BlueprintInternalUseOnly)
diff --git a/Engine/Plugins/Animation/PoseSearch/Source/Runtime/Public/PoseSearch/PoseSearchHistory.h b/Engine/Plugins/Animation/PoseSearch/Source/Runtime/Public/PoseSearch/PoseSearchHistory.h
--- a/Engine/Plugins/Animation/PoseSearch/Source/Runtime/Public/PoseSearch/PoseSearchHistory.h
+++ b/Engine/Plugins/Animation/PoseSearch/Source/Runtime/Public/PoseSearch/PoseSearchHistory.h
@@ -206,11 +206,16 @@ private:
 	TArray<FPoseHistoryEntry, TInlineAllocator<4, TMemStackAllocator<>>> FutureEntries;
 };
 
-class IPoseHistoryProvider : public UE::Anim::IGraphMessage
+class POSESEARCH_API FPoseHistoryProvider : public UE::Anim::IGraphMessage
 {
-	DECLARE_ANIMGRAPH_MESSAGE(IPoseHistoryProvider);
+	DECLARE_ANIMGRAPH_MESSAGE(FPoseHistoryProvider);
+
 public:
-	virtual const IPoseHistory& GetPoseHistory() const = 0;
+	FPoseHistoryProvider(const IPoseHistory& InPoseHistory) : PoseHistory(InPoseHistory) { }
+	const IPoseHistory& GetPoseHistory() const { return PoseHistory; }
+
+private:
+	const IPoseHistory& PoseHistory;
 };
 
 struct FHistoricalPoseIndex
diff --git a/Engine/Plugins/Animation/PoseSearch/Source/Runtime/Public/PoseSearch/PoseSearchLibrary.h b/Engine/Plugins/Animation/PoseSearch/Source/Runtime/Public/PoseSearch/PoseSearchLibrary.h
--- a/Engine/Plugins/Animation/PoseSearch/Source/Runtime/Public/PoseSearch/PoseSearchLibrary.h
+++ b/Engine/Plugins/Animation/PoseSearch/Source/Runtime/Public/PoseSearch/PoseSearchLibrary.h
@@ -18,6 +18,9 @@ namespace UE::PoseSearch
 	struct FSearchContext;
 } // namespace UE::PoseSearch
 
+struct FAnimationUpdateContext;
+struct FAnimNode_PoseSearchHistoryCollector_Base;
+
 UENUM()
 enum class EPoseSearchInterruptMode : uint8
 {
@@ -39,8 +42,6 @@ enum class EPoseSearchInterruptMode : uint8
 	ForceInterruptAndInvalidateContinuingPose,
 };
 
-struct FAnimationUpdateContext;
-
 struct FMotionMatchingState
 {
 	// Reset the state to a default state using the current Database
@@ -247,5 +248,9 @@ public:
 		const FPoseSearchContinuingProperties& ContinuingProperties,
 		const int32 DebugSessionUniqueIdentifier,
 		float DesiredPermutationTimeOffset = 0.f);
+
+	static const FAnimNode_PoseSearchHistoryCollector_Base* FindPoseHistoryNode(
+		const FName PoseHistoryName,
+		const UAnimInstance* AnimInstance);
 };
 
diff --git a/Engine/Plugins/Animation/PoseSearch/Source/Runtime/Public/PoseSearch/PoseSearchResult.h b/Engine/Plugins/Animation/PoseSearch/Source/Runtime/Public/PoseSearch/PoseSearchResult.h
--- a/Engine/Plugins/Animation/PoseSearch/Source/Runtime/Public/PoseSearch/PoseSearchResult.h
+++ b/Engine/Plugins/Animation/PoseSearch/Source/Runtime/Public/PoseSearch/PoseSearchResult.h
@@ -32,11 +32,11 @@ struct FSearchResult
 	// state back to the default state.
 	void Update(float NewAssetTime);
 
-	bool IsValid() const;
+	bool IsValid() const { return PoseIdx != INDEX_NONE && Database != nullptr; }
 
 	void Reset();
 
-	const FSearchIndexAsset* GetSearchIndexAsset(bool bMandatory = false) const;
+	POSESEARCH_API const FSearchIndexAsset* GetSearchIndexAsset(bool bMandatory = false) const;
 	
 	bool CanAdvance(float DeltaTime) const;
 };
