motion matching - fix linking issue for UPoseSearchSchema::GetRoledSkeletons

diff --git a/Engine/Plugins/Animation/PoseSearch/Source/Runtime/Private/PoseSearchSchema.cpp b/Engine/Plugins/Animation/PoseSearch/Source/Runtime/Private/PoseSearchSchema.cpp
--- a/Engine/Plugins/Animation/PoseSearch/Source/Runtime/Private/PoseSearchSchema.cpp
+++ b/Engine/Plugins/Animation/PoseSearch/Source/Runtime/Private/PoseSearchSchema.cpp
@@ -131,11 +131,6 @@ const UE::PoseSearch::FRole UPoseSearchSchema::GetDefaultRole() const
 	}
 	return UE::PoseSearch::DefaultRole;
 }
-
-const TArray<FPoseSearchRoledSkeleton>& UPoseSearchSchema::GetRoledSkeletons() const
-{
-	return Skeletons;
-}
 #endif // WITH_EDITOR
 
 USkeleton* UPoseSearchSchema::GetSkeleton(const UE::PoseSearch::FRole& Role) const
diff --git a/Engine/Plugins/Animation/PoseSearch/Source/Runtime/Public/PoseSearch/PoseSearchSchema.h b/Engine/Plugins/Animation/PoseSearch/Source/Runtime/Public/PoseSearch/PoseSearchSchema.h
--- a/Engine/Plugins/Animation/PoseSearch/Source/Runtime/Public/PoseSearch/PoseSearchSchema.h
+++ b/Engine/Plugins/Animation/PoseSearch/Source/Runtime/Public/PoseSearch/PoseSearchSchema.h
@@ -150,7 +150,7 @@ public:
 	virtual void PostEditChangeProperty(struct FPropertyChangedEvent& PropertyChangedEvent) override;
 	const UE::PoseSearch::FRole GetDefaultRole() const;
 #endif // WITH_EDITOR
-	const TArray<FPoseSearchRoledSkeleton>& GetRoledSkeletons() const;
+	const TArray<FPoseSearchRoledSkeleton>& GetRoledSkeletons() const { return Skeletons; }
 
 	TConstArrayView<float> BuildQuery(UE::PoseSearch::FSearchContext& SearchContext) const;
 
