motion matching - posesearch debugger improvements for multi character MM. Generating a SearchId per search instead of relying on a NodeId, as well as using a separate rewind debugger track per search

diff --git a/Engine/Plugins/Animation/PoseSearch/Source/Editor/Private/PoseSearchDatabaseEditorReflection.cpp b/Engine/Plugins/Animation/PoseSearch/Source/Editor/Private/PoseSearchDatabaseEditorReflection.cpp
--- a/Engine/Plugins/Animation/PoseSearch/Source/Editor/Private/PoseSearchDatabaseEditorReflection.cpp
+++ b/Engine/Plugins/Animation/PoseSearch/Source/Editor/Private/PoseSearchDatabaseEditorReflection.cpp
@@ -10,8 +10,6 @@
 
 #define LOCTEXT_NAMESPACE "UPoseSearchDatabaseReflection"
 
-#if WITH_EDITOR
-
 void UPoseSearchDatabaseReflectionBase::SetSourceLink(
 	const TWeakPtr<UE::PoseSearch::FDatabaseAssetTreeNode>& InWeakAssetTreeNode,
 	const TSharedPtr<UE::PoseSearch::SDatabaseAssetTree>& InAssetTreeWidget)
@@ -169,8 +167,6 @@ bool UPoseSearchDatabaseMultiSequenceReflection::ApplyChanges()
 	return false;
 }
 
-#endif // WITH_EDITOR
-
 void UPoseSearchDatabaseStatistics::Initialize(const UPoseSearchDatabase* PoseSearchDatabase)
 {
 	static FText TimeFormat = LOCTEXT("TimeFormat", "{0} {0}|plural(one=Second,other=Seconds)");
diff --git a/Engine/Plugins/Animation/PoseSearch/Source/Editor/Private/PoseSearchDatabaseViewModel.cpp b/Engine/Plugins/Animation/PoseSearch/Source/Editor/Private/PoseSearchDatabaseViewModel.cpp
--- a/Engine/Plugins/Animation/PoseSearch/Source/Editor/Private/PoseSearchDatabaseViewModel.cpp
+++ b/Engine/Plugins/Animation/PoseSearch/Source/Editor/Private/PoseSearchDatabaseViewModel.cpp
@@ -217,6 +217,11 @@ bool FDatabasePreviewActor::DrawPreviewActors(TConstArrayView<FDatabasePreviewAc
 
 	for (const FDatabasePreviewActor& PreviewActor : PreviewActors)
 	{
+		if (EAsyncBuildIndexResult::Success != FAsyncPoseSearchDatabasesManagement::RequestAsyncBuildIndex(PoseSearchDatabase, ERequestAsyncBuildFlag::ContinueRequest))
+		{
+			return false;
+		}
+
 		if (!PoseSearchDatabase->GetSearchIndex().IsValidPoseIndex(PreviewActor.GetCurrentPoseIndex()))
 		{
 			return false;
diff --git a/Engine/Plugins/Animation/PoseSearch/Source/Editor/Private/PoseSearchDebugger.cpp b/Engine/Plugins/Animation/PoseSearch/Source/Editor/Private/PoseSearchDebugger.cpp
--- a/Engine/Plugins/Animation/PoseSearch/Source/Editor/Private/PoseSearchDebugger.cpp
+++ b/Engine/Plugins/Animation/PoseSearch/Source/Editor/Private/PoseSearchDebugger.cpp
@@ -23,30 +23,62 @@ typedef SCurveTimelineView::FTimelineCurveData::CurvePoint FCurvePoint;
 class SCostCurveTimelineView : public SCurveTimelineView
 {
 public:
+
+	SLATE_BEGIN_ARGS(SCostCurveTimelineView) {}
+		SLATE_ATTRIBUTE(FLinearColor, CurveColor)
+	SLATE_END_ARGS()
+
+	void Construct( const FArguments& InArgs );
+
 	TRange<double> GetViewRange() const { return ViewRange.Get(); }
 	virtual FReply OnMouseMove(const FGeometry& MyGeometry, const FPointerEvent& MouseEvent) override { return FReply::Unhandled(); }
+
+	TSharedPtr<SCostCurveTimelineView::FTimelineCurveData> CurveData;
 };
 
+void SCostCurveTimelineView::Construct(const FArguments& InArgs)
+{
+	CurveData = MakeShared<SCurveTimelineView::FTimelineCurveData>();
+
+	SCurveTimelineView::FArguments CurveTimelineViewArgs;
+
+	CurveTimelineViewArgs
+	.CurveColor(InArgs._CurveColor)
+	.ViewRange_Lambda([]()
+	{
+		return IRewindDebugger::Instance()->GetCurrentViewRange();
+	})
+	.RenderFill(false)
+	.CurveData_Lambda([this]()
+	{
+		return CurveData;
+	});
+
+	SCurveTimelineView::Construct(CurveTimelineViewArgs);
+}
+
+///////////////////////////////////////////////////////
+// SCostTimelineView
 class SCostTimelineView : public SOverlay
 {
 public:
-	SLATE_BEGIN_ARGS(SCostTimelineView) {}
+	SLATE_BEGIN_ARGS(SCostTimelineView)
+	: _SearchId(0)
+	{}
+		SLATE_ARGUMENT( int32, SearchId )
+
 	SLATE_END_ARGS()
 
 	void Construct(const FArguments& InArgs);
 	void UpdateInternal(uint64 ObjectId);
+	int32 GetSearchId() const { return SearchId; }
 
 protected:
 	virtual FReply OnMouseMove(const FGeometry& MyGeometry, const FPointerEvent& MouseEvent) override;
 
 	TSharedPtr<SCostCurveTimelineView> BestCostView;
-	TSharedPtr<SCostCurveTimelineView::FTimelineCurveData> BestCostData;
-
 	TSharedPtr<SCostCurveTimelineView> BruteForceCostView;
-	TSharedPtr<SCostCurveTimelineView::FTimelineCurveData> BruteForceCostData;
-
 	TSharedPtr<SCostCurveTimelineView> BestPosePosView;
-	TSharedPtr<SCostCurveTimelineView::FTimelineCurveData> BestPosePosData;
 
 	TSharedPtr<SToolTip> CostToolTip;
 
@@ -54,49 +86,16 @@ protected:
 	FText ToolTipCost;
 	FText ToolTipCostBruteForce;
 	FText ToolTipBestPosePos;
+	int32 SearchId = 0;
 };
 
 void SCostTimelineView::Construct(const FArguments& InArgs)
 {
-	BestCostData = MakeShared<SCurveTimelineView::FTimelineCurveData>();
-	BruteForceCostData = MakeShared<SCurveTimelineView::FTimelineCurveData>();
-	BestPosePosData = MakeShared<SCurveTimelineView::FTimelineCurveData>();
-	
-	BestCostView = SNew(SCostCurveTimelineView)
-	.CurveColor(FLinearColor::White)
-	.ViewRange_Lambda([]()
-	{
-		return IRewindDebugger::Instance()->GetCurrentViewRange();
-	})
-	.RenderFill(false)
-	.CurveData_Lambda([this]()
-	{
-		return BestCostData;
-	});
-
-	BruteForceCostView = SNew(SCostCurveTimelineView)
-	.CurveColor(FLinearColor::Red)
-	.ViewRange_Lambda([]()
-	{
-		return IRewindDebugger::Instance()->GetCurrentViewRange();
-	})
-	.RenderFill(false)
-	.CurveData_Lambda([this]()
-	{
-		return BruteForceCostData;
-	});
+	SearchId = InArgs._SearchId;
 
-	BestPosePosView = SNew(SCostCurveTimelineView)
-	.CurveColor(FLinearColor::Blue)
-	.ViewRange_Lambda([]()
-	{
-		return IRewindDebugger::Instance()->GetCurrentViewRange();
-	})
-	.RenderFill(false)
-	.CurveData_Lambda([this]()
-	{
-		return BestPosePosData;
-	});
+	BestCostView = SNew(SCostCurveTimelineView).CurveColor(FLinearColor::White);
+	BruteForceCostView = SNew(SCostCurveTimelineView).CurveColor(FLinearColor::Red);
+	BestPosePosView = SNew(SCostCurveTimelineView).CurveColor(FLinearColor::Blue);
 		
 	AddSlot()
 	[
@@ -122,30 +121,25 @@ void SCostTimelineView::UpdateInternal(uint64 ObjectId)
 	{
 		TraceServices::FAnalysisSessionReadScope SessionReadScope(*AnalysisSession);
 
-		TArray<FCurvePoint>& BestCostPoints = BestCostData->Points;
-		BestCostPoints.Reset();
-
-		TArray<FCurvePoint>& BruteForceCostPoints = BruteForceCostData->Points;
-		BruteForceCostPoints.Reset();
-
-		TArray<FCurvePoint>& BestPosePosPoints = BestPosePosData->Points;
-		BestPosePosPoints.Reset();
+		BestCostView->CurveData->Points.Reset();
+		BruteForceCostView->CurveData->Points.Reset();
+		BestPosePosView->CurveData->Points.Reset();
 
 		// convert time range to from rewind debugger times to profiler times
 		TRange<double> TraceTimeRange = RewindDebugger->GetCurrentTraceRange();
 		double StartTime = TraceTimeRange.GetLowerBoundValue();
 		double EndTime = TraceTimeRange.GetUpperBoundValue();
 
-		PoseSearchProvider->EnumerateMotionMatchingStateTimelines(ObjectId, [StartTime, EndTime, &BestCostPoints, &BruteForceCostPoints, &BestPosePosPoints](const FTraceProvider::FMotionMatchingStateTimeline& InTimeline)
+		PoseSearchProvider->EnumerateMotionMatchingStateTimelines(ObjectId, [StartTime, EndTime, this](const FTraceProvider::FMotionMatchingStateTimeline& InTimeline)
 		{
 			// this isn't very efficient, and it gets called every frame.  will need optimizing
-			InTimeline.EnumerateEvents(StartTime, EndTime, [StartTime, EndTime, &BestCostPoints, &BruteForceCostPoints, &BestPosePosPoints](double InStartTime, double InEndTime, uint32 InDepth, const FTraceMotionMatchingStateMessage& InMessage)
+			InTimeline.EnumerateEvents(StartTime, EndTime, [StartTime, EndTime, this](double InStartTime, double InEndTime, uint32 InDepth, const FTraceMotionMatchingStateMessage& InMessage)
 			{
-				if (InEndTime > StartTime && InStartTime < EndTime)
+				if (InMessage.GetSearchId() == SearchId && InEndTime > StartTime && InStartTime < EndTime)
 				{
-					BestCostPoints.Add({ InMessage.RecordingTime, InMessage.SearchBestCost });
-					BruteForceCostPoints.Add({ InMessage.RecordingTime, InMessage.SearchBruteForceCost });
-					BestPosePosPoints.Add({ InMessage.RecordingTime, float(InMessage.SearchBestPosePos) });
+					BestCostView->CurveData->Points.Add({ InMessage.RecordingTime, InMessage.SearchBestCost });
+					BruteForceCostView->CurveData->Points.Add({ InMessage.RecordingTime, InMessage.SearchBruteForceCost });
+					BestPosePosView->CurveData->Points.Add({ InMessage.RecordingTime, float(InMessage.SearchBestPosePos) });
 				}
 				return TraceServices::EEventEnumerate::Continue;
 			});
@@ -153,19 +147,73 @@ void SCostTimelineView::UpdateInternal(uint64 ObjectId)
 
 		float MinValue = UE_MAX_FLT;
 		float MaxValue = -UE_MAX_FLT;
-		for (const FCurvePoint& CurvePoint : BestCostPoints)
+
+		bool bAnyInvalidBestCostPoints = false;
+		bool bAnyInvalidBruteForceCostPoints = false;
+
+		bool bAnyValidBestCostPoints = false;
+		bool bAnyValidBruteForceCostPoints = false;
+		for (const FCurvePoint& CurvePoint : BestCostView->CurveData->Points)
 		{
-			MinValue = FMath::Min(MinValue, CurvePoint.Value);
-			MaxValue = FMath::Max(MaxValue, CurvePoint.Value);
+			if (FPoseSearchCost::IsCostValid(CurvePoint.Value))
+			{
+				MinValue = FMath::Min(MinValue, CurvePoint.Value);
+				MaxValue = FMath::Max(MaxValue, CurvePoint.Value);
+				bAnyValidBestCostPoints = true;
+			}
+			else
+			{
+				bAnyInvalidBestCostPoints = true;
+			}
 		}
-		for (const FCurvePoint& CurvePoint : BruteForceCostPoints)
+		for (const FCurvePoint& CurvePoint : BruteForceCostView->CurveData->Points)
 		{
-			MinValue = FMath::Min(MinValue, CurvePoint.Value);
-			MaxValue = FMath::Max(MaxValue, CurvePoint.Value);
+			if (FPoseSearchCost::IsCostValid(CurvePoint.Value))
+			{
+				MinValue = FMath::Min(MinValue, CurvePoint.Value);
+				MaxValue = FMath::Max(MaxValue, CurvePoint.Value);
+				bAnyValidBruteForceCostPoints = true;
+			}
+			else
+			{
+				bAnyInvalidBruteForceCostPoints = true;
+			}
+		}
+
+		if ((bAnyInvalidBestCostPoints && bAnyValidBestCostPoints) || (bAnyInvalidBruteForceCostPoints && bAnyValidBruteForceCostPoints))
+		{
+			// highliting invalid cost points
+			const float InvalidCostValue = (MaxValue - MinValue) * 2 + MinValue;
+			MaxValue = InvalidCostValue;
+		}
+
+		if (bAnyInvalidBestCostPoints)
+		{
+			for (FCurvePoint& CurvePoint : BestCostView->CurveData->Points)
+			{
+				CurvePoint.Value = FMath::Min(MaxValue, CurvePoint.Value);
+			}
 		}
 
 		BestCostView->SetFixedRange(MinValue, MaxValue);
-		BruteForceCostView->SetFixedRange(MinValue, MaxValue);
+
+		if (bAnyValidBruteForceCostPoints)
+		{
+			if (bAnyInvalidBruteForceCostPoints)
+			{
+				for (FCurvePoint& CurvePoint : BruteForceCostView->CurveData->Points)
+				{
+					CurvePoint.Value = FMath::Min(MaxValue, CurvePoint.Value);
+				}
+			}
+
+			BruteForceCostView->SetFixedRange(MinValue, MaxValue);
+			BruteForceCostView->SetVisibility(EVisibility::Visible);
+		}
+		else
+		{
+			BruteForceCostView->SetVisibility(EVisibility::Hidden);
+		}
 	}
 }
 
@@ -183,7 +231,7 @@ FReply SCostTimelineView::OnMouseMove(const FGeometry& MyGeometry, const FPointe
 		const double TargetTime = RangeToScreen.LocalXToInput(HitPosition.X);
 
 		// Get curve value at given time
-		const TArray<FCurvePoint>& CurvePoints = BestCostData->Points;
+		const TArray<FCurvePoint>& CurvePoints = BestCostView->CurveData->Points;
 		const int32 NumPoints = CurvePoints.Num();
 
 		if (NumPoints > 0)
@@ -205,8 +253,8 @@ FReply SCostTimelineView::OnMouseMove(const FGeometry& MyGeometry, const FPointe
 
 					const float Time = CurvePoints[TargetPointIndex].Time;
 					const float BestCost = CurvePoints[TargetPointIndex].Value;
-					const float BruteForceCost = BruteForceCostData->Points[TargetPointIndex].Value;
-					const int32 BestPosePos = FMath::RoundToInt(BestPosePosData->Points[TargetPointIndex].Value);
+					const float BruteForceCost = BruteForceCostView->CurveData->Points[TargetPointIndex].Value;
+					const int32 BestPosePos = FMath::RoundToInt(BestPosePosView->CurveData->Points[TargetPointIndex].Value);
 
 					// Tooltip text formatting
 					FNumberFormattingOptions FormattingOptions;
@@ -216,7 +264,7 @@ FReply SCostTimelineView::OnMouseMove(const FGeometry& MyGeometry, const FPointe
 					ToolTipTime = FText::Format(LOCTEXT("CostTimelineViewToolTip_TimeFormat", "Search Time: {0}"), FText::AsNumber(Time, &FormattingOptions));
 					ToolTipCost = FText::Format(LOCTEXT("CostTimelineViewToolTip_CostFormat", "Search Cost: {0}"), FText::AsNumber(BestCost, &FormattingOptions));
 
-					if (FMath::IsNearlyEqual(BestCost, BruteForceCost))
+					if (!FPoseSearchCost::IsCostValid(BruteForceCost) || FMath::IsNearlyEqual(BestCost, BruteForceCost))
 					{
 						ToolTipCostBruteForce = FText::GetEmpty();
 					}
@@ -230,7 +278,7 @@ FReply SCostTimelineView::OnMouseMove(const FGeometry& MyGeometry, const FPointe
 					{
 						SetToolTip(
 							SAssignNew(CostToolTip, SToolTip)
-							.BorderImage(FCoreStyle::Get().GetBrush("ToolTip.BrightBackground"))
+							.BorderImage(FCoreStyle::Get().GetBrush("ToolTip.Background"))
 							[
 								SNew(SVerticalBox)
 								+ SVerticalBox::Slot()
@@ -245,14 +293,14 @@ FReply SCostTimelineView::OnMouseMove(const FGeometry& MyGeometry, const FPointe
 									SNew(STextBlock)
 									.Text_Lambda([this]() { return ToolTipBestPosePos; })
 									.Font(FCoreStyle::Get().GetFontStyle("ToolTip.LargerFont"))
-									.ColorAndOpacity(FLinearColor::Black)
+									.ColorAndOpacity(FLinearColor::Blue)
 								]
 								+ SVerticalBox::Slot()
 								[
 									SNew(STextBlock)
 									.Text_Lambda([this]() { return ToolTipCost; })
 									.Font(FCoreStyle::Get().GetFontStyle("ToolTip.LargerFont"))
-									.ColorAndOpacity(FLinearColor::Black)
+									.ColorAndOpacity(FLinearColor::White)
 								]
 								+ SVerticalBox::Slot()
 								[
@@ -260,7 +308,7 @@ FReply SCostTimelineView::OnMouseMove(const FGeometry& MyGeometry, const FPointe
 									.Visibility_Lambda([this]() { return ToolTipCostBruteForce.IsEmpty() ? EVisibility::Collapsed : EVisibility::Visible; })
 									.Text_Lambda([this]() { return ToolTipCostBruteForce; })
 									.Font(FCoreStyle::Get().GetFontStyle("ToolTip.LargerFont"))
-									.ColorAndOpacity(FLinearColor::Black)
+									.ColorAndOpacity(FLinearColor::Red)
 								]
 							]);
 					}
@@ -274,6 +322,8 @@ FReply SCostTimelineView::OnMouseMove(const FGeometry& MyGeometry, const FPointe
 	return FReply::Unhandled();
 }
 
+///////////////////////////////////////////////////////
+// FDebugger
 FDebugger* FDebugger::Debugger;
 void FDebugger::Initialize()
 {
@@ -356,89 +406,173 @@ TSharedPtr<FDebuggerViewModel> FDebugger::GetViewModel(uint64 InAnimInstanceId)
 	return nullptr;
 }
 
-TSharedPtr<SDebuggerView> FDebugger::GenerateInstance(uint64 InAnimInstanceId)
+TSharedPtr<SDebuggerView> FDebugger::GenerateInstance(uint64 InAnimInstanceId, int32 InWantedSearchId)
 {
 	ViewModels.Add_GetRef(MakeShared<FDebuggerViewModel>(InAnimInstanceId))->RewindDebugger.BindStatic(&FDebugger::GetRewindDebugger);
 
-	TSharedPtr<SDebuggerView> DebuggerView;
-
-	SAssignNew(DebuggerView, SDebuggerView, InAnimInstanceId)
+	TSharedPtr<SDebuggerView> DebuggerViewSharedPtr;
+	SAssignNew(DebuggerViewSharedPtr, SDebuggerView, InAnimInstanceId, InWantedSearchId)
 		.ViewModel_Static(&FDebugger::GetViewModel, InAnimInstanceId)
 		.OnViewClosed_Static(&FDebugger::OnViewClosed);
 
-	return DebuggerView;
+	DebuggerView = DebuggerViewSharedPtr;
+	return DebuggerViewSharedPtr;
 }
 
-FDebuggerTrack::FDebuggerTrack(uint64 InObjectId)
-: ObjectId(InObjectId)
+///////////////////////////////////////////////////////
+// FSearchTrack
+FSearchTrack::FSearchTrack(uint64 InObjectId, int32 InSearchId, FText InTrackName)
+: RewindDebugger::FRewindDebuggerTrack()
+, CostTimelineView(SNew(SCostTimelineView).SearchId(InSearchId))
+, ObjectId(InObjectId)
+, TrackName(InTrackName)
+, Icon(FSlateIconFinder::FindIconForClass(UAnimInstance::StaticClass()))
 {
-	CostTimelineView = SNew(SCostTimelineView);
 }
 
-FSlateIcon FDebuggerTrack::GetIconInternal()
+int32 FSearchTrack::GetSearchId() const
 {
-#if WITH_EDITOR
-	return FSlateIconFinder::FindIconForClass(UAnimInstance::StaticClass());
-#else
-	return FSlateIcon();
-#endif
+	return CostTimelineView->GetSearchId();
 }
 
-bool FDebuggerTrack::UpdateInternal()
+FText FSearchTrack::GetDisplayNameInternal() const
 {
-	TRACE_CPUPROFILER_EVENT_SCOPE(PoseSearchDebuggerTrack::UpdateInternal);
-	CostTimelineView->UpdateInternal(ObjectId);
+	return TrackName;
+}
 
-	if (TSharedPtr<IRewindDebuggerView> PinnedView = View.Pin())
-	{
-		PinnedView->SetTimeMarker(IRewindDebugger::Instance()->CurrentTraceTime());
-	}
+bool FSearchTrack::UpdateInternal()
+{
+	TRACE_CPUPROFILER_EVENT_SCOPE(PoseSearchSearchTrack::UpdateInternal);
+	CostTimelineView->UpdateInternal(ObjectId);
 
 	return false;
 }
 
-FName FDebuggerTrack::GetNameInternal() const
+TSharedPtr<SWidget> FSearchTrack::GetTimelineViewInternal()
 {
-	static const FName Name("PoseSearchDebugger");
-	return Name;
+	return CostTimelineView;
 }
 
-FText FDebuggerTrack::GetDisplayNameInternal() const
+TSharedPtr<SWidget> FSearchTrack::GetDetailsViewInternal()
 {
-	return LOCTEXT("PoseSearchDebuggerTabTitle", "Pose Search");
+	return FDebugger::Get()->GenerateInstance(ObjectId, GetSearchId());
 }
 
-TSharedPtr<SWidget> FDebuggerTrack::GetTimelineViewInternal()
+///////////////////////////////////////////////////////
+// FDebuggerTrack
+FDebuggerTrack::FDebuggerTrack(uint64 InObjectId)
+: RewindDebugger::FRewindDebuggerTrack()
+, ObjectId(InObjectId)
+, Icon(FSlateIconFinder::FindIconForClass(UAnimInstance::StaticClass()))
 {
-	return CostTimelineView;
 }
 
-TSharedPtr<SWidget> FDebuggerTrack::GetDetailsViewInternal()
+bool FDebuggerTrack::UpdateInternal()
 {
-	TSharedPtr<IRewindDebuggerView> RewindDebuggerView = FDebugger::Get()->GenerateInstance(ObjectId);
-	View = RewindDebuggerView;
-	return RewindDebuggerView;
-}
+	TRACE_CPUPROFILER_EVENT_SCOPE(PoseSearchDebuggerTrack::UpdateInternal);
+	
+	IRewindDebugger* RewindDebugger = IRewindDebugger::Instance();
 
-// FDebuggerTrackCreator
-///////////////////////////////////////////////////
+	if (TSharedPtr<IRewindDebuggerView> PinnedView = FDebugger::Get()->GetDebuggerView().Pin())
+	{
+		PinnedView->SetTimeMarker(RewindDebugger->CurrentTraceTime());
+	}
 
-FName FDebuggerTrackCreator::GetTargetTypeNameInternal() const
-{
-	static FName TargetTypeName = "AnimInstance";
-	return TargetTypeName;
+	bool bChanged = false;
+
+	const TraceServices::IAnalysisSession* AnalysisSession = RewindDebugger->GetAnalysisSession();
+	check(AnalysisSession);
+	if (const FTraceProvider* PoseSearchProvider = AnalysisSession->ReadProvider<FTraceProvider>(FTraceProvider::ProviderName))
+	{
+		TraceServices::FAnalysisSessionReadScope SessionReadScope(*AnalysisSession);
+
+		// convert time range to from rewind debugger times to profiler times
+		TRange<double> TraceTimeRange = RewindDebugger->GetCurrentTraceRange();
+		double StartTime = TraceTimeRange.GetLowerBoundValue();
+		double EndTime = TraceTimeRange.GetUpperBoundValue();
+
+		TArray<int32, TInlineAllocator<64>> OldSearchIds;
+		for (TSharedPtr<FSearchTrack>& SearchTrack : SearchTracks)
+		{
+			OldSearchIds.Add(SearchTrack->GetSearchId());
+		}
+
+		TMap<int32, FText, TInlineSetAllocator<64>> SearchIdNames;
+		PoseSearchProvider->EnumerateMotionMatchingStateTimelines(ObjectId, [StartTime, EndTime, &SearchIdNames](const FTraceProvider::FMotionMatchingStateTimeline& InTimeline)
+			{
+				// this isn't very efficient, and it gets called every frame.  will need optimizing
+				InTimeline.EnumerateEvents(StartTime, EndTime, [StartTime, EndTime, &SearchIdNames](double InStartTime, double InEndTime, uint32 InDepth, const FTraceMotionMatchingStateMessage& InMessage)
+					{
+						if (!SearchIdNames.Find(InMessage.GetSearchId()) && InEndTime > StartTime && InStartTime < EndTime)
+						{
+							SearchIdNames.Add(InMessage.GetSearchId()) = InMessage.GenerateSearchName();
+						}
+						return TraceServices::EEventEnumerate::Continue;
+					});
+			});
+
+		TArray<int32, TInlineAllocator<64>> SearchIds;
+		for (TPair<int32, FText> SearchIdNamePair : SearchIdNames)
+		{
+			SearchIds.Add(SearchIdNamePair.Key);
+		}
+		SearchIds.StableSort();
+
+		if (SearchIds != OldSearchIds)
+		{
+			TMap<int32, TSharedPtr<FSearchTrack>, TInlineSetAllocator<64>> OldSearchIdsMap;
+			for (TSharedPtr<FSearchTrack>& SearchTrack : SearchTracks)
+			{
+				OldSearchIdsMap.Add(SearchTrack->GetSearchId()) = SearchTrack;
+			}
+
+			SearchTracks.SetNum(SearchIds.Num());
+			for (int32 SearchIdIndex = 0; SearchIdIndex < SearchIds.Num(); ++SearchIdIndex)
+			{
+				if (TSharedPtr<FSearchTrack>* SearchTrack = OldSearchIdsMap.Find(SearchIds[SearchIdIndex]))
+				{
+					SearchTracks[SearchIdIndex] = *SearchTrack;
+				}
+				else
+				{
+					SearchTracks[SearchIdIndex] = MakeShared<FSearchTrack>(ObjectId, SearchIds[SearchIdIndex], SearchIdNames[SearchIds[SearchIdIndex]]);
+				}
+			}
+
+			bChanged = true;
+		}
+		
+		for (TSharedPtr<FSearchTrack>& SearchTrack : SearchTracks)
+		{
+			if (SearchTrack.IsValid())
+			{
+				bChanged |= SearchTrack->Update();
+			}
+		}
+	}
+
+	return bChanged;
 }
-	
-static const FName PoseSearchDebuggerName("PoseSearchDebugger");
 
-FName FDebuggerTrackCreator::GetNameInternal() const
+TSharedPtr<SWidget> FDebuggerTrack::GetDetailsViewInternal()
 {
-	return PoseSearchDebuggerName;
+	return FDebugger::Get()->GenerateInstance(ObjectId);
 }
 
+void FDebuggerTrack::IterateSubTracksInternal(TFunction<void(TSharedPtr<FRewindDebuggerTrack> SubTrack)> IteratorFunction)
+{
+	for(TSharedPtr<FSearchTrack>& SearchTrack : SearchTracks)
+	{
+		IteratorFunction(SearchTrack);
+	}
+};
+
+// FDebuggerTrackCreator
+///////////////////////////////////////////////////
+
 void FDebuggerTrackCreator::GetTrackTypesInternal(TArray<RewindDebugger::FRewindDebuggerTrackType>& Types) const
 {
-	Types.Add({PoseSearchDebuggerName, LOCTEXT("Pose Search", "Pose Search")});
+	Types.Add({ GetNameInternal(), LOCTEXT("Pose Search", "Pose Search") });
 }
 
 TSharedPtr<RewindDebugger::FRewindDebuggerTrack> FDebuggerTrackCreator::CreateTrackInternal(uint64 ObjectId) const
diff --git a/Engine/Plugins/Animation/PoseSearch/Source/Editor/Private/PoseSearchDebugger.h b/Engine/Plugins/Animation/PoseSearch/Source/Editor/Private/PoseSearchDebugger.h
--- a/Engine/Plugins/Animation/PoseSearch/Source/Editor/Private/PoseSearchDebugger.h
+++ b/Engine/Plugins/Animation/PoseSearch/Source/Editor/Private/PoseSearchDebugger.h
@@ -5,6 +5,7 @@
 #include "CoreMinimal.h"
 #include "IRewindDebuggerExtension.h"
 #include "IRewindDebuggerTrackCreator.h"
+#include "PoseSearch/PoseSearchDefines.h"
 #include "RewindDebuggerTrack.h"
 #include "SCurveTimelineView.h"
 
@@ -42,7 +43,8 @@ public:
 	static const IRewindDebugger* GetRewindDebugger();
 
 	/** Generates the slate debugger view widget */
-	TSharedPtr<SDebuggerView> GenerateInstance(uint64 InAnimInstanceId);
+	TSharedPtr<SDebuggerView> GenerateInstance(uint64 InAnimInstanceId, int32 InWantedSearchId = InvalidSearchId);
+	TWeakPtr<SDebuggerView> GetDebuggerView() { return DebuggerView; }
 
 private:
 	/** Removes the reference from the model array when closed, destroying the model */
@@ -57,40 +59,64 @@ private:
 	/** List of all active debugger instances */
 	TArray<TSharedRef<FDebuggerViewModel>> ViewModels;
 	
+	TWeakPtr<SDebuggerView> DebuggerView;
+
 	/** Internal instance */
 	static FDebugger* Debugger;
 };
 
-/**
- * Creates the slate widgets associated with the PoseSearch debugger
- * when prompted by the Rewind Debugger
- */
-class FDebuggerTrack : public RewindDebugger::FRewindDebuggerTrack
+class FSearchTrack : public RewindDebugger::FRewindDebuggerTrack
 {
 public:
-	FDebuggerTrack(uint64 InObjectId);
+	FSearchTrack(uint64 InObjectId, int32 InSearchId, FText InTrackName);
+	int32 GetSearchId() const;
 
 private:
-	virtual FSlateIcon GetIconInternal() override;
+	virtual FSlateIcon GetIconInternal() override { return Icon;  }
 	virtual TSharedPtr<SWidget> GetTimelineViewInternal() override;
 	virtual TSharedPtr<SWidget> GetDetailsViewInternal() override;
-	virtual FName GetNameInternal() const override;
+	virtual FName GetNameInternal() const override { return "PoseSearchTrack"; }
 	virtual FText GetDisplayNameInternal() const override;
 	virtual uint64 GetObjectIdInternal() const override { return ObjectId; }
 	virtual bool UpdateInternal() override;
 
 	TSharedPtr<SCostTimelineView> CostTimelineView;
 
-	TWeakPtr<IRewindDebuggerView> View;
+	uint64 ObjectId;
+	FText TrackName;
 	FSlateIcon Icon;
+};
+
+/**
+ * Creates the slate widgets associated with the PoseSearch debugger
+ * when prompted by the Rewind Debugger
+ */
+class FDebuggerTrack : public RewindDebugger::FRewindDebuggerTrack
+{
+public:
+	explicit FDebuggerTrack(uint64 InObjectId);
+
+private:
+	virtual FSlateIcon GetIconInternal() override { return Icon;  }
+	virtual TSharedPtr<SWidget> GetTimelineViewInternal() override { return nullptr; }
+	virtual TSharedPtr<SWidget> GetDetailsViewInternal() override;
+	virtual FName GetNameInternal() const override { return "PoseSearchDebugger"; }
+	virtual FText GetDisplayNameInternal() const override { return NSLOCTEXT("PoseSearchDebugger", "PoseSearchDebuggerTabTitle", "Pose Search"); }
+	virtual uint64 GetObjectIdInternal() const override { return ObjectId; }
+	virtual void IterateSubTracksInternal(TFunction<void(TSharedPtr<FRewindDebuggerTrack> SubTrack)> IteratorFunction) override;
+	virtual bool UpdateInternal() override;
+
 	uint64 ObjectId;
+	FSlateIcon Icon;
+
+	TArray<TSharedPtr<FSearchTrack>> SearchTracks;
 };
 
 class FDebuggerTrackCreator : public RewindDebugger::IRewindDebuggerTrackCreator
 {
 private:
-	virtual FName GetTargetTypeNameInternal() const override;
-	virtual FName GetNameInternal() const override;
+	virtual FName GetTargetTypeNameInternal() const override { return "AnimInstance"; }
+	virtual FName GetNameInternal() const override { return "PoseSearchDebugger"; }
 	virtual void GetTrackTypesInternal(TArray<RewindDebugger::FRewindDebuggerTrackType>& Types) const override;
 	virtual TSharedPtr<RewindDebugger::FRewindDebuggerTrack> CreateTrackInternal(uint64 ObjectId) const override;
 	virtual bool HasDebugInfoInternal(uint64 ObjectId) const override;
diff --git a/Engine/Plugins/Animation/PoseSearch/Source/Editor/Private/PoseSearchDebuggerView.cpp b/Engine/Plugins/Animation/PoseSearch/Source/Editor/Private/PoseSearchDebuggerView.cpp
--- a/Engine/Plugins/Animation/PoseSearch/Source/Editor/Private/PoseSearchDebuggerView.cpp
+++ b/Engine/Plugins/Animation/PoseSearch/Source/Editor/Private/PoseSearchDebuggerView.cpp
@@ -102,7 +102,7 @@ void SDebuggerDetailsView::UpdateReflection(const FTraceMotionMatchingStateMessa
 	Reflection->AnimAngularVelocity = State.AnimAngularVelocity;
 }
 
-void SDebuggerView::Construct(const FArguments& InArgs, uint64 InAnimInstanceId)
+void SDebuggerView::Construct(const FArguments& InArgs, uint64 InAnimInstanceId, int32 InWantedSearchId)
 {
 	ViewModel = InArgs._ViewModel;
 	OnViewClosed = InArgs._OnViewClosed;
@@ -112,7 +112,8 @@ void SDebuggerView::Construct(const FArguments& InArgs, uint64 InAnimInstanceId)
 	check(OnViewClosed.IsBound());
 	
 	AnimInstanceId = InAnimInstanceId;
-	SelectedNodeId = INDEX_NONE;
+	WantedSearchId = InWantedSearchId;
+	SelectedSearchId = InWantedSearchId;
 
 	ChildSlot
 	[
@@ -217,7 +218,7 @@ void SDebuggerView::Tick(const FGeometry& AllottedGeometry, const double InCurre
 			Model->OnUpdate();
 			if (UpdateNodeSelection())
 			{
-				Model->OnUpdateNodeSelection(SelectedNodeId);
+				Model->OnUpdateSearchSelection(SelectedSearchId);
 				UpdateViews();
 			}
 			bUpdated = true;
@@ -392,41 +393,60 @@ bool SDebuggerView::UpdateNodeSelection()
 	TSharedPtr<FDebuggerViewModel> Model = ViewModel.Get();
 	check(Model.IsValid());
 
+	const TArray<FTraceMotionMatchingStateMessage>& MotionMatchingStates = Model->GetMotionMatchingStates();
+
 	// Update selection view if no node selected
-	bool bNodeSelected = SelectedNodeId != INDEX_NONE;
-	if (!bNodeSelected)
+	if (SelectedSearchId != InvalidSearchId)
 	{
-		const TArray<FTraceMotionMatchingStateMessage>& MotionMatchingStates = Model->GetMotionMatchingStates();
-		// Only one active state, bypass selection view
-		if (MotionMatchingStates.Num() == 1)
-		{
-			SelectedNodeId = MotionMatchingStates[0].NodeId;
-			bNodeSelected = true;
-		}
-		// Create selection view with buttons for each node, displaying the database name
-		else
+		if (!MotionMatchingStates.IsEmpty())
 		{
-			SelectionView->ClearChildren();
+			// making sure SelectedSearchId is still valid. if not, let's pick the first available MotionMatchingStates
 			for (const FTraceMotionMatchingStateMessage& MotionMatchingState : MotionMatchingStates)
 			{
-				Model->OnUpdateNodeSelection(MotionMatchingState.NodeId);
-				SelectionView->AddSlot()
-				.HAlign(HAlign_Fill)
-				.VAlign(VAlign_Center)
-				.Padding(10.0f)
-				[
-					SNew(SButton)
-					.Text(FText::FromString(GetNameSafe(Model->GetCurrentDatabase())))
-					.HAlign(HAlign_Center)
-					.VAlign(VAlign_Center)
-					.ContentPadding(10.0f)
-					.OnClicked(this, &SDebuggerView::OnUpdateNodeSelection, MotionMatchingState.NodeId)
-				];
+				if (MotionMatchingState.GetSearchId() == SelectedSearchId)
+				{
+					return true;
+				}
+			}
+
+			if (WantedSearchId == InvalidSearchId)
+			{
+				// SelectedSearchId is not valid, and since there's no WantedSearchId, by specifically double clicking the 
+				// search track instead of selecting "Pose Search" track, we reassign SelectedSearchId to the first valid SearchId
+				SelectedSearchId = MotionMatchingStates[0].GetSearchId();
 			}
 		}
+
+		return true;
+	}
+	
+	// Only one active state, bypass selection view
+	if (MotionMatchingStates.Num() == 1)
+	{
+		SelectedSearchId = MotionMatchingStates[0].GetSearchId();
+		return true;
+	}
+
+	// Create selection view with buttons for each node, displaying the database name
+	SelectionView->ClearChildren();
+	for (const FTraceMotionMatchingStateMessage& MotionMatchingState : MotionMatchingStates)
+	{
+		Model->OnUpdateSearchSelection(MotionMatchingState.GetSearchId());
+		SelectionView->AddSlot()
+		.HAlign(HAlign_Fill)
+		.VAlign(VAlign_Center)
+		.Padding(10.0f)
+		[
+			SNew(SButton)
+			.Text(MotionMatchingState.GenerateSearchName())
+			.HAlign(HAlign_Center)
+			.VAlign(VAlign_Center)
+			.ContentPadding(10.0f)
+			.OnClicked(this, &SDebuggerView::OnUpdateSearchSelection, MotionMatchingState.GetSearchId())
+		];
 	}
 
-	return bNodeSelected;
+	return false;
 }
 
 void SDebuggerView::UpdateViews() const
@@ -461,7 +481,7 @@ int32 SDebuggerView::SelectView() const
 	check(Model.IsValid());
 
 	const bool bNoActiveNodes = Model->GetNodesNum() == 0;
-	const bool bNodeSelectedWithoutData = SelectedNodeId != INDEX_NONE && Model->GetMotionMatchingState() == nullptr;
+	const bool bNodeSelectedWithoutData = SelectedSearchId != InvalidSearchId && Model->GetMotionMatchingState() == nullptr;
 
 	// No active nodes, or node selected has no data
 	if (bNoActiveNodes || bNodeSelectedWithoutData)
@@ -470,7 +490,7 @@ int32 SDebuggerView::SelectView() const
     }
 
 	// Node not selected yet, showcase selection view
-	if (SelectedNodeId == INDEX_NONE)
+	if (SelectedSearchId == InvalidSearchId)
 	{
 		return Selection;
 	}
@@ -490,10 +510,10 @@ void SDebuggerView::OnPoseSelectionChanged(const UPoseSearchDatabase* Database,
 	}
 }
 
-FReply SDebuggerView::OnUpdateNodeSelection(int32 InSelectedNodeId)
+FReply SDebuggerView::OnUpdateSearchSelection(int32 InSelectedSearchId)
 {
-	// -1 will backtrack to selection view
-	SelectedNodeId = InSelectedNodeId;
+	// InvalidSearchId will backtrack to selection view
+	SelectedSearchId = InSelectedSearchId;
 	bUpdated = false;
 	return FReply::Handled();
 }
@@ -538,7 +558,7 @@ TSharedRef<SHorizontalBox> SDebuggerView::GenerateReturnButtonView()
 			.Visibility_Lambda([this] { return ViewModel.Get()->GetNodesNum() > 1 ? EVisibility::Visible : EVisibility::Hidden; })
 			.ButtonStyle(FAppStyle::Get(), "SimpleButton")
 			.ContentPadding( FMargin(1, 0) )
-			.OnClicked(this, &SDebuggerView::OnUpdateNodeSelection, static_cast<int32>(INDEX_NONE))
+			.OnClicked(this, &SDebuggerView::OnUpdateSearchSelection, InvalidSearchId)
 			// Contents of button, icon then text
 			[
 				SNew(SHorizontalBox)
@@ -557,7 +577,7 @@ TSharedRef<SHorizontalBox> SDebuggerView::GenerateReturnButtonView()
 				.VAlign(VAlign_Center)
 				[
 					SNew(STextBlock)
-					.Text(FText::FromString("Return to Database Selection"))
+					.Text(FText::FromString("Return to Search Selection"))
 					.Justification(ETextJustify::Center)
 				]
 			]
diff --git a/Engine/Plugins/Animation/PoseSearch/Source/Editor/Private/PoseSearchDebuggerView.h b/Engine/Plugins/Animation/PoseSearch/Source/Editor/Private/PoseSearchDebuggerView.h
--- a/Engine/Plugins/Animation/PoseSearch/Source/Editor/Private/PoseSearchDebuggerView.h
+++ b/Engine/Plugins/Animation/PoseSearch/Source/Editor/Private/PoseSearchDebuggerView.h
@@ -3,6 +3,7 @@
 #pragma once
 
 #include "IRewindDebuggerView.h"
+#include "PoseSearch/PoseSearchDefines.h"
 
 class IDetailsView;
 class SVerticalBox;
@@ -68,7 +69,7 @@ public:
 	SDebuggerView() = default;
     virtual ~SDebuggerView() override;
 
-	void Construct(const FArguments& InArgs, uint64 InAnimInstanceId);
+	void Construct(const FArguments& InArgs, uint64 InAnimInstanceId, int32 InWantedSearchId);
 	virtual void SetTimeMarker(double InTimeMarker) override;
 	virtual void Tick(const FGeometry& AllottedGeometry, const double InCurrentTime, const float InDeltaTime) override;
 	virtual FName GetName() const override;
@@ -88,7 +89,7 @@ private:
 	int32 SelectView() const;
 
 	/** Callback when a button in the selection view is clicked */
-	FReply OnUpdateNodeSelection(int32 InSelectedNodeId);
+	FReply OnUpdateSearchSelection(int32 InSelectedSearchId);
 
 	void OnPoseSelectionChanged(const UPoseSearchDatabase* Database, int32 PoseIdx, float Time);
 
@@ -107,8 +108,11 @@ private:
 	/** Destroy the debugger instanced when closed */
 	FOnViewClosed OnViewClosed;
 
-	/** Active node being debugged */
-	int32 SelectedNodeId = INDEX_NONE;
+	/** Active User selected SearchId*/
+	int32 WantedSearchId = InvalidSearchId;
+	
+	/** Active SearchId being debugged */
+	int32 SelectedSearchId = InvalidSearchId;
 
 	/** Database view of the motion matching node */
 	TSharedPtr<SDebuggerDatabaseView> DatabaseView;
diff --git a/Engine/Plugins/Animation/PoseSearch/Source/Editor/Private/PoseSearchDebuggerViewModel.cpp b/Engine/Plugins/Animation/PoseSearch/Source/Editor/Private/PoseSearchDebuggerViewModel.cpp
--- a/Engine/Plugins/Animation/PoseSearch/Source/Editor/Private/PoseSearchDebuggerViewModel.cpp
+++ b/Engine/Plugins/Animation/PoseSearch/Source/Editor/Private/PoseSearchDebuggerViewModel.cpp
@@ -187,16 +187,16 @@ void FDebuggerViewModel::OnUpdate()
 	}
 }
 
-void FDebuggerViewModel::OnUpdateNodeSelection(int32 InNodeId)
+void FDebuggerViewModel::OnUpdateSearchSelection(int32 InSearchId)
 {
-	if (InNodeId != INDEX_NONE)
+	if (InSearchId != InvalidSearchId)
 	{
 		// Find node in all motion matching states this frame
 		ActiveMotionMatchingStateIdx = INDEX_NONE;
 		const int32 NodesNum = MotionMatchingStates.Num();
 		for (int32 i = 0; i < NodesNum; ++i)
 		{
-			if (MotionMatchingStates[i].NodeId == InNodeId)
+			if (MotionMatchingStates[i].GetSearchId() == InSearchId)
 			{
 				ActiveMotionMatchingStateIdx = i;
 				break;
diff --git a/Engine/Plugins/Animation/PoseSearch/Source/Editor/Private/PoseSearchDebuggerViewModel.h b/Engine/Plugins/Animation/PoseSearch/Source/Editor/Private/PoseSearchDebuggerViewModel.h
--- a/Engine/Plugins/Animation/PoseSearch/Source/Editor/Private/PoseSearchDebuggerViewModel.h
+++ b/Engine/Plugins/Animation/PoseSearch/Source/Editor/Private/PoseSearchDebuggerViewModel.h
@@ -36,7 +36,7 @@ public:
 	void OnUpdate();
 	
 	/** Updates active motion matching state based on node selection */
-	void OnUpdateNodeSelection(int32 InNodeId);
+	void OnUpdateSearchSelection(int32 InSearchId);
 
 	void SetVerbose(bool bVerbose) { bIsVerbose = bVerbose; }
 	bool IsVerbose() const { return bIsVerbose; }
diff --git a/Engine/Plugins/Animation/PoseSearch/Source/Editor/Private/Trace/PoseSearchTraceProvider.cpp b/Engine/Plugins/Animation/PoseSearch/Source/Editor/Private/Trace/PoseSearchTraceProvider.cpp
--- a/Engine/Plugins/Animation/PoseSearch/Source/Editor/Private/Trace/PoseSearchTraceProvider.cpp
+++ b/Engine/Plugins/Animation/PoseSearch/Source/Editor/Private/Trace/PoseSearchTraceProvider.cpp
@@ -11,29 +11,23 @@ FTraceProvider::FTraceProvider(TraceServices::IAnalysisSession& InSession) : Ses
 {
 }
 
-TSet<int32> FTraceProvider::GetMotionMatchingNodeIds(uint64 InAnimInstanceId) const
-{
-	return MotionMatchingStateTimelineStorage.GetNodeIds(InAnimInstanceId);
-}
-
-bool FTraceProvider::ReadMotionMatchingStateTimeline(uint64 InAnimInstanceId, int32 InNodeId, TFunctionRef<void(const FMotionMatchingStateTimeline&)> Callback) const
+bool FTraceProvider::ReadMotionMatchingStateTimeline(uint64 InAnimInstanceId, int32 InSearchId, TFunctionRef<void(const FMotionMatchingStateTimeline&)> Callback) const
 {
 	Session.ReadAccessCheck();
-	return MotionMatchingStateTimelineStorage.ReadTimeline(InAnimInstanceId, InNodeId, Callback);
+	return MotionMatchingStateTimelineStorage.ReadTimeline(InAnimInstanceId, InSearchId, Callback);
 }
 
 bool FTraceProvider::EnumerateMotionMatchingStateTimelines(uint64 InAnimInstanceId, TFunctionRef<void(const FMotionMatchingStateTimeline&)> Callback) const
 {
 	Session.ReadAccessCheck();
-	return MotionMatchingStateTimelineStorage.EnumerateNodeTimelines(InAnimInstanceId, Callback);
+	return MotionMatchingStateTimelineStorage.EnumerateSearchTimelines(InAnimInstanceId, Callback);
 }
 
-
 void FTraceProvider::AppendMotionMatchingState(const FTraceMotionMatchingStateMessage& InMessage, double InTime)
 {
 	Session.WriteAccessCheck();
 
-	TSharedRef<TraceServices::TPointTimeline<FTraceMotionMatchingStateMessage>> Timeline = MotionMatchingStateTimelineStorage.GetTimeline(Session, InMessage.AnimInstanceId, InMessage.NodeId);
+	TSharedRef<TraceServices::TPointTimeline<FTraceMotionMatchingStateMessage>> Timeline = MotionMatchingStateTimelineStorage.GetTimeline(Session, InMessage.AnimInstanceId, InMessage.GetSearchId());
 	Timeline->AppendEvent(InTime, InMessage);
 	
 	Session.UpdateDurationSeconds(InTime);
diff --git a/Engine/Plugins/Animation/PoseSearch/Source/Editor/Private/Trace/PoseSearchTraceProvider.h b/Engine/Plugins/Animation/PoseSearch/Source/Editor/Private/Trace/PoseSearchTraceProvider.h
--- a/Engine/Plugins/Animation/PoseSearch/Source/Editor/Private/Trace/PoseSearchTraceProvider.h
+++ b/Engine/Plugins/Animation/PoseSearch/Source/Editor/Private/Trace/PoseSearchTraceProvider.h
@@ -18,13 +18,10 @@ class FTraceProvider : public TraceServices::IProvider
 public:
 	explicit FTraceProvider(TraceServices::IAnalysisSession& InSession);
 
-	/** Get all node ids of motion matching nodes appended to the object's timeline thus far */
-	TSet<int32> GetMotionMatchingNodeIds(uint64 InAnimInstanceId) const;
-
 	using FMotionMatchingStateTimeline = TraceServices::ITimeline<FTraceMotionMatchingStateMessage>;
 	
 	/** Read the node-relative timeline info of the AnimInstance and provide it via callback */
-	bool ReadMotionMatchingStateTimeline(uint64 InAnimInstanceId, int32 InNodeId, TFunctionRef<void(const FMotionMatchingStateTimeline&)> Callback) const;
+	bool ReadMotionMatchingStateTimeline(uint64 InAnimInstanceId, int32 InSearchId, TFunctionRef<void(const FMotionMatchingStateTimeline&)> Callback) const;
 	
 	/** Enumerate all node timelines on the AnimInstance and provide them via callback */
 	bool EnumerateMotionMatchingStateTimelines(uint64 InAnimInstanceId, TFunctionRef<void(const FMotionMatchingStateTimeline&)> Callback) const;
@@ -43,25 +40,25 @@ private:
 	template <class MessageType, class TimelineType = TraceServices::TPointTimeline<MessageType>>
 	struct TTimelineStorage
 	{
-		// Map the timelines as AnimInstanceId -> NodeId -> Timeline
-		using FNodeToTimelineMap = TMap<int32, uint32>;
-		using FAnimInstanceToTimelineMap = TMap<uint64, FNodeToTimelineMap>;
+		// Map the timelines as AnimInstanceId -> SearchId -> Timeline
+		using FSearchIdToTimelineMap = TMap<int32, uint32>;
+		using FAnimInstanceToTimelineMap = TMap<uint64, FSearchIdToTimelineMap>;
 
 		/** Retrieves the timeline for internal use, creating it if it does not exist */
-		TSharedRef<TimelineType> GetTimeline(TraceServices::IAnalysisSession& InSession, uint64 InAnimInstanceId, int32 InNodeId)
+		TSharedRef<TimelineType> GetTimeline(TraceServices::IAnalysisSession& InSession, uint64 InAnimInstanceId, int32 InSearchId)
 		{
 			const uint32* TimelineIndex = nullptr;
 			
-			FNodeToTimelineMap* NodeToTimelineMap = AnimInstanceIdToTimelines.Find(InAnimInstanceId);
-			if (NodeToTimelineMap == nullptr)
+			FSearchIdToTimelineMap* SearchIdToTimelineMap = AnimInstanceIdToTimelines.Find(InAnimInstanceId);
+			if (SearchIdToTimelineMap == nullptr)
 			{
 				// Create the timeline map if this is the first time accessing with this anim instance
-				NodeToTimelineMap = &AnimInstanceIdToTimelines.Add(InAnimInstanceId, {});
+				SearchIdToTimelineMap = &AnimInstanceIdToTimelines.Add(InAnimInstanceId, {});
 			}
 			else
 			{
-				// Anim instance already used, attempt to find the NodeId's timeline
-				TimelineIndex = NodeToTimelineMap->Find(InNodeId);
+				// Anim instance already used, attempt to find the SearchId's timeline
+				TimelineIndex = SearchIdToTimelineMap->Find(InSearchId);
 			}
 
 			if (TimelineIndex == nullptr)
@@ -69,16 +66,16 @@ private:
 				// Append our timeline to the storage and our object + the storage index to the map
 				TSharedRef<TimelineType> Timeline = MakeShared<TimelineType>(InSession.GetLinearAllocator());
 				const uint32 NewIndex = Timelines.Add(Timeline);
-				NodeToTimelineMap->Add(InNodeId, NewIndex);
+				SearchIdToTimelineMap->Add(InSearchId, NewIndex);
 				return Timeline;
 			}
 
 			return Timelines[*TimelineIndex];
 		}
 
-		bool ReadNodeTimeline(const FNodeToTimelineMap* NodeToTimelineMap, int32 InNodeId, TFunctionRef<void(const TraceServices::ITimeline<MessageType>&)> Callback) const
+		bool ReadSearchTimeline(const FSearchIdToTimelineMap* SearchIdToTimelineMap, int32 InSearchId, TFunctionRef<void(const TraceServices::ITimeline<MessageType>&)> Callback) const
 		{
-			const uint32* TimelineIndex = NodeToTimelineMap->Find(InNodeId);
+			const uint32* TimelineIndex = SearchIdToTimelineMap->Find(InSearchId);
 			if (TimelineIndex == nullptr || !Timelines.IsValidIndex(*TimelineIndex))
 			{
 				return false;
@@ -89,30 +86,30 @@ private:
 		}
 
 		/** Retrieve a timeline from an anim instance + node and execute the callback */
-		bool ReadTimeline(uint64 InAnimInstanceId, int32 InNodeId, TFunctionRef<void(const TraceServices::ITimeline<MessageType>&)> Callback) const
+		bool ReadTimeline(uint64 InAnimInstanceId, int32 InSearchId, TFunctionRef<void(const TraceServices::ITimeline<MessageType>&)> Callback) const
 		{
-			const FNodeToTimelineMap* NodeToTimelineMap = AnimInstanceIdToTimelines.Find(InAnimInstanceId);
-			if (NodeToTimelineMap == nullptr)
+			const FSearchIdToTimelineMap* SearchIdToTimelineMap = AnimInstanceIdToTimelines.Find(InAnimInstanceId);
+			if (SearchIdToTimelineMap == nullptr)
 			{
 				return false;
 			}
 
-			return ReadNodeTimeline(NodeToTimelineMap, InNodeId, Callback);
+			return ReadSearchTimeline(SearchIdToTimelineMap, InSearchId, Callback);
 		}
 
 		/** Enumerates timelines of all nodes given an AnimInstance */
-		bool EnumerateNodeTimelines(uint64 InAnimInstanceId, TFunctionRef<void(const TraceServices::ITimeline<MessageType>&)> Callback) const
+		bool EnumerateSearchTimelines(uint64 InAnimInstanceId, TFunctionRef<void(const TraceServices::ITimeline<MessageType>&)> Callback) const
 		{
-			const FNodeToTimelineMap* NodeToTimelineMap = AnimInstanceIdToTimelines.Find(InAnimInstanceId);
-			if (NodeToTimelineMap == nullptr)
+			const FSearchIdToTimelineMap* SearchIdToTimelineMap = AnimInstanceIdToTimelines.Find(InAnimInstanceId);
+			if (SearchIdToTimelineMap == nullptr)
 			{
 				return false;
 			}
 
 			bool bSuccess = true;
-			for (const TPair<int32, uint32>& Pair : *NodeToTimelineMap)
+			for (const TPair<int32, uint32>& Pair : *SearchIdToTimelineMap)
 			{
-				bSuccess &= ReadNodeTimeline(NodeToTimelineMap, Pair.Key, Callback);
+				bSuccess &= ReadSearchTimeline(SearchIdToTimelineMap, Pair.Key, Callback);
 				if (!bSuccess)
 				{
 					break;
@@ -122,25 +119,25 @@ private:
 			return bSuccess;
 		}
 
-		TSet<int32> GetNodeIds(uint64 InAnimInstanceId) const
+		TSet<int32> GetSearchIds(uint64 InAnimInstanceId) const
 		{
-			TSet<int32> NodeIds;
+			TSet<int32> SearchIds;
 			if (!AnimInstanceIdToTimelines.Find(InAnimInstanceId))
 			{
-				return NodeIds;
+				return SearchIds;
 			}
-			const FNodeToTimelineMap& NodeToTimelineMap = AnimInstanceIdToTimelines[InAnimInstanceId];
+			const FSearchIdToTimelineMap& SearchIdToTimelineMap = AnimInstanceIdToTimelines[InAnimInstanceId];
 
 			// For each node that has a timeline, add the node to the set
-			for(const TTuple<int32, uint32>& Pair : NodeToTimelineMap)
+			for(const TTuple<int32, uint32>& Pair : SearchIdToTimelineMap)
 			{
-				NodeIds.Add(Pair.Key);
+				SearchIds.Add(Pair.Key);
 			}
 			
-			return NodeIds;
+			return SearchIds;
 		}
 
-		/** Maps AnimInstanceIds to a map of NodeIds to timelines. AnimInstanceId -> NodeId -> Timeline */
+		/** Maps AnimInstanceIds to a map of SearchIds to timelines. AnimInstanceId -> SearchId -> Timeline */
 		FAnimInstanceToTimelineMap AnimInstanceIdToTimelines;
 
 		/** Timelines per node */
diff --git a/Engine/Plugins/Animation/PoseSearch/Source/Editor/Public/PoseSearchDatabaseEditorReflection.h b/Engine/Plugins/Animation/PoseSearch/Source/Editor/Public/PoseSearchDatabaseEditorReflection.h
--- a/Engine/Plugins/Animation/PoseSearch/Source/Editor/Public/PoseSearchDatabaseEditorReflection.h
+++ b/Engine/Plugins/Animation/PoseSearch/Source/Editor/Public/PoseSearchDatabaseEditorReflection.h
@@ -25,10 +25,7 @@ public:
 		const TSharedPtr<UE::PoseSearch::SDatabaseAssetTree>& InAssetTreeWidget);
 	
 	virtual bool ApplyChanges() { return false; };
-
-#if WITH_EDITOR
 	virtual void PostEditChangeProperty(FPropertyChangedEvent& PropertyChangedEvent) override;
-#endif
 	
 protected:
 	TWeakPtr<UE::PoseSearch::FDatabaseAssetTreeNode> WeakAssetTreeNode;
diff --git a/Engine/Plugins/Animation/PoseSearch/Source/Runtime/Private/PoseSearchContext.cpp b/Engine/Plugins/Animation/PoseSearch/Source/Runtime/Private/PoseSearchContext.cpp
--- a/Engine/Plugins/Animation/PoseSearch/Source/Runtime/Private/PoseSearchContext.cpp
+++ b/Engine/Plugins/Animation/PoseSearch/Source/Runtime/Private/PoseSearchContext.cpp
@@ -142,37 +142,40 @@ FQuat FDebugDrawParams::ExtractRotation(TConstArrayView<float> PoseVector, float
 	if (const UPoseSearchSchema* Schema = GetSchema())
 	{
 		int32 HeadingAxisFoundNum = 0;
-		const UPoseSearchFeatureChannel_Heading* FoundHeading[int32(EHeadingAxis::Num)];
-		FVector DecodedHeading[int32(EHeadingAxis::Num)];
+		const UPoseSearchFeatureChannel_Heading* FoundHeading[int32(EHeadingAxis::Num)] = { nullptr };
+		FVector DecodedHeading[int32(EHeadingAxis::Num)] = { FVector::ZeroVector };
 		for (int32 HeadingAxis = 0; HeadingAxis < int32(EHeadingAxis::Num); ++HeadingAxis)
 		{
 			// looking for a UPoseSearchFeatureChannel_Heading that matches the SampleTimeOffset, SchemaBoneIdx, and with OriginTimeOffset as zero.
 			// the features data associated to this channel would be a heading vector in GetRootTransform space (since OriginTimeOffset is zero)), 
 			// so by finding at least two with differnt axis we'll be able to compose a delta rotation from OriginTimeOffset (zero) to SampleTimeOffset
-			FoundHeading[HeadingAxis] = static_cast<const UPoseSearchFeatureChannel_Heading*>(
-				Schema->FindChannel([SampleTimeOffset, SchemaBoneIdx, &Role, PermutationTimeType, SamplingAttributeId, HeadingAxis](const UPoseSearchFeatureChannel* Channel) -> const UPoseSearchFeatureChannel_Heading*
+			Schema->FindChannel([SampleTimeOffset, SchemaBoneIdx, &Role, PermutationTimeType, SamplingAttributeId, HeadingAxis, PoseVector, &FoundHeading, &DecodedHeading](const UPoseSearchFeatureChannel* Channel) -> const UPoseSearchFeatureChannel_Heading*
+				{
+					const UPoseSearchFeatureChannel_Heading* Heading = Cast<UPoseSearchFeatureChannel_Heading>(Channel);
+					if (Heading &&
+						Heading->SchemaBoneIdx == SchemaBoneIdx &&
+						Heading->SampleTimeOffset == SampleTimeOffset &&
+						Heading->OriginTimeOffset == 0.f &&
+						Heading->PermutationTimeType == PermutationTimeType &&
+						Heading->SamplingAttributeId == SamplingAttributeId &&
+						Heading->SchemaOriginBoneIdx == RootSchemaBoneIdx &&
+						Heading->SampleRole == Role &&
+						Heading->OriginRole == Role &&
+						int32(Heading->HeadingAxis) == HeadingAxis)
 					{
-						if (const UPoseSearchFeatureChannel_Heading* Heading = Cast<UPoseSearchFeatureChannel_Heading>(Channel))
+						FVector DecodedHeadingValue = FFeatureVectorHelper::DecodeVector(PoseVector, Heading->GetChannelDataOffset(), Heading->ComponentStripping);
+						if (DecodedHeadingValue.Normalize())
 						{
-							if (Heading->SchemaBoneIdx == SchemaBoneIdx &&
-								Heading->SampleTimeOffset == SampleTimeOffset &&
-								Heading->OriginTimeOffset == 0.f &&
-								Heading->PermutationTimeType == PermutationTimeType &&
-								Heading->SamplingAttributeId == SamplingAttributeId &&
-								Heading->SchemaOriginBoneIdx == RootSchemaBoneIdx &&
-								Heading->SampleRole == Role &&
-								Heading->OriginRole == Role &&
-								int32(Heading->HeadingAxis) == HeadingAxis)
-							{
-								return Heading;
-							}
+							FoundHeading[HeadingAxis] = Heading;
+							DecodedHeading[HeadingAxis] = DecodedHeadingValue;
+							return Heading;
 						}
-						return nullptr;
-					}));
+					}
+					return nullptr;
+				});
+
 			if (FoundHeading[HeadingAxis])
 			{
-				DecodedHeading[HeadingAxis] = FFeatureVectorHelper::DecodeVector(PoseVector, FoundHeading[HeadingAxis]->GetChannelDataOffset(), FoundHeading[HeadingAxis]->ComponentStripping);
-
 				++HeadingAxisFoundNum;
 				if (HeadingAxisFoundNum == 2)
 				{
@@ -191,14 +194,17 @@ FQuat FDebugDrawParams::ExtractRotation(TConstArrayView<float> PoseVector, float
 				if (!FoundHeading[int32(EHeadingAxis::X)])
 				{
 					DecodedHeading[int32(EHeadingAxis::X)] = FVector::CrossProduct(DecodedHeading[int32(EHeadingAxis::Y)], DecodedHeading[int32(EHeadingAxis::Z)]);
+					bAbleToReconstructMissingAxis &= DecodedHeading[int32(EHeadingAxis::X)].Normalize();
 				}
 				else if (!FoundHeading[int32(EHeadingAxis::Y)])
 				{
 					DecodedHeading[int32(EHeadingAxis::Y)] = FVector::CrossProduct(DecodedHeading[int32(EHeadingAxis::Z)], DecodedHeading[int32(EHeadingAxis::X)]);
+					bAbleToReconstructMissingAxis &= DecodedHeading[int32(EHeadingAxis::Y)].Normalize();
 				}
 				else // if (!FoundHeading[int32(EHeadingAxis::Z)])
 				{
 					DecodedHeading[int32(EHeadingAxis::Z)] = FVector::CrossProduct(DecodedHeading[int32(EHeadingAxis::X)], DecodedHeading[int32(EHeadingAxis::Y)]);
+					bAbleToReconstructMissingAxis &= DecodedHeading[int32(EHeadingAxis::Z)].Normalize();
 				}
 			}
 			else 
@@ -211,18 +217,21 @@ FQuat FDebugDrawParams::ExtractRotation(TConstArrayView<float> PoseVector, float
 					DecodedHeading[int32(EHeadingAxis::Y)] = FVector::CrossProduct(FVector::ZAxisVector, DecodedHeading[int32(EHeadingAxis::X)]);
 					bAbleToReconstructMissingAxis &= DecodedHeading[int32(EHeadingAxis::Y)].Normalize();
 					DecodedHeading[int32(EHeadingAxis::Z)] = FVector::CrossProduct(DecodedHeading[int32(EHeadingAxis::X)], DecodedHeading[int32(EHeadingAxis::Y)]);
+					bAbleToReconstructMissingAxis &= DecodedHeading[int32(EHeadingAxis::Z)].Normalize();
 				}
 				else if (FoundHeading[int32(EHeadingAxis::Y)])
 				{
 					DecodedHeading[int32(EHeadingAxis::X)] = FVector::CrossProduct(DecodedHeading[int32(EHeadingAxis::Y)], FVector::ZAxisVector);
 					bAbleToReconstructMissingAxis &= DecodedHeading[int32(EHeadingAxis::X)].Normalize();
 					DecodedHeading[int32(EHeadingAxis::Z)] = FVector::CrossProduct(DecodedHeading[int32(EHeadingAxis::X)], DecodedHeading[int32(EHeadingAxis::Y)]);
+					bAbleToReconstructMissingAxis &= DecodedHeading[int32(EHeadingAxis::Z)].Normalize();
 				}
 				else // if (FoundHeading[int32(EHeadingAxis::Z)])
 				{
 					DecodedHeading[int32(EHeadingAxis::X)] = FVector::CrossProduct(FVector::YAxisVector, DecodedHeading[int32(EHeadingAxis::Z)]);
 					bAbleToReconstructMissingAxis &= DecodedHeading[int32(EHeadingAxis::X)].Normalize();
 					DecodedHeading[int32(EHeadingAxis::Y)] = FVector::CrossProduct(DecodedHeading[int32(EHeadingAxis::Z)], DecodedHeading[int32(EHeadingAxis::X)]);
+					bAbleToReconstructMissingAxis &= DecodedHeading[int32(EHeadingAxis::Y)].Normalize();
 				}
 			}
 
diff --git a/Engine/Plugins/Animation/PoseSearch/Source/Runtime/Private/PoseSearchDatabase.cpp b/Engine/Plugins/Animation/PoseSearch/Source/Runtime/Private/PoseSearchDatabase.cpp
--- a/Engine/Plugins/Animation/PoseSearch/Source/Runtime/Private/PoseSearchDatabase.cpp
+++ b/Engine/Plugins/Animation/PoseSearch/Source/Runtime/Private/PoseSearchDatabase.cpp
@@ -760,12 +760,12 @@ int32 UPoseSearchDatabase::GetNumberOfPrincipalComponents() const
 
 bool UPoseSearchDatabase::GetSkipSearchIfPossible() const
 {
-#if WITH_EDITOR && ENABLE_ANIM_DEBUG
+#if WITH_EDITOR && ENABLE_ANIM_DEBUG && UE_POSE_SEARCH_TRACE_ENABLED
 	if (UE::PoseSearch::CVarMotionMatchCompareAgainstBruteForce.GetValueOnAnyThread())
 	{
 		return false;
 	}
-#endif // WITH_EDITOR && ENABLE_ANIM_DEBUG
+#endif // WITH_EDITOR && ENABLE_ANIM_DEBUG && UE_POSE_SEARCH_TRACE_ENABLED
 	return true;
 }
 
@@ -1076,50 +1076,44 @@ UE::PoseSearch::FSearchResult UPoseSearchDatabase::Search(UE::PoseSearch::FSearc
 	}
 #endif // WITH_EDITOR
 
-	if (PoseSearchMode == EPoseSearchMode::BruteForce
-#if WITH_EDITOR && ENABLE_ANIM_DEBUG
-		|| CVarMotionMatchCompareAgainstBruteForce.GetValueOnAnyThread()
-#endif // WITH_EDITOR && ENABLE_ANIM_DEBUG
-		)
-	{
-		Result = SearchBruteForce(SearchContext);
-	}
-
-#if WITH_EDITOR && ENABLE_ANIM_DEBUG
-	const FPoseSearchCost BruteForcePoseCost = Result.BruteForcePoseCost;
-	const int32 BruteForcePoseIdx = Result.PoseIdx;
-#endif // WITH_EDITOR && ENABLE_ANIM_DEBUG
-
-	if (PoseSearchMode == EPoseSearchMode::VPTree)
+	switch (PoseSearchMode)
 	{
-		Result = SearchVPTree(SearchContext);
-
-#if WITH_EDITOR && ENABLE_ANIM_DEBUG
-		Result.BruteForcePoseCost = BruteForcePoseCost;
-		if (Result.PoseIdx != BruteForcePoseIdx && CVarMotionMatchCompareAgainstBruteForce.GetValueOnAnyThread())
+	case EPoseSearchMode::BruteForce:
 		{
-			UE_LOG(LogPoseSearch, Error, TEXT("UPoseSearchDatabase::Search - VPTree search PoseIdx %d differs from BruteForce search PoseIdx %d"), Result.PoseIdx, BruteForcePoseIdx);
+			Result = SearchBruteForce(SearchContext);
+
+#if WITH_EDITOR && ENABLE_ANIM_DEBUG && UE_POSE_SEARCH_TRACE_ENABLED
+			if (CVarMotionMatchCompareAgainstBruteForce.GetValueOnAnyThread())
+			{
+				Result.BruteForcePoseCost = Result.PoseCost;
+			}
+#endif // WITH_EDITOR && ENABLE_ANIM_DEBUG && UE_POSE_SEARCH_TRACE_ENABLED
+			break;
 		}
-#endif // UE_POSE_SEARCH_TRACE_ENABLED
-	}
-	else if (PoseSearchMode == EPoseSearchMode::PCAKDTree)
-	{
-		Result = SearchPCAKDTree(SearchContext);
+	case EPoseSearchMode::VPTree:
+		{
+			Result = SearchVPTree(SearchContext);
 
-#if WITH_EDITOR && ENABLE_ANIM_DEBUG
-		Result.BruteForcePoseCost = BruteForcePoseCost;
-		if (CVarMotionMatchCompareAgainstBruteForce.GetValueOnAnyThread())
+#if WITH_EDITOR && ENABLE_ANIM_DEBUG && UE_POSE_SEARCH_TRACE_ENABLED
+			if (CVarMotionMatchCompareAgainstBruteForce.GetValueOnAnyThread())
+			{
+				Result.BruteForcePoseCost = SearchBruteForce(SearchContext).PoseCost;
+			}
+#endif // WITH_EDITOR && ENABLE_ANIM_DEBUG && UE_POSE_SEARCH_TRACE_ENABLED
+			break;
+		}
+	case EPoseSearchMode::PCAKDTree:
 		{
-			const float BruteForceTotalCost = Result.BruteForcePoseCost.GetTotalCost();
-			const float PCAKDTreeTotalCost = Result.PoseCost.GetTotalCost();
-			check(BruteForceTotalCost <= PCAKDTreeTotalCost);
+			Result = SearchPCAKDTree(SearchContext);
 
-			if (!FMath::IsNearlyEqual(BruteForceTotalCost, PCAKDTreeTotalCost))
+#if WITH_EDITOR && ENABLE_ANIM_DEBUG && UE_POSE_SEARCH_TRACE_ENABLED
+			if (CVarMotionMatchCompareAgainstBruteForce.GetValueOnAnyThread())
 			{
-				UE_LOG(LogPoseSearch, Warning, TEXT("UPoseSearchDatabase::Search - PCAKDTree cost comparison %f (PCAKDTreeTotalCost %f, BruteForceTotalCost %f)"), PCAKDTreeTotalCost - BruteForceTotalCost, PCAKDTreeTotalCost, BruteForceTotalCost);
+				Result.BruteForcePoseCost = SearchBruteForce(SearchContext).PoseCost;
 			}
+#endif // WITH_EDITOR && ENABLE_ANIM_DEBUG && UE_POSE_SEARCH_TRACE_ENABLED
+			break;
 		}
-#endif // WITH_EDITOR && ENABLE_ANIM_DEBUG
 	}
 
 #if UE_POSE_SEARCH_TRACE_ENABLED
@@ -1150,12 +1144,12 @@ static inline void EvaluatePoseKernel(UE::PoseSearch::FSearchResult& Result, con
 			Result.PoseCost = PoseCost;
 			Result.PoseIdx = PoseIdx;
 
-#if UE_POSE_SEARCH_TRACE_ENABLED
+#if WITH_EDITOR && ENABLE_ANIM_DEBUG && UE_POSE_SEARCH_TRACE_ENABLED
 			if (bUpdateBestCandidates)
 			{
 				Result.BestPosePos = ResultIndex;
 			}
-#endif // UE_POSE_SEARCH_TRACE_ENABLED
+#endif // WITH_EDITOR && ENABLE_ANIM_DEBUG && UE_POSE_SEARCH_TRACE_ENABLED
 		}
 
 #if UE_POSE_SEARCH_TRACE_ENABLED
@@ -1252,6 +1246,13 @@ UE::PoseSearch::FSearchResult UPoseSearchDatabase::SearchContinuingPose(UE::Pose
 #endif // UE_POSE_SEARCH_TRACE_ENABLED
 	}
 
+#if WITH_EDITOR && ENABLE_ANIM_DEBUG && UE_POSE_SEARCH_TRACE_ENABLED
+	if (CVarMotionMatchCompareAgainstBruteForce.GetValueOnAnyThread())
+	{
+		Result.BruteForcePoseCost = Result.PoseCost;
+	}
+#endif // WITH_EDITOR && ENABLE_ANIM_DEBUG && UE_POSE_SEARCH_TRACE_ENABLED
+
 	return Result;
 }
 
@@ -1657,9 +1658,5 @@ UE::PoseSearch::FSearchResult UPoseSearchDatabase::SearchBruteForce(UE::PoseSear
 		Result.Database = this;
 	}
 
-#if UE_POSE_SEARCH_TRACE_ENABLED
-	Result.BruteForcePoseCost = Result.PoseCost; 
-#endif // UE_POSE_SEARCH_TRACE_ENABLED
-
 	return Result;
 }
\ No newline at end of file
diff --git a/Engine/Plugins/Animation/PoseSearch/Source/Runtime/Private/PoseSearchLibrary.cpp b/Engine/Plugins/Animation/PoseSearch/Source/Runtime/Private/PoseSearchLibrary.cpp
--- a/Engine/Plugins/Animation/PoseSearch/Source/Runtime/Private/PoseSearchLibrary.cpp
+++ b/Engine/Plugins/Animation/PoseSearch/Source/Runtime/Private/PoseSearchLibrary.cpp
@@ -275,18 +275,19 @@
 }
 
 #if UE_POSE_SEARCH_TRACE_ENABLED
-void UPoseSearchLibrary::TraceMotionMatchingState(
+void UPoseSearchLibrary::TraceMotionMatching(
 	UE::PoseSearch::FSearchContext& SearchContext,
 	const UE::PoseSearch::FSearchResult& CurrentResult,
 	float ElapsedPoseSearchTime,
 	const FTransform& RootMotionTransformDelta,
-	int32 NodeId,
 	float DeltaTime,
 	bool bSearch,
 	float RecordingTime)
 {
 	using namespace UE::PoseSearch;
 	
+	uint32 SearchId = 787;
+
 	FTraceMotionMatchingStateMessage TraceState;
 	
 	const int32 AnimInstancesNum = SearchContext.GetAnimInstances().Num();
@@ -300,6 +301,8 @@
 		TRACE_OBJECT(AnimInstance);
 
 		TraceState.SkeletalMeshComponentIds[AnimInstanceIndex] = FObjectTrace::GetObjectId(SkeletalMeshComponent);
+
+		SearchId = HashCombineFast(SearchId, GetTypeHash(FObjectTrace::GetObjectId(AnimInstance)));
 	}
 
 	TraceState.Roles.SetNum(AnimInstancesNum);
@@ -308,12 +311,16 @@
 		TraceState.Roles[RoleToIndexPair.Value] = RoleToIndexPair.Key;
 	}
 
+	SearchId = HashCombineFast(SearchId, GetTypeHash(TraceState.Roles));
+
+	// @todo: do we need to hash pose history names in SearchId as well?
 	TraceState.PoseHistories.SetNum(AnimInstancesNum);
 	for (int32 AnimInstanceIndex = 0; AnimInstanceIndex < AnimInstancesNum; ++AnimInstanceIndex)
 	{
 		TraceState.PoseHistories[AnimInstanceIndex].InitFrom(SearchContext.GetPoseHistories()[AnimInstanceIndex]);
 	}
 
+	TArray<uint64, TInlineAllocator<64>> DatabaseIds;
 	int32 DbEntryIdx = 0;
 	const int32 CurrentPoseIdx = bSearch && CurrentResult.PoseCost.IsValid() ? CurrentResult.PoseIdx : INDEX_NONE;
 	TraceState.DatabaseEntries.SetNum(SearchContext.GetBestPoseCandidatesMap().Num());
@@ -327,6 +334,7 @@
 		// if throttling is on, the continuing pose can be valid, but no actual search occurred, so the query will not be cached, and we need to build it
 		DbEntry.QueryVector = SearchContext.GetOrBuildQuery(Database->Schema);
 		DbEntry.DatabaseId = FTraceMotionMatchingStateMessage::GetIdFromObject(Database);
+		DatabaseIds.Add(DbEntry.DatabaseId);
 
 		for (int32 CandidateIdx = 0; CandidateIdx < DatabaseBestPoseCandidates.Value.Num(); ++CandidateIdx)
 		{
@@ -354,6 +362,9 @@
 		++DbEntryIdx;
 	}
 
+	DatabaseIds.Sort();
+	SearchId = HashCombineFast(SearchId, GetTypeHash(DatabaseIds));
+
 	if (DeltaTime > SMALL_NUMBER)
 	{
 		// simulation
@@ -387,12 +398,18 @@
 
 	TraceState.RecordingTime = RecordingTime;
 	TraceState.SearchBestCost = CurrentResult.PoseCost.GetTotalCost();
+#if WITH_EDITOR && ENABLE_ANIM_DEBUG
 	TraceState.SearchBruteForceCost = CurrentResult.BruteForcePoseCost.GetTotalCost();
 	TraceState.SearchBestPosePos = CurrentResult.BestPosePos;
+#else // WITH_EDITOR && ENABLE_ANIM_DEBUG
+	TraceState.SearchBruteForceCost = 0.f;
+	TraceState.SearchBestPosePos = 0;
+#endif // WITH_EDITOR && ENABLE_ANIM_DEBUG
 
 	TraceState.Cycle = FPlatformTime::Cycles64();
 	TraceState.AnimInstanceId = FObjectTrace::GetObjectId(SearchContext.GetAnimInstances()[0]);
-	TraceState.NodeId = NodeId;
 
+	TraceState.NodeId = SearchId;
+
 	TraceState.Output();
 }
@@ -477,24 +494,25 @@
 		{
 			if (ensure(Database))
 			{
-				FSearchResult NewSearchResult = Database->Search(SearchContext);
-				if (NewSearchResult.PoseCost.GetTotalCost() < SearchResult.PoseCost.GetTotalCost())
+				const FSearchResult NewSearchResult = Database->Search(SearchContext);
+
+#if WITH_EDITOR && ENABLE_ANIM_DEBUG && UE_POSE_SEARCH_TRACE_ENABLED
+				const FPoseSearchCost BestBruteForcePoseCost = NewSearchResult.BruteForcePoseCost < SearchResult.BruteForcePoseCost ? NewSearchResult.BruteForcePoseCost : SearchResult.BruteForcePoseCost;
+#endif // WITH_EDITOR && ENABLE_ANIM_DEBUG && UE_POSE_SEARCH_TRACE_ENABLED
+
+				if (NewSearchResult.PoseCost < SearchResult.PoseCost)
 				{
 					bJumpToPose = true;
 					SearchResult = NewSearchResult;
 					SearchContext.UpdateCurrentBestCost(SearchResult.PoseCost);
 				}
+
+#if WITH_EDITOR && ENABLE_ANIM_DEBUG && UE_POSE_SEARCH_TRACE_ENABLED
+				SearchResult.BruteForcePoseCost = BestBruteForcePoseCost;
+#endif // WITH_EDITOR && ENABLE_ANIM_DEBUG && UE_POSE_SEARCH_TRACE_ENABLED
 			}
 		}
 
-#if UE_POSE_SEARCH_TRACE_ENABLED
-		if (!SearchResult.BruteForcePoseCost.IsValid())
-		{
-			SearchResult.BruteForcePoseCost = SearchResult.PoseCost;
-		}
-#endif // UE_POSE_SEARCH_TRACE_ENABLED
-
-		
 #if WITH_EDITOR
 		// resetting CurrentSearchResult if any DDC indexing on the requested databases is still in progress
 		if (SearchContext.IsAsyncBuildIndexInProgress())
@@ -544,9 +562,9 @@
 		else
 		{
 			// copying few properties of SearchResult into CurrentSearchResult to facilitate debug drawing
-#if UE_POSE_SEARCH_TRACE_ENABLED
+#if WITH_EDITOR && ENABLE_ANIM_DEBUG && UE_POSE_SEARCH_TRACE_ENABLED
 			InOutMotionMatchingState.CurrentSearchResult.BruteForcePoseCost = SearchResult.BruteForcePoseCost;
-#endif // UE_POSE_SEARCH_TRACE_ENABLED
+#endif // WITH_EDITOR && ENABLE_ANIM_DEBUG && UE_POSE_SEARCH_TRACE_ENABLED
 			InOutMotionMatchingState.CurrentSearchResult.PoseCost = SearchResult.PoseCost;
 		}
 	}
@@ -564,8 +582,8 @@
 	// Record debugger details
 	if (IsTracing(Context))
 	{
-		TraceMotionMatchingState(SearchContext, InOutMotionMatchingState.CurrentSearchResult, InOutMotionMatchingState.ElapsedPoseSearchTime,
-			InOutMotionMatchingState.RootMotionTransformDelta, Context.GetCurrentNodeId(), DeltaTime, bSearch,
+		TraceMotionMatching(SearchContext, InOutMotionMatchingState.CurrentSearchResult, InOutMotionMatchingState.ElapsedPoseSearchTime,
+			InOutMotionMatchingState.RootMotionTransformDelta, DeltaTime, bSearch,
 			AnimInstance ? FObjectTrace::GetWorldElapsedTime(AnimInstance->GetWorld()) : 0.f);
 	}
 #endif // UE_POSE_SEARCH_TRACE_ENABLED
@@ -605,8 +623,7 @@
 	const FName PoseHistoryName,
 	const FPoseSearchContinuingProperties ContinuingProperties,
 	const FPoseSearchFutureProperties Future,
-	FPoseSearchBlueprintResult& Result,
-	const int32 DebugSessionUniqueIdentifier)
+	FPoseSearchBlueprintResult& Result)
 {
 	using namespace UE::PoseSearch;
 
@@ -619,7 +636,7 @@
 	Roles.Add(UE::PoseSearch::DefaultRole);
 
 	TArray<const UObject*>& AssetsToSearchConst = reinterpret_cast<TArray<const UObject*>&>(AssetsToSearch);
-	MotionMatch(AnimInstances, Roles, AssetsToSearchConst, PoseHistoryName, ContinuingProperties, Future, Result, DebugSessionUniqueIdentifier);
+	MotionMatch(AnimInstances, Roles, AssetsToSearchConst, PoseHistoryName, ContinuingProperties, Future, Result);
 }
 
 void UPoseSearchLibrary::MotionMatchMulti(
@@ -629,8 +646,7 @@
 	const FName PoseHistoryName,
 	const FPoseSearchContinuingProperties ContinuingProperties,
 	const FPoseSearchFutureProperties Future,
-	FPoseSearchBlueprintResult& Result,
-	const int32 DebugSessionUniqueIdentifier)
+	FPoseSearchBlueprintResult& Result)
 {
 	using namespace UE::PoseSearch;
 
@@ -648,7 +664,7 @@
 	}
 
 	TArray<const UObject*>& AssetsToSearchConst = reinterpret_cast<TArray<const UObject*>&>(AssetsToSearch);
-	MotionMatch(AnimInstances, Roles, AssetsToSearchConst, PoseHistoryName, ContinuingProperties, Future, Result, DebugSessionUniqueIdentifier);
+	MotionMatch(AnimInstances, Roles, AssetsToSearchConst, PoseHistoryName, ContinuingProperties, Future, Result);
 }
 
 void UPoseSearchLibrary::MotionMatch(
@@ -658,8 +674,7 @@
 	const FName PoseHistoryName,
 	const FPoseSearchContinuingProperties& ContinuingProperties,
 	const FPoseSearchFutureProperties& Future,
-	FPoseSearchBlueprintResult& Result,
-	const int32 DebugSessionUniqueIdentifier)
+	FPoseSearchBlueprintResult& Result)
 {
 	using namespace UE::Anim;
 	using namespace UE::PoseSearch;
@@ -710,7 +725,7 @@
 		return;
 	}
 
-	const FSearchResult SearchResult = MotionMatch(AnimInstances, Roles, PoseHistories, AssetsToSearch, ContinuingProperties, Future, DebugSessionUniqueIdentifier);
+	const FSearchResult SearchResult = MotionMatch(AnimInstances, Roles, PoseHistories, AssetsToSearch, ContinuingProperties, Future);
 	if (SearchResult.IsValid())
 	{
 		const UPoseSearchDatabase* Database = SearchResult.Database.Get();
@@ -750,8 +765,7 @@
 	const TArrayView<const UE::PoseSearch::IPoseHistory*> PoseHistories, 
 	const TArrayView<const UObject*> AssetsToSearch,
 	const FPoseSearchContinuingProperties& ContinuingProperties,
-	const FPoseSearchFutureProperties& Future,
-	const int32 DebugSessionUniqueIdentifier)
+	const FPoseSearchFutureProperties& Future)
 {
 	check(!AnimInstances.IsEmpty() && AnimInstances.Num() == Roles.Num() && AnimInstances.Num() == PoseHistories.Num());
 
@@ -902,11 +916,20 @@
 								SearchContext.UpdateCurrentResultPoseVector();
 
 								const FSearchResult NewSearchResult = Database->SearchContinuingPose(SearchContext);
-								if (NewSearchResult.PoseCost.GetTotalCost() < SearchResult.PoseCost.GetTotalCost())
+
+#if WITH_EDITOR && ENABLE_ANIM_DEBUG && UE_POSE_SEARCH_TRACE_ENABLED
+								const FPoseSearchCost BestBruteForcePoseCost = NewSearchResult.BruteForcePoseCost < SearchResult.BruteForcePoseCost ? NewSearchResult.BruteForcePoseCost : SearchResult.BruteForcePoseCost;
+#endif // WITH_EDITOR && ENABLE_ANIM_DEBUG && UE_POSE_SEARCH_TRACE_ENABLED
+
+								if (NewSearchResult.PoseCost < SearchResult.PoseCost)
 								{
 									SearchResult = NewSearchResult;
 									SearchContext.UpdateCurrentBestCost(SearchResult.PoseCost);
 								}
+								
+#if WITH_EDITOR && ENABLE_ANIM_DEBUG && UE_POSE_SEARCH_TRACE_ENABLED
+								SearchResult.BruteForcePoseCost = BestBruteForcePoseCost;
+#endif // WITH_EDITOR && ENABLE_ANIM_DEBUG && UE_POSE_SEARCH_TRACE_ENABLED
 							}
 						}
 					}
@@ -933,11 +956,20 @@
 			SearchContext.SetAssetsToConsider(AssetsToSearchPerDatabasePair.Value);
 
 			const FSearchResult NewSearchResult = Database->Search(SearchContext);
-			if (NewSearchResult.PoseCost.GetTotalCost() < SearchResult.PoseCost.GetTotalCost())
+
+#if WITH_EDITOR && ENABLE_ANIM_DEBUG && UE_POSE_SEARCH_TRACE_ENABLED
+			const FPoseSearchCost BestBruteForcePoseCost = NewSearchResult.BruteForcePoseCost < SearchResult.BruteForcePoseCost ? NewSearchResult.BruteForcePoseCost : SearchResult.BruteForcePoseCost;
+#endif // WITH_EDITOR && ENABLE_ANIM_DEBUG && UE_POSE_SEARCH_TRACE_ENABLED
+
+			if (NewSearchResult.PoseCost < SearchResult.PoseCost)
 			{
 				SearchResult = NewSearchResult;
 				SearchContext.UpdateCurrentBestCost(SearchResult.PoseCost);
 			}
+
+#if WITH_EDITOR && ENABLE_ANIM_DEBUG && UE_POSE_SEARCH_TRACE_ENABLED
+			SearchResult.BruteForcePoseCost = BestBruteForcePoseCost;
+#endif // WITH_EDITOR && ENABLE_ANIM_DEBUG && UE_POSE_SEARCH_TRACE_ENABLED
 		}
 	}
 
@@ -973,48 +1005,13 @@
 #endif // ENABLE_DRAW_DEBUG && ENABLE_ANIM_DEBUG
 
 #if UE_POSE_SEARCH_TRACE_ENABLED
-	const float SearchBestCost = SearchResult.PoseCost.GetTotalCost();
-	const float SearchBruteForceCost = SearchResult.BruteForcePoseCost.GetTotalCost();
-	TraceMotionMatchingState(SearchContext, SearchResult, 0.f, FTransform::Identity, DebugSessionUniqueIdentifier,
+	TraceMotionMatching(SearchContext, SearchResult, 0.f, FTransform::Identity,
 		DeltaSeconds, true, FObjectTrace::GetWorldElapsedTime(AnimInstances[0]->GetWorld()));
 #endif // UE_POSE_SEARCH_TRACE_ENABLED
 
 	return SearchResult;
 }
 
-// deprecated signatures
-UE::PoseSearch::FSearchResult UPoseSearchLibrary::MotionMatch(
-	const FAnimationBaseContext& Context,
-	TArrayView<const UObject*> AssetsToSearch,
-	const FPoseSearchContinuingProperties& ContinuingProperties)
-{
-	using namespace UE::PoseSearch;
-
-	const IPoseHistory* PoseHistory = nullptr;
-	if (FPoseHistoryProvider* PoseHistoryProvider = Context.GetMessage<FPoseHistoryProvider>())
-	{
-		PoseHistory = &PoseHistoryProvider->GetPoseHistory();
-	}
-
-	UAnimInstance* AnimInstance = Cast<UAnimInstance>(Context.AnimInstanceProxy->GetAnimInstanceObject());
-	check(AnimInstance);
-
-	return MotionMatch(MakeArrayView(&AnimInstance, 1), MakeArrayView(&DefaultRole, 1), MakeArrayView(&PoseHistory, 1), 
-		AssetsToSearch, ContinuingProperties, FPoseSearchFutureProperties(), Context.GetCurrentNodeId());
-}
-		
-UE::PoseSearch::FSearchResult UPoseSearchLibrary::MotionMatch(
-	TArrayView<UAnimInstance*> AnimInstances,
-	TArrayView<const UE::PoseSearch::FRole> Roles,
-	TArrayView<const UE::PoseSearch::IPoseHistory*> PoseHistories,
-	TArrayView<const UObject*> AssetsToSearch,
-	const FPoseSearchContinuingProperties& ContinuingProperties,
-	const int32 DebugSessionUniqueIdentifier,
-	float DesiredPermutationTimeOffset)
-{
-	return MotionMatch(AnimInstances, Roles, PoseHistories, AssetsToSearch, ContinuingProperties, FPoseSearchFutureProperties(), DebugSessionUniqueIdentifier);
-}
-	
 const FAnimNode_PoseSearchHistoryCollector_Base* UPoseSearchLibrary::FindPoseHistoryNode(
 	const FName PoseHistoryName,
 	const UAnimInstance* AnimInstance)
@@ -1058,4 +1055,66 @@
 	return nullptr;
 }
 
+///////////////////////////////////////////////////////////////////////////////////////////
+// Begin deprecated signatures
+void UPoseSearchLibrary::MotionMatch(
+	const TArrayView<UAnimInstance*> AnimInstances,
+	const TArrayView<const UE::PoseSearch::FRole> Roles,
+	const TArrayView<const UObject*> AssetsToSearch,
+	const FName PoseHistoryName,
+	const FPoseSearchContinuingProperties& ContinuingProperties,
+	const FPoseSearchFutureProperties& Future,
+	FPoseSearchBlueprintResult& Result,
+	const int32 DebugSessionUniqueIdentifier)
+{
+	MotionMatch(AnimInstances, Roles, AssetsToSearch, PoseHistoryName, ContinuingProperties, Future, Result);
+}
+
+UE::PoseSearch::FSearchResult UPoseSearchLibrary::MotionMatch(
+	const TArrayView<UAnimInstance*> AnimInstances,
+	const TArrayView<const UE::PoseSearch::FRole> Roles,
+	const TArrayView<const UE::PoseSearch::IPoseHistory*> PoseHistories,
+	const TArrayView<const UObject*> AssetsToSearch,
+	const FPoseSearchContinuingProperties& ContinuingProperties,
+	const FPoseSearchFutureProperties& Future,
+	const int32 DebugSessionUniqueIdentifier)
+{
+	return MotionMatch(AnimInstances, Roles, PoseHistories, AssetsToSearch, ContinuingProperties, Future);
+}
+
+UE::PoseSearch::FSearchResult UPoseSearchLibrary::MotionMatch(
+	const FAnimationBaseContext& Context,
+	TArrayView<const UObject*> AssetsToSearch,
+	const FPoseSearchContinuingProperties& ContinuingProperties)
+{
+	using namespace UE::PoseSearch;
+
+	const IPoseHistory* PoseHistory = nullptr;
+	if (FPoseHistoryProvider* PoseHistoryProvider = Context.GetMessage<FPoseHistoryProvider>())
+	{
+		PoseHistory = &PoseHistoryProvider->GetPoseHistory();
+	}
+
+	UAnimInstance* AnimInstance = Cast<UAnimInstance>(Context.AnimInstanceProxy->GetAnimInstanceObject());
+	check(AnimInstance);
+
+	return MotionMatch(MakeArrayView(&AnimInstance, 1), MakeArrayView(&DefaultRole, 1), MakeArrayView(&PoseHistory, 1),
+		AssetsToSearch, ContinuingProperties, FPoseSearchFutureProperties());
+}
+		
+UE::PoseSearch::FSearchResult UPoseSearchLibrary::MotionMatch(
+	TArrayView<UAnimInstance*> AnimInstances,
+	TArrayView<const UE::PoseSearch::FRole> Roles,
+	TArrayView<const UE::PoseSearch::IPoseHistory*> PoseHistories,
+	TArrayView<const UObject*> AssetsToSearch,
+	const FPoseSearchContinuingProperties& ContinuingProperties,
+	const int32 DebugSessionUniqueIdentifier,
+	float DesiredPermutationTimeOffset)
+{
+	return MotionMatch(AnimInstances, Roles, PoseHistories, AssetsToSearch, ContinuingProperties, FPoseSearchFutureProperties());
+}
+
+// End deprecated signatures
+///////////////////////////////////////////////////////////////////////////////////////////
+
 #undef LOCTEXT_NAMESPACE
diff --git a/Engine/Plugins/Animation/PoseSearch/Source/Runtime/Private/PoseSearchModule.cpp b/Engine/Plugins/Animation/PoseSearch/Source/Runtime/Private/PoseSearchModule.cpp
--- a/Engine/Plugins/Animation/PoseSearch/Source/Runtime/Private/PoseSearchModule.cpp
+++ b/Engine/Plugins/Animation/PoseSearch/Source/Runtime/Private/PoseSearchModule.cpp
@@ -49,8 +49,8 @@ public:
 		UAnimInstance* AnimInstance = Cast<UAnimInstance>(GraphContext.AnimInstanceProxy->GetAnimInstanceObject());
 		check(AnimInstance);
 		
-		const UE::PoseSearch::FSearchResult SearchResult = UPoseSearchLibrary::MotionMatch(MakeArrayView(&AnimInstance, 1), MakeArrayView(&DefaultRole, 1),
-			MakeArrayView(&PoseHistory, 1), AssetsToSearch, ContinuingProperties, Future, GraphContext.GetCurrentNodeId());
+		const UE::PoseSearch::FSearchResult SearchResult = UPoseSearchLibrary::MotionMatch(MakeArrayView(&AnimInstance, 1),
+			MakeArrayView(&DefaultRole, 1), MakeArrayView(&PoseHistory, 1), AssetsToSearch, ContinuingProperties, Future);
 
 		UE::Anim::IPoseSearchProvider::FSearchResult ProviderResult;
 		if (const UE::PoseSearch::FSearchIndexAsset* SearchIndexAsset = SearchResult.GetSearchIndexAsset())
diff --git a/Engine/Plugins/Animation/PoseSearch/Source/Runtime/Private/Trace/PoseSearchTraceLogger.cpp b/Engine/Plugins/Animation/PoseSearch/Source/Runtime/Private/Trace/PoseSearchTraceLogger.cpp
--- a/Engine/Plugins/Animation/PoseSearch/Source/Runtime/Private/Trace/PoseSearchTraceLogger.cpp
+++ b/Engine/Plugins/Animation/PoseSearch/Source/Runtime/Private/Trace/PoseSearchTraceLogger.cpp
@@ -132,4 +132,76 @@ const FTraceMotionMatchingStatePoseEntry* FTraceMotionMatchingStateMessage::GetC
 	return nullptr;
 }
 
+FText FTraceMotionMatchingStateMessage::GenerateSearchName() const
+{
+	TStringBuilder<256> StringBuilder;
+	bool bAddDatabasesSeparator = false;
+	for (const FTraceMotionMatchingStateDatabaseEntry& DbEntry : DatabaseEntries)
+	{
+		if (bAddDatabasesSeparator)
+		{
+			StringBuilder.Append(" - ");
+		}
+
+		if (const UPoseSearchDatabase* Database = GetObjectFromId<UPoseSearchDatabase>(DbEntry.DatabaseId))
+		{
+			StringBuilder.Append(Database->GetName());
+		}
+		else
+		{
+			StringBuilder.Append("Unknown");
+		}
+		
+		bAddDatabasesSeparator = true;
+	}
+
+	if (Roles.Num() > 1)
+	{
+		if (Roles.Num() != SkeletalMeshComponentIds.Num())
+		{
+			StringBuilder.Append("Error!");
+		}
+		else
+		{
+			StringBuilder.Append(" [");
+
+			bool bAddRolesSeparator = false;
+
+			for (int32 RoleIndex = 0; RoleIndex < Roles.Num(); ++RoleIndex)
+			{
+				if (bAddRolesSeparator)
+				{
+					StringBuilder.Append(" - ");
+				}
+
+				StringBuilder.Append(Roles[RoleIndex].ToString());
+
+				StringBuilder.Append(": ");
+
+				if (const USkeletalMeshComponent* SkeletalMeshComponent = GetObjectFromId<USkeletalMeshComponent>(SkeletalMeshComponentIds[RoleIndex]))
+				{
+					if (const AActor* Actor = SkeletalMeshComponent->GetOwner())
+					{
+						StringBuilder.Append(Actor->GetName());
+					}
+					else
+					{
+						StringBuilder.Append("Error!");
+					}
+				}
+				else
+				{
+					StringBuilder.Append("Error!");
+				}
+
+				bAddRolesSeparator = true;
+			}
+
+			StringBuilder.Append("]");
+		}
+	}
+
+	return FText::FromString(StringBuilder.ToString());
+}
+
 } // namespace UE::PoseSearch
\ No newline at end of file
diff --git a/Engine/Plugins/Animation/PoseSearch/Source/Runtime/Public/PoseSearch/PoseSearchCost.h b/Engine/Plugins/Animation/PoseSearch/Source/Runtime/Public/PoseSearch/PoseSearchCost.h
--- a/Engine/Plugins/Animation/PoseSearch/Source/Runtime/Public/PoseSearch/PoseSearchCost.h
+++ b/Engine/Plugins/Animation/PoseSearch/Source/Runtime/Public/PoseSearch/PoseSearchCost.h
@@ -18,7 +18,9 @@ public:
 #endif // WITH_EDITORONLY_DATA
 	}
 
-	bool IsValid() const { return TotalCost != MAX_flt; }
+	static bool IsCostValid(const float Cost) { return Cost != MAX_flt; }
+
+	bool IsValid() const { return IsCostValid(TotalCost); }
 	float GetTotalCost() const { return TotalCost; }
 	bool operator<(const FPoseSearchCost& Other) const { return TotalCost < Other.TotalCost; }
 
diff --git a/Engine/Plugins/Animation/PoseSearch/Source/Runtime/Public/PoseSearch/PoseSearchDefines.h b/Engine/Plugins/Animation/PoseSearch/Source/Runtime/Public/PoseSearch/PoseSearchDefines.h
--- a/Engine/Plugins/Animation/PoseSearch/Source/Runtime/Public/PoseSearch/PoseSearchDefines.h
+++ b/Engine/Plugins/Animation/PoseSearch/Source/Runtime/Public/PoseSearch/PoseSearchDefines.h
@@ -33,6 +33,7 @@ static constexpr int32 PreallocatedCachedQueriesNum = 8;
 static constexpr int32 PreallocatedCachedChannelDataNum = 64;
 static constexpr float FiniteDelta = 1 / 60.0f; // Time delta used for computing pose derivatives
 static constexpr int32 MaxNumberOfCollectedPoseCandidatesPerDatabase = 200;
+static constexpr int32 InvalidSearchId = INDEX_NONE; // @todo: it should be 0, since the search id is an hash!
 
 template<uint32 Alignment = DEFAULT_ALIGNMENT> class TMemStackSetAllocator : public TSetAllocator<TSparseArrayAllocator<TMemStackAllocator<Alignment>, TMemStackAllocator<Alignment>>, TMemStackAllocator<Alignment>> {};
 
diff --git a/Engine/Plugins/Animation/PoseSearch/Source/Runtime/Public/PoseSearch/PoseSearchLibrary.h b/Engine/Plugins/Animation/PoseSearch/Source/Runtime/Public/PoseSearch/PoseSearchLibrary.h
--- a/Engine/Plugins/Animation/PoseSearch/Source/Runtime/Public/PoseSearch/PoseSearchLibrary.h
+++ b/Engine/Plugins/Animation/PoseSearch/Source/Runtime/Public/PoseSearch/PoseSearchLibrary.h
@@ -127,6 +127,16 @@ class POSESEARCH_API UPoseSearchLibrary : public UBlueprintFunctionLibrary
 	GENERATED_BODY()
 
 #if UE_POSE_SEARCH_TRACE_ENABLED
+	static void TraceMotionMatching(
+		UE::PoseSearch::FSearchContext& SearchContext,
+		const UE::PoseSearch::FSearchResult& CurrentResult,
+		float ElapsedPoseSearchTime,
+		const FTransform& RootMotionTransformDelta,
+		float DeltaTime,
+		bool bSearch,
+		float RecordingTime);
+
+	UE_DEPRECATED(5.4, "Use TraceMotionMatching instead")
 	static void TraceMotionMatchingState(
 		UE::PoseSearch::FSearchContext& SearchContext,
 		const UE::PoseSearch::FSearchResult& CurrentResult,
@@ -135,7 +145,11 @@ class POSESEARCH_API UPoseSearchLibrary : public UBlueprintFunctionLibrary
 		int32 NodeId,
 		float DeltaTime,
 		bool bSearch,
-		float RecordingTime);
+		float RecordingTime)
+	{
+		TraceMotionMatching(SearchContext, CurrentResult, ElapsedPoseSearchTime, RootMotionTransformDelta, DeltaTime, bSearch, RecordingTime);
+	}
+	
 #endif // UE_POSE_SEARCH_TRACE_ENABLED
 
 public:
@@ -181,7 +195,6 @@ public:
 	* @param Future							Input future properties to match (animation / start time / time offset)
 	* @param SelectedAnimation				Output selected animation from the Database asset
 	* @param Result							Output FPoseSearchBlueprintResult with the search result
-	* @param DebugSessionUniqueIdentifier	Input unique identifier used to identify TraceMotionMatchingState (rewind debugger / pose search debugger) session. Similarly the MM node uses Context.GetCurrentNodeId()
 	*/
 	UFUNCTION(BlueprintPure, Category = "Animation|Pose Search|Experimental", meta = (BlueprintThreadSafe, Keywords = "PoseMatch"))
 	static void MotionMatch(
@@ -190,8 +203,7 @@ public:
 		const FName PoseHistoryName,
 		const FPoseSearchContinuingProperties ContinuingProperties,
 		const FPoseSearchFutureProperties Future,
-		FPoseSearchBlueprintResult& Result,
-		const int32 DebugSessionUniqueIdentifier = 6174);
+		FPoseSearchBlueprintResult& Result);
 
 	/**
 	* Implementation of the core motion matching algorithm for multiple characters
@@ -201,7 +213,6 @@ public:
 	* @param AssetsToSearch					Input assets to search (UPoseSearchDatabase or any animation asset containing UAnimNotifyState_PoseSearchBranchIn)
 	* @param PoseHistoryName				Input tag of the associated PoseSearchHistoryCollector node in the anim graphs of the AnimInstances
 	* @param Result							Output FPoseSearchBlueprintResult with the search result
-	* @param DebugSessionUniqueIdentifier	Input unique identifier used to identify TraceMotionMatchingState (rewind debugger / pose search debugger) session. Similarly the MM node uses Context.GetCurrentNodeId()
 	*/
 	UFUNCTION(BlueprintPure, Category = "Animation|Pose Search|Experimental", meta = (BlueprintThreadSafe, Keywords = "PoseMatch"))
 	static void MotionMatchMulti(
@@ -211,9 +222,26 @@ public:
 		const FName PoseHistoryName,
 		const FPoseSearchContinuingProperties ContinuingProperties,
 		const FPoseSearchFutureProperties Future,
-		FPoseSearchBlueprintResult& Result,
-		const int32 DebugSessionUniqueIdentifier = 6174);
+		FPoseSearchBlueprintResult& Result);
 
+	static void MotionMatch(
+		const TArrayView<UAnimInstance*> AnimInstances,
+		const TArrayView<const UE::PoseSearch::FRole> Roles,
+		const TArrayView<const UObject*> AssetsToSearch,
+		const FName PoseHistoryName,
+		const FPoseSearchContinuingProperties& ContinuingProperties,
+		const FPoseSearchFutureProperties& Future,
+		FPoseSearchBlueprintResult& Result);
+
+	static UE::PoseSearch::FSearchResult MotionMatch(
+		const TArrayView<UAnimInstance*> AnimInstances,
+		const TArrayView<const UE::PoseSearch::FRole> Roles,
+		const TArrayView<const UE::PoseSearch::IPoseHistory*> PoseHistories, 
+		const TArrayView<const UObject*> AssetsToSearch,
+		const FPoseSearchContinuingProperties& ContinuingProperties,
+		const FPoseSearchFutureProperties& Future);
+
+	UE_DEPRECATED(5.4, "Use other MotionMatch signatures instead")
 	static void MotionMatch(
 		const TArrayView<UAnimInstance*> AnimInstances,
 		const TArrayView<const UE::PoseSearch::FRole> Roles,
@@ -224,6 +252,7 @@ public:
 		FPoseSearchBlueprintResult& Result,
 		const int32 DebugSessionUniqueIdentifier);
 
+	UE_DEPRECATED(5.4, "Use other MotionMatch signatures instead")
 	static UE::PoseSearch::FSearchResult MotionMatch(
 		const TArrayView<UAnimInstance*> AnimInstances,
 		const TArrayView<const UE::PoseSearch::FRole> Roles,
diff --git a/Engine/Plugins/Animation/PoseSearch/Source/Runtime/Public/PoseSearch/PoseSearchResult.h b/Engine/Plugins/Animation/PoseSearch/Source/Runtime/Public/PoseSearch/PoseSearchResult.h
--- a/Engine/Plugins/Animation/PoseSearch/Source/Runtime/Public/PoseSearch/PoseSearchResult.h
+++ b/Engine/Plugins/Animation/PoseSearch/Source/Runtime/Public/PoseSearch/PoseSearchResult.h
@@ -22,10 +22,10 @@ struct FSearchResult
 	float AssetTime = 0.f;
 	bool bIsContinuingPoseSearch = false;
 
-#if UE_POSE_SEARCH_TRACE_ENABLED
+#if WITH_EDITOR && ENABLE_ANIM_DEBUG && UE_POSE_SEARCH_TRACE_ENABLED
 	FPoseSearchCost BruteForcePoseCost;
 	int32 BestPosePos = 0;
-#endif // UE_POSE_SEARCH_TRACE_ENABLED
+#endif // WITH_EDITOR && ENABLE_ANIM_DEBUG && UE_POSE_SEARCH_TRACE_ENABLED
 
 	// Attempts to set the internal state to match the provided asset time including updating the internal DbPoseIdx. 
 	// If the provided asset time is out of bounds for the currently playing asset then this function will reset the 
diff --git a/Engine/Plugins/Animation/PoseSearch/Source/Runtime/Public/PoseSearch/Trace/PoseSearchTraceLogger.h b/Engine/Plugins/Animation/PoseSearch/Source/Runtime/Public/PoseSearch/Trace/PoseSearchTraceLogger.h
--- a/Engine/Plugins/Animation/PoseSearch/Source/Runtime/Public/PoseSearch/Trace/PoseSearchTraceLogger.h
+++ b/Engine/Plugins/Animation/PoseSearch/Source/Runtime/Public/PoseSearch/Trace/PoseSearchTraceLogger.h
@@ -27,8 +27,10 @@ struct POSESEARCH_API FTraceMessage
 
 	uint64 AnimInstanceId = 0;
 
-	/** Node Id of the motion matching node associated with this message */
-	int32 NodeId = 0;
+	// motion matching Search Id associated with this message
+	// @todo: rename it to SearchId
+	int32 NodeId = InvalidSearchId;
+	int32 GetSearchId() const { return NodeId; }
 };
 POSESEARCH_API FArchive& operator<<(FArchive& Ar, FTraceMessage& State);
 
@@ -119,6 +121,8 @@ struct POSESEARCH_API FTraceMotionMatchingStateMessage : public FTraceMessage
 #endif
 	}
 	
+	FText GenerateSearchName() const;
+
 	static const FName Name;
 };
 
