motion matching - improved synchronization between UPoseSearchDatabase and UAnimNotifyState_PoseSearchBranchIn

diff --git a/Engine/Plugins/Animation/PoseSearch/Source/Editor/Private/PoseSearchDatabaseAssetListItem.cpp b/Engine/Plugins/Animation/PoseSearch/Source/Editor/Private/PoseSearchDatabaseAssetListItem.cpp
--- a/Engine/Plugins/Animation/PoseSearch/Source/Editor/Private/PoseSearchDatabaseAssetListItem.cpp
+++ b/Engine/Plugins/Animation/PoseSearch/Source/Editor/Private/PoseSearchDatabaseAssetListItem.cpp
@@ -352,7 +352,7 @@ namespace UE::PoseSearch
 					  {
 						 if (const FPoseSearchDatabaseAnimationAssetBase* DatabaseAnimationAssetBase = Database->GetAnimationAssetBase(AssetTreeNode->SourceAssetIdx))
 						 {
-							if (DatabaseAnimationAssetBase->bSynchronizeWithExternalDependency)
+							if (DatabaseAnimationAssetBase->IsSynchronizedWithExternalDependency())
 							{
 							   return EVisibility::Visible;
 							}
diff --git a/Engine/Plugins/Animation/PoseSearch/Source/Editor/Private/PoseSearchDatabaseAssetTree.cpp b/Engine/Plugins/Animation/PoseSearch/Source/Editor/Private/PoseSearchDatabaseAssetTree.cpp
--- a/Engine/Plugins/Animation/PoseSearch/Source/Editor/Private/PoseSearchDatabaseAssetTree.cpp
+++ b/Engine/Plugins/Animation/PoseSearch/Source/Editor/Private/PoseSearchDatabaseAssetTree.cpp
@@ -893,7 +893,7 @@ namespace UE::PoseSearch
 							if (const FPoseSearchDatabaseAnimationAssetBase* DatabaseAnimationAsset = Database->GetMutableAnimationAssetBase(SelectedNode->SourceAssetIdx))
 							{
 								// @todo: Support copying assets added via BranchIn notifies. 
-								if (!DatabaseAnimationAsset->bSynchronizeWithExternalDependency)
+								if (!DatabaseAnimationAsset->IsSynchronizedWithExternalDependency())
 								{
 									ClipboardContent->CopyDatabaseItem(DatabaseAnimationAsset);
 								}
@@ -972,7 +972,7 @@ namespace UE::PoseSearch
 					{
 						const FPoseSearchDatabaseAnimationAssetBase* DatabaseAnimationAsset = Database->GetAnimationAssetBase(SelectedNode->SourceAssetIdx);
 						
-						if (DatabaseAnimationAsset && !DatabaseAnimationAsset->bSynchronizeWithExternalDependency)
+						if (DatabaseAnimationAsset && !DatabaseAnimationAsset->IsSynchronizedWithExternalDependency())
 						{
 							OnDeleteAsset(SelectedNode, false);
 						}
@@ -1093,7 +1093,7 @@ namespace UE::PoseSearch
 
 							UAnimNotifyState_PoseSearchBranchIn* PoseSearchBranchIn = CastChecked<UAnimNotifyState_PoseSearchBranchIn>(UAnimationBlueprintLibrary::AddAnimationNotifyStateEvent(AnimSequenceBase, TrackName, StartTime, Duration, UAnimNotifyState_PoseSearchBranchIn::StaticClass()));
 							PoseSearchBranchIn->Database = PoseSearchDatabase;
-							DatabaseAnimationAssetBase->bSynchronizeWithExternalDependency = true;
+							DatabaseAnimationAssetBase->BranchInId = PoseSearchBranchIn->GetBranchInId();
 						}
 					}
 				}
diff --git a/Engine/Plugins/Animation/PoseSearch/Source/Editor/Private/PoseSearchDatabaseViewModel.cpp b/Engine/Plugins/Animation/PoseSearch/Source/Editor/Private/PoseSearchDatabaseViewModel.cpp
--- a/Engine/Plugins/Animation/PoseSearch/Source/Editor/Private/PoseSearchDatabaseViewModel.cpp
+++ b/Engine/Plugins/Animation/PoseSearch/Source/Editor/Private/PoseSearchDatabaseViewModel.cpp
@@ -596,30 +596,37 @@ bool FDatabaseViewModel::DeleteFromDatabase(int32 AnimationAssetIndex)
 		{
 			Database->Modify();
 			
-			if (UAnimSequenceBase* AnimSequenceBase = Cast<UAnimSequenceBase>(DatabaseAnimationAssetBase->GetAnimationAsset()))
+			if (DatabaseAnimationAssetBase->IsSynchronizedWithExternalDependency())
 			{
-				bool bModified = false;
-				for (int32 NotifyIndex = AnimSequenceBase->Notifies.Num() - 1; NotifyIndex >= 0; --NotifyIndex)
+				if (UAnimSequenceBase* AnimSequenceBase = Cast<UAnimSequenceBase>(DatabaseAnimationAssetBase->GetAnimationAsset()))
 				{
-					const FAnimNotifyEvent& NotifyEvent = AnimSequenceBase->Notifies[NotifyIndex];
-					if (const UAnimNotifyState_PoseSearchBranchIn* PoseSearchBranchIn = Cast<UAnimNotifyState_PoseSearchBranchIn>(NotifyEvent.NotifyStateClass))
+					bool bModified = false;
+					for (int32 NotifyIndex = AnimSequenceBase->Notifies.Num() - 1; NotifyIndex >= 0; --NotifyIndex)
 					{
-						if (PoseSearchBranchIn->Database == Database)
+						const FAnimNotifyEvent& NotifyEvent = AnimSequenceBase->Notifies[NotifyIndex];
+						if (const UAnimNotifyState_PoseSearchBranchIn* PoseSearchBranchIn = Cast<UAnimNotifyState_PoseSearchBranchIn>(NotifyEvent.NotifyStateClass))
 						{
-							if (!bModified)
+							if (PoseSearchBranchIn->Database == Database && PoseSearchBranchIn->GetBranchInId() == DatabaseAnimationAssetBase->BranchInId)
 							{
-								AnimSequenceBase->Modify();
-								bModified = true;
-							}
+								if (!bModified)
+								{
+									AnimSequenceBase->Modify();
+									bModified = true;
+								}
 
-							AnimSequenceBase->Notifies.RemoveAt(NotifyIndex);
+								AnimSequenceBase->Notifies.RemoveAt(NotifyIndex);
+							}
 						}
 					}
-				}
 
-				if (bModified)
+					if (bModified)
+					{
+						AnimSequenceBase->RefreshCacheData();
+					}
+				}
+				else
 				{
-					AnimSequenceBase->RefreshCacheData();	
+					UE_LOG(LogPoseSearchEditor, Error, TEXT("found DatabaseAnimationAssetBase with valid BranchInId, but invalid AnimSequenceBase in %s"), *Database->GetName());
 				}
 			}
 
diff --git a/Engine/Plugins/Animation/PoseSearch/Source/Runtime/Private/PoseSearchAnimNotifies.cpp b/Engine/Plugins/Animation/PoseSearch/Source/Runtime/Private/PoseSearchAnimNotifies.cpp
--- a/Engine/Plugins/Animation/PoseSearch/Source/Runtime/Private/PoseSearchAnimNotifies.cpp
+++ b/Engine/Plugins/Animation/PoseSearch/Source/Runtime/Private/PoseSearchAnimNotifies.cpp
@@ -4,3 +4,9 @@
 
 #include UE_INLINE_GENERATED_CPP_BY_NAME(PoseSearchAnimNotifies)
 
+uint32 UAnimNotifyState_PoseSearchBranchIn::GetBranchInId() const
+{
+	const uint32 BranchInId = GetTypeHash(GetFullName());
+	check(BranchInId != 0);
+	return BranchInId;
+}
diff --git a/Engine/Plugins/Animation/PoseSearch/Source/Runtime/Private/PoseSearchDatabase.cpp b/Engine/Plugins/Animation/PoseSearch/Source/Runtime/Private/PoseSearchDatabase.cpp
--- a/Engine/Plugins/Animation/PoseSearch/Source/Runtime/Private/PoseSearchDatabase.cpp
+++ b/Engine/Plugins/Animation/PoseSearch/Source/Runtime/Private/PoseSearchDatabase.cpp
@@ -774,6 +774,33 @@ void UPoseSearchDatabase::PostLoad()
 #if WITH_EDITOR
 	using namespace UE::PoseSearch;
 
+	bool bRequiresSynchronization = false;
+	for (FInstancedStruct& AnimationAsset : AnimationAssets)
+	{
+		if (FPoseSearchDatabaseAnimationAssetBase* AnimationAssetBase = AnimationAsset.GetMutablePtr<FPoseSearchDatabaseAnimationAssetBase>())
+		{
+			if (AnimationAssetBase->bSynchronizeWithExternalDependency_DEPRECATED)
+			{
+				Modify();
+				bRequiresSynchronization = true;
+				break;
+			}
+		}
+	}
+	
+	if (bRequiresSynchronization)
+	{
+		SynchronizeWithExternalDependencies();
+
+		for (FInstancedStruct& AnimationAsset : AnimationAssets)
+		{
+			if (FPoseSearchDatabaseAnimationAssetBase* AnimationAssetBase = AnimationAsset.GetMutablePtr<FPoseSearchDatabaseAnimationAssetBase>())
+			{
+				AnimationAssetBase->bSynchronizeWithExternalDependency_DEPRECATED = false;
+			}
+		}
+	}
+
 	ERequestAsyncBuildFlag Flag = ERequestAsyncBuildFlag::NewRequest;
 #if WITH_ENGINE
 	// If there isn't an EditorEngine (ex. Standalone Game via -game argument) we WaitForCompletion
@@ -879,7 +906,7 @@ void UPoseSearchDatabase::SynchronizeWithExternalDependencies(TConstArrayView<UA
 		DisableReselection.Add(AnimationAssetBase.bDisableReselection);
 		AnimationAssetBase.bDisableReselection = false;
 
-		const bool bRequiresSynchronization = AnimationAssetBase.bSynchronizeWithExternalDependency && SequencesBase.Contains(AnimationAssetBase.GetAnimationAsset());
+		const bool bRequiresSynchronization = AnimationAssetBase.IsSynchronizedWithExternalDependency() && SequencesBase.Contains(AnimationAssetBase.GetAnimationAsset());
 		if (!bRequiresSynchronization)
 		{
 			NewAnimationAssets.Add(AnimationAsset);
@@ -912,7 +939,7 @@ void UPoseSearchDatabase::SynchronizeWithExternalDependencies(TConstArrayView<UA
 							FPoseSearchDatabaseSequence DatabaseSequence;
 							DatabaseSequence.Sequence = Sequence;
 							DatabaseSequence.SamplingRange = GetSamplingRange(NotifyEvent, SequenceBase);
-							DatabaseSequence.bSynchronizeWithExternalDependency = true;
+							DatabaseSequence.BranchInId = PoseSearchBranchIn->GetBranchInId();
 							NewAnimationAssets.Add(FInstancedStruct::Make(DatabaseSequence));
 						}
 						else if (UAnimComposite* AnimComposite = Cast<UAnimComposite>(SequenceBase))
@@ -920,7 +947,7 @@ void UPoseSearchDatabase::SynchronizeWithExternalDependencies(TConstArrayView<UA
 							FPoseSearchDatabaseAnimComposite DatabaseAnimComposite;
 							DatabaseAnimComposite.AnimComposite = AnimComposite;
 							DatabaseAnimComposite.SamplingRange = GetSamplingRange(NotifyEvent, SequenceBase);
-							DatabaseAnimComposite.bSynchronizeWithExternalDependency = true;
+							DatabaseAnimComposite.BranchInId = PoseSearchBranchIn->GetBranchInId();
 							NewAnimationAssets.Add(FInstancedStruct::Make(DatabaseAnimComposite));
 						}
 						else if (UAnimMontage* AnimMontage = Cast<UAnimMontage>(SequenceBase))
@@ -928,7 +955,7 @@ void UPoseSearchDatabase::SynchronizeWithExternalDependencies(TConstArrayView<UA
 							FPoseSearchDatabaseAnimMontage DatabaseAnimMontage;
 							DatabaseAnimMontage.AnimMontage = AnimMontage;
 							DatabaseAnimMontage.SamplingRange = GetSamplingRange(NotifyEvent, SequenceBase);
-							DatabaseAnimMontage.bSynchronizeWithExternalDependency = true;
+							DatabaseAnimMontage.BranchInId = PoseSearchBranchIn->GetBranchInId();
 							NewAnimationAssets.Add(FInstancedStruct::Make(DatabaseAnimMontage));
 						}
 					}
diff --git a/Engine/Plugins/Animation/PoseSearch/Source/Runtime/Public/PoseSearch/PoseSearchAnimNotifies.h b/Engine/Plugins/Animation/PoseSearch/Source/Runtime/Public/PoseSearch/PoseSearchAnimNotifies.h
--- a/Engine/Plugins/Animation/PoseSearch/Source/Runtime/Public/PoseSearch/PoseSearchAnimNotifies.h
+++ b/Engine/Plugins/Animation/PoseSearch/Source/Runtime/Public/PoseSearch/PoseSearchAnimNotifies.h
@@ -114,4 +114,6 @@ public:
 
 	UPROPERTY(EditAnywhere, Category = Config)
 	TObjectPtr<UPoseSearchDatabase> Database;
+
+	uint32 GetBranchInId() const;
 };
diff --git a/Engine/Plugins/Animation/PoseSearch/Source/Runtime/Public/PoseSearch/PoseSearchDatabase.h b/Engine/Plugins/Animation/PoseSearch/Source/Runtime/Public/PoseSearch/PoseSearchDatabase.h
--- a/Engine/Plugins/Animation/PoseSearch/Source/Runtime/Public/PoseSearch/PoseSearchDatabase.h
+++ b/Engine/Plugins/Animation/PoseSearch/Source/Runtime/Public/PoseSearch/PoseSearchDatabase.h
@@ -82,6 +82,8 @@ struct POSESEARCH_API FPoseSearchDatabaseAnimationAssetBase
 	virtual int64 GetEditorMemSize() const;
 	virtual int64 GetApproxCookedSize() const { return GetEditorMemSize(); }
 
+	bool IsSynchronizedWithExternalDependency() const { return bSynchronizeWithExternalDependency_DEPRECATED || BranchInId != 0; }
+
 	// This allows users to enable or exclude animations from this database. Useful for debugging.
 	UPROPERTY(EditAnywhere, Category = "Settings", meta = (DisplayPriority = 1))
 	bool bEnabled = true;
@@ -95,10 +97,15 @@ struct POSESEARCH_API FPoseSearchDatabaseAnimationAssetBase
 	UPROPERTY(EditAnywhere, Category = "Settings", meta = (DisplayPriority = 3))
 	EPoseSearchMirrorOption MirrorOption = EPoseSearchMirrorOption::UnmirroredOnly;
 
+	// SynchronizeWithExternalDependency is true when this asset has been added via SynchronizeWithExternalDependencies.
+	// To delete it, remove the PoseSearchBranchIn notify state
+	UPROPERTY()
+	bool bSynchronizeWithExternalDependency_DEPRECATED = false;
+
 	// SynchronizeWithExternalDependency is true when this asset has been added via SynchronizeWithExternalDependencies.
 	// To delete it, remove the PoseSearchBranchIn notify state
 	UPROPERTY(VisibleAnywhere, Category = "Settings", meta = (DisplayPriority = 20))
-	bool bSynchronizeWithExternalDependency = false;
+	uint32 BranchInId = 0;
 #endif // WITH_EDITORONLY_DATA
 };
 
