PoseSearch: Only log on missing bone index when finalizing to avoid crashing.

#rnx

diff --git a/Engine/Plugins/Animation/PoseSearch/Source/Runtime/Private/PoseSearchSchema.cpp b/Engine/Plugins/Animation/PoseSearch/Source/Runtime/Private/PoseSearchSchema.cpp
--- a/Engine/Plugins/Animation/PoseSearch/Source/Runtime/Private/PoseSearchSchema.cpp
+++ b/Engine/Plugins/Animation/PoseSearch/Source/Runtime/Private/PoseSearchSchema.cpp
@@ -305,10 +305,18 @@ void UPoseSearchSchema::Finalize()
 
 			if (RoledSkeleton.MirrorDataTable)
 			{
-				const FSkeletonPoseBoneIndex MirroredBoneIndex = RoledSkeleton.MirrorDataTable->BoneToMirrorBoneIndex[BoneRef.BoneIndex];
-				if (MirroredBoneIndex.IsValid())
+				if (RoledSkeleton.MirrorDataTable->BoneToMirrorBoneIndex.IsValidIndex(BoneRef.BoneIndex))
 				{
-					RoledSkeleton.BoneIndicesWithParents.AddUnique(MirroredBoneIndex.GetInt());
+					const FSkeletonPoseBoneIndex MirroredBoneIndex = RoledSkeleton.MirrorDataTable->BoneToMirrorBoneIndex[BoneRef.BoneIndex];
+					if (MirroredBoneIndex.IsValid())
+					{
+						RoledSkeleton.BoneIndicesWithParents.AddUnique(MirroredBoneIndex.GetInt());
+					}
+				}
+				else
+				{
+					UE_LOG(LogPoseSearch, Warning, TEXT("UPoseSearchSchema::Finalize: couldn't Finalize '%s' because bone index doest not exist in mirror table or mirrot table is empty."), *GetNameSafe(this));
+					// @todo: Investigate PostLoad dependency issues.
 				}
 			}
 		}
