motion matching - adding a ConditionalPostLoad dependency to UMirrorDataTable, that via UMirrorDataTable::FillMirrorArrays populates UMirrorDataTable::BoneToMirrorBoneIndex used in UPoseSearchSchema::Finalize

diff --git a/Engine/Plugins/Animation/PoseSearch/Source/Runtime/Private/PoseSearchSchema.cpp b/Engine/Plugins/Animation/PoseSearch/Source/Runtime/Private/PoseSearchSchema.cpp
--- a/Engine/Plugins/Animation/PoseSearch/Source/Runtime/Private/PoseSearchSchema.cpp
+++ b/Engine/Plugins/Animation/PoseSearch/Source/Runtime/Private/PoseSearchSchema.cpp
@@ -316,7 +316,6 @@ void UPoseSearchSchema::Finalize()
 				else
 				{
 					UE_LOG(LogPoseSearch, Warning, TEXT("UPoseSearchSchema::Finalize: couldn't Finalize '%s' because bone index doest not exist in mirror table or mirrot table is empty."), *GetNameSafe(this));
-					// @todo: Investigate PostLoad dependency issues.
 				}
 			}
 		}
@@ -357,6 +356,16 @@ void UPoseSearchSchema::PostLoad()
 		MirrorDataTable_DEPRECATED = nullptr;
 	}
 
+	for (FPoseSearchRoledSkeleton& Skeleton : Skeletons)
+	{
+		if (Skeleton.MirrorDataTable)
+		{
+			// adding a ConditionalPostLoad dependency to UMirrorDataTable, that via UMirrorDataTable::FillMirrorArrays
+			// populates UMirrorDataTable::BoneToMirrorBoneIndex used in UPoseSearchSchema::Finalize 
+			Skeleton.MirrorDataTable->ConditionalPostLoad();
+		}
+	}
+
 	Finalize();
 }
 
