motion matching - adding ClampToMinMaxLimits meta to database's assets SamplingRange

diff --git a/Engine/Plugins/Animation/PoseSearch/Source/Editor/Private/PoseSearchDatabaseAssetListItem.cpp b/Engine/Plugins/Animation/PoseSearch/Source/Editor/Private/PoseSearchDatabaseAssetListItem.cpp
--- a/Engine/Plugins/Animation/PoseSearch/Source/Editor/Private/PoseSearchDatabaseAssetListItem.cpp
+++ b/Engine/Plugins/Animation/PoseSearch/Source/Editor/Private/PoseSearchDatabaseAssetListItem.cpp
@@ -160,7 +160,7 @@ namespace UE::PoseSearch
 		AssetTypeColor = FColor::White;
 		if (UPoseSearchDatabase* Database = InEditorViewModel->GetPoseSearchDatabase())
 		{
-			if (const FPoseSearchDatabaseAnimationAssetBase* DatabaseAnimationAsset = Database->GetAnimationAssetBase(WeakAssetTreeNode.Pin()->SourceAssetIdx))
+			if (const FPoseSearchDatabaseAnimationAssetBase* DatabaseAnimationAsset = Database->GetDatabaseAnimationAsset<FPoseSearchDatabaseAnimationAssetBase>(WeakAssetTreeNode.Pin()->SourceAssetIdx))
 			{
 				static FAssetToolsModule& AssetToolsModule = FModuleManager::LoadModuleChecked<FAssetToolsModule>(TEXT("AssetTools"));
 				if (TSharedPtr<IAssetTypeActions> AssetTypeActions = AssetToolsModule.Get().GetAssetTypeActionsForClass(DatabaseAnimationAsset->GetAnimationAssetStaticClass()).Pin())
@@ -242,7 +242,7 @@ namespace UE::PoseSearch
 			TSharedPtr<FDatabaseViewModel> ViewModel = EditorViewModel.Pin();
 			if (const UPoseSearchDatabase* Database = ViewModel->GetPoseSearchDatabase())
 			{
-				if (const FPoseSearchDatabaseAnimationAssetBase* DatabaseAnimationAsset = Database->GetAnimationAssetBase(AssetTreeNode->SourceAssetIdx))
+				if (const FPoseSearchDatabaseAnimationAssetBase* DatabaseAnimationAsset = Database->GetDatabaseAnimationAsset<FPoseSearchDatabaseAnimationAssetBase>(AssetTreeNode->SourceAssetIdx))
 				{
 					if (UAssetEditorSubsystem* AssetEditorSS = GEditor->GetEditorSubsystem<UAssetEditorSubsystem>())
 					{
@@ -283,7 +283,7 @@ namespace UE::PoseSearch
 		{
 			if (TSharedPtr<FDatabaseAssetTreeNode> AssetTreeNode = WeakAssetTreeNode.Pin())
 			{
-				if (const FPoseSearchDatabaseAnimationAssetBase* DatabaseAnimationAsset = Database->GetAnimationAssetBase(AssetTreeNode->SourceAssetIdx))
+				if (const FPoseSearchDatabaseAnimationAssetBase* DatabaseAnimationAsset = Database->GetDatabaseAnimationAsset<FPoseSearchDatabaseAnimationAssetBase>(AssetTreeNode->SourceAssetIdx))
 				{
 					return FText::FromString(DatabaseAnimationAsset->GetName());
 				}
@@ -350,7 +350,7 @@ namespace UE::PoseSearch
 				   {
 					  if (const TSharedPtr<FDatabaseAssetTreeNode> AssetTreeNode = WeakAssetTreeNode.Pin())
 					  {
-						 if (const FPoseSearchDatabaseAnimationAssetBase* DatabaseAnimationAssetBase = Database->GetAnimationAssetBase(AssetTreeNode->SourceAssetIdx))
+						 if (const FPoseSearchDatabaseAnimationAssetBase* DatabaseAnimationAssetBase = Database->GetDatabaseAnimationAsset<FPoseSearchDatabaseAnimationAssetBase>(AssetTreeNode->SourceAssetIdx))
 						 {
 							if (DatabaseAnimationAssetBase->IsSynchronizedWithExternalDependency())
 							{
@@ -373,7 +373,7 @@ namespace UE::PoseSearch
 				TSharedPtr<SImage> ItemIconWidget;
 				if (UPoseSearchDatabase* Database = ViewModel->GetPoseSearchDatabase())
 				{
-					if (const FPoseSearchDatabaseAnimationAssetBase* DatabaseAnimationAsset = Database->GetAnimationAssetBase(SourceAssetIdx))
+					if (const FPoseSearchDatabaseAnimationAssetBase* DatabaseAnimationAsset = Database->GetDatabaseAnimationAsset<FPoseSearchDatabaseAnimationAssetBase>(SourceAssetIdx))
 					{
 						SAssignNew(ItemIconWidget, SImage)
 						.Image(FSlateIconFinder::FindIconBrushForClass(DatabaseAnimationAsset->GetAnimationAssetStaticClass()));
@@ -442,7 +442,7 @@ namespace UE::PoseSearch
 			TSharedPtr<SObjectPropertyEntryBox> AssetPickerWidget;
 			if (UPoseSearchDatabase* Database = ViewModel->GetPoseSearchDatabase())
 			{
-				if (const FPoseSearchDatabaseAnimationAssetBase* DatabaseAnimationAsset = Database->GetAnimationAssetBase(SourceAssetIdx))
+				if (const FPoseSearchDatabaseAnimationAssetBase* DatabaseAnimationAsset = Database->GetDatabaseAnimationAsset<FPoseSearchDatabaseAnimationAssetBase>(SourceAssetIdx))
 				{
 					SAssignNew(AssetPickerWidget, SObjectPropertyEntryBox)
 					.AllowClear(false)
@@ -679,7 +679,7 @@ namespace UE::PoseSearch
 		{
 			if (const UPoseSearchDatabase* Database = ViewModelPtr->GetPoseSearchDatabase())
 			{
-				if (const FPoseSearchDatabaseAnimationAssetBase* DatabaseAnimationAsset = Database->GetAnimationAssetBase(AssetTreeNode->SourceAssetIdx))
+				if (const FPoseSearchDatabaseAnimationAssetBase* DatabaseAnimationAsset = Database->GetDatabaseAnimationAsset<FPoseSearchDatabaseAnimationAssetBase>(AssetTreeNode->SourceAssetIdx))
 				{
 					if (const UObject* AnimAsset = DatabaseAnimationAsset->GetAnimationAsset())
 					{
@@ -718,7 +718,7 @@ namespace UE::PoseSearch
 		{
 			if (const TSharedPtr<FDatabaseAssetTreeNode> AssetTreeNode = WeakAssetTreeNode.Pin())
 			{
-				if (const FPoseSearchDatabaseAnimationAssetBase* DatabaseAnimationAssetBase = Database->GetAnimationAssetBase(AssetTreeNode->SourceAssetIdx))
+				if (const FPoseSearchDatabaseAnimationAssetBase* DatabaseAnimationAssetBase = Database->GetDatabaseAnimationAsset<FPoseSearchDatabaseAnimationAssetBase>(AssetTreeNode->SourceAssetIdx))
 				{
 					if (DatabaseAnimationAssetBase->IsEnabled())
 					{
@@ -742,7 +742,7 @@ namespace UE::PoseSearch
 		{
 			if (const TSharedPtr<FDatabaseAssetTreeNode> AssetTreeNode = WeakAssetTreeNode.Pin())
 			{
-				if (const FPoseSearchDatabaseAnimationAssetBase* DatabaseAnimationAssetBase = Database->GetAnimationAssetBase(AssetTreeNode->SourceAssetIdx))
+				if (const FPoseSearchDatabaseAnimationAssetBase* DatabaseAnimationAssetBase = Database->GetDatabaseAnimationAsset<FPoseSearchDatabaseAnimationAssetBase>(AssetTreeNode->SourceAssetIdx))
 				{
 					if (DatabaseAnimationAssetBase->IsEnabled())
 					{
@@ -853,7 +853,7 @@ namespace UE::PoseSearch
 		{
 			if (const TSharedPtr<FDatabaseAssetTreeNode> AssetTreeNode = WeakAssetTreeNode.Pin())
 			{
-				if (const FPoseSearchDatabaseAnimationAssetBase* DatabaseAnimationAssetBase = Database->GetAnimationAssetBase(AssetTreeNode->SourceAssetIdx))
+				if (const FPoseSearchDatabaseAnimationAssetBase* DatabaseAnimationAssetBase = Database->GetDatabaseAnimationAsset<FPoseSearchDatabaseAnimationAssetBase>(AssetTreeNode->SourceAssetIdx))
 				{
 					if (DatabaseAnimationAssetBase->IsEnabled())
 					{
diff --git a/Engine/Plugins/Animation/PoseSearch/Source/Editor/Private/PoseSearchDatabaseAssetTree.cpp b/Engine/Plugins/Animation/PoseSearch/Source/Editor/Private/PoseSearchDatabaseAssetTree.cpp
--- a/Engine/Plugins/Animation/PoseSearch/Source/Editor/Private/PoseSearchDatabaseAssetTree.cpp
+++ b/Engine/Plugins/Animation/PoseSearch/Source/Editor/Private/PoseSearchDatabaseAssetTree.cpp
@@ -246,8 +246,8 @@ namespace UE::PoseSearch
 
 			IndexArray.Sort([Database](int32 SequenceIdxA, int32 SequenceIdxB)
 			{
-				const FPoseSearchDatabaseAnimationAssetBase* A = Database->GetAnimationAssetBase(SequenceIdxA);
-				const FPoseSearchDatabaseAnimationAssetBase* B = Database->GetAnimationAssetBase(SequenceIdxB);
+				const FPoseSearchDatabaseAnimationAssetBase* A = Database->GetDatabaseAnimationAsset<FPoseSearchDatabaseAnimationAssetBase>(SequenceIdxA);
+				const FPoseSearchDatabaseAnimationAssetBase* B = Database->GetDatabaseAnimationAsset<FPoseSearchDatabaseAnimationAssetBase>(SequenceIdxB);
 
 				//If its null add it to the end of the list 
 				if (!B->GetAnimationAsset())
@@ -269,7 +269,7 @@ namespace UE::PoseSearch
 			{
 				const int32 MappedId = IndexArray[AnimationAssetIdx];
 
-				if (const FPoseSearchDatabaseAnimationAssetBase* DatabaseAnimationAsset = Database->GetAnimationAssetBase(MappedId))
+				if (const FPoseSearchDatabaseAnimationAssetBase* DatabaseAnimationAsset = Database->GetDatabaseAnimationAsset<FPoseSearchDatabaseAnimationAssetBase>(MappedId))
 				{
 					const bool bFiltered = (DatabaseAnimationAsset->GetAnimationAsset() == nullptr || GetAssetFilterString().IsEmpty()) ? false : !DatabaseAnimationAsset->GetName().Contains(GetAssetFilterString());
 
@@ -890,7 +890,7 @@ namespace UE::PoseSearch
 					{
 						if (UPoseSearchDatabase* Database = ViewModel->GetPoseSearchDatabase())
 						{
-							if (const FPoseSearchDatabaseAnimationAssetBase* DatabaseAnimationAsset = Database->GetMutableAnimationAssetBase(SelectedNode->SourceAssetIdx))
+							if (const FPoseSearchDatabaseAnimationAssetBase* DatabaseAnimationAsset = Database->GetDatabaseAnimationAsset<FPoseSearchDatabaseAnimationAssetBase>(SelectedNode->SourceAssetIdx))
 							{
 								// @todo: Support copying assets added via BranchIn notifies. 
 								if (!DatabaseAnimationAsset->IsSynchronizedWithExternalDependency())
@@ -970,7 +970,7 @@ namespace UE::PoseSearch
 				{
 					if (UPoseSearchDatabase* Database = ViewModel->GetPoseSearchDatabase())
 					{
-						const FPoseSearchDatabaseAnimationAssetBase* DatabaseAnimationAsset = Database->GetAnimationAssetBase(SelectedNode->SourceAssetIdx);
+						const FPoseSearchDatabaseAnimationAssetBase* DatabaseAnimationAsset = Database->GetDatabaseAnimationAsset<FPoseSearchDatabaseAnimationAssetBase>(SelectedNode->SourceAssetIdx);
 						
 						if (DatabaseAnimationAsset && !DatabaseAnimationAsset->IsSynchronizedWithExternalDependency())
 						{
@@ -1071,7 +1071,7 @@ namespace UE::PoseSearch
 				
 				for (const TSharedPtr<FDatabaseAssetTreeNode>& SelectedNode : SelectedNodes)
 				{
-					if (FPoseSearchDatabaseAnimationAssetBase* DatabaseAnimationAssetBase = PoseSearchDatabase->GetMutableAnimationAssetBase(SelectedNode->SourceAssetIdx))
+					if (FPoseSearchDatabaseAnimationAssetBase* DatabaseAnimationAssetBase = PoseSearchDatabase->GetMutableDatabaseAnimationAsset<FPoseSearchDatabaseAnimationAssetBase>(SelectedNode->SourceAssetIdx))
 					{
 						if (UAnimSequenceBase* AnimSequenceBase = Cast<UAnimSequenceBase>(DatabaseAnimationAssetBase->GetAnimationAsset()))
 						{
diff --git a/Engine/Plugins/Animation/PoseSearch/Source/Editor/Private/PoseSearchDatabaseAssetTreeNode.cpp b/Engine/Plugins/Animation/PoseSearch/Source/Editor/Private/PoseSearchDatabaseAssetTreeNode.cpp
--- a/Engine/Plugins/Animation/PoseSearch/Source/Editor/Private/PoseSearchDatabaseAssetTreeNode.cpp
+++ b/Engine/Plugins/Animation/PoseSearch/Source/Editor/Private/PoseSearchDatabaseAssetTreeNode.cpp
@@ -41,7 +41,7 @@ namespace UE::PoseSearch
 	{
 		if (const UPoseSearchDatabase* Database = EditorViewModel.Pin()->GetPoseSearchDatabase())
 		{
-			if (const FPoseSearchDatabaseAnimationAssetBase* DatabaseAnimationAsset = Database->GetAnimationAssetBase(SourceAssetIdx))
+			if (const FPoseSearchDatabaseAnimationAssetBase* DatabaseAnimationAsset = Database->GetDatabaseAnimationAsset<FPoseSearchDatabaseAnimationAssetBase>(SourceAssetIdx))
 			{
 				return DatabaseAnimationAsset->IsRootMotionEnabled();
 			}
@@ -54,7 +54,7 @@ namespace UE::PoseSearch
 	{
 		if (const UPoseSearchDatabase* Database = EditorViewModel.Pin()->GetPoseSearchDatabase())
 		{
-			if (const FPoseSearchDatabaseAnimationAssetBase* DatabaseAnimationAsset = Database->GetAnimationAssetBase(SourceAssetIdx))
+			if (const FPoseSearchDatabaseAnimationAssetBase* DatabaseAnimationAsset = Database->GetDatabaseAnimationAsset<FPoseSearchDatabaseAnimationAssetBase>(SourceAssetIdx))
 			{
 				return DatabaseAnimationAsset->IsLooping();
 			}
@@ -67,7 +67,7 @@ namespace UE::PoseSearch
 	{
 		if (const UPoseSearchDatabase* Database = EditorViewModel.Pin()->GetPoseSearchDatabase())
 		{
-			if (const FPoseSearchDatabaseAnimationAssetBase* DatabaseAnimationAsset = Database->GetAnimationAssetBase(SourceAssetIdx))
+			if (const FPoseSearchDatabaseAnimationAssetBase* DatabaseAnimationAsset = Database->GetDatabaseAnimationAsset<FPoseSearchDatabaseAnimationAssetBase>(SourceAssetIdx))
 			{
 				return DatabaseAnimationAsset->GetMirrorOption();
 			}
diff --git a/Engine/Plugins/Animation/PoseSearch/Source/Editor/Private/PoseSearchDatabaseEditor.cpp b/Engine/Plugins/Animation/PoseSearch/Source/Editor/Private/PoseSearchDatabaseEditor.cpp
--- a/Engine/Plugins/Animation/PoseSearch/Source/Editor/Private/PoseSearchDatabaseEditor.cpp
+++ b/Engine/Plugins/Animation/PoseSearch/Source/Editor/Private/PoseSearchDatabaseEditor.cpp
@@ -563,7 +563,7 @@ namespace UE::PoseSearch
 					continue;
 				}
 
-				const FInstancedStruct& DatabaseAsset = PoseSearchDatabase->GetAnimationAssetStruct(SelectedItem->SourceAssetIdx);
+				const FInstancedStruct& DatabaseAsset = PoseSearchDatabase->GetAnimationAssets()[SelectedItem->SourceAssetIdx];
 				const UScriptStruct* ScriptStruct = DatabaseAsset.GetScriptStruct();
 				FSelectionWidget& SelectionWidget = FindOrAddSelectionWidget(ScriptStruct);
 
diff --git a/Engine/Plugins/Animation/PoseSearch/Source/Editor/Private/PoseSearchDatabaseEditorReflection.cpp b/Engine/Plugins/Animation/PoseSearch/Source/Editor/Private/PoseSearchDatabaseEditorReflection.cpp
--- a/Engine/Plugins/Animation/PoseSearch/Source/Editor/Private/PoseSearchDatabaseEditorReflection.cpp
+++ b/Engine/Plugins/Animation/PoseSearch/Source/Editor/Private/PoseSearchDatabaseEditorReflection.cpp
@@ -38,6 +38,13 @@ void UPoseSearchDatabaseReflectionBase::PostEditChangeProperty(FPropertyChangedE
 
 bool UPoseSearchDatabaseSequenceReflection::ApplyChanges() 
 {
+	const float PlayLength = Sequence.GetPlayLength();
+	const FFloatInterval ClampedSamplingRange(FMath::Clamp(Sequence.SamplingRange.Min, 0.0f, PlayLength), FMath::Clamp(Sequence.SamplingRange.Max, 0.0f, PlayLength));
+	if (ClampedSamplingRange != Sequence.SamplingRange)
+	{
+		Sequence.SamplingRange = ClampedSamplingRange;
+	}
+
 	if (const TSharedPtr<UE::PoseSearch::FDatabaseAssetTreeNode> AssetTreeNode = WeakAssetTreeNode.Pin())
 	{
 		if (const TSharedPtr<UE::PoseSearch::FDatabaseViewModel> ViewModel = AssetTreeNode->EditorViewModel.Pin())
@@ -45,16 +52,14 @@ bool UPoseSearchDatabaseSequenceReflection::ApplyChanges()
 			UPoseSearchDatabase* Database = ViewModel->GetPoseSearchDatabase();
 			if (IsValid(Database))
 			{
-				FInstancedStruct& DatabaseAsset = Database->GetMutableAnimationAssetStruct(AssetTreeNode->SourceAssetIdx);
-				if (FPoseSearchDatabaseSequence* DatabaseSequence = DatabaseAsset.GetMutablePtr<FPoseSearchDatabaseSequence>())
+				if (FPoseSearchDatabaseSequence* DatabaseSequence = Database->GetMutableDatabaseAnimationAsset<FPoseSearchDatabaseSequence>(AssetTreeNode->SourceAssetIdx))
 				{
-					Sequence.SamplingRange.Min = FMath::Clamp(Sequence.SamplingRange.Min, 0.0f, Sequence.GetPlayLength()); 
-					Sequence.SamplingRange.Max = FMath::Clamp(Sequence.SamplingRange.Max, 0.0f, Sequence.GetPlayLength()); 
-					
-					*DatabaseSequence = Sequence;
-					Database->MarkPackageDirty();
-
-					return true;
+					if (*DatabaseSequence != Sequence)
+					{
+						*DatabaseSequence = Sequence;
+						Database->MarkPackageDirty();
+						return true;
+					}
 				}
 			}
 		}
@@ -72,13 +77,14 @@ bool UPoseSearchDatabaseBlendSpaceReflection::ApplyChanges()
 			UPoseSearchDatabase* Database = ViewModel->GetPoseSearchDatabase();
 			if (IsValid(Database))
 			{
-				FInstancedStruct& DatabaseAsset = Database->GetMutableAnimationAssetStruct(AssetTreeNode->SourceAssetIdx);
-				if (FPoseSearchDatabaseBlendSpace* DatabaseBlendSpace = DatabaseAsset.GetMutablePtr<FPoseSearchDatabaseBlendSpace>())
+				if (FPoseSearchDatabaseBlendSpace* DatabaseBlendSpace = Database->GetMutableDatabaseAnimationAsset<FPoseSearchDatabaseBlendSpace>(AssetTreeNode->SourceAssetIdx))
 				{
-					*DatabaseBlendSpace = BlendSpace;
-					Database->MarkPackageDirty();
-
-					return true;
+					if (*DatabaseBlendSpace != BlendSpace)
+					{
+						*DatabaseBlendSpace = BlendSpace;
+						Database->MarkPackageDirty();
+						return true;
+					}
 				}
 			}
 		}
@@ -89,6 +95,13 @@ bool UPoseSearchDatabaseBlendSpaceReflection::ApplyChanges()
 
 bool UPoseSearchDatabaseAnimCompositeReflection::ApplyChanges()
 {
+	const float PlayLength = AnimComposite.GetPlayLength();
+	const FFloatInterval ClampedSamplingRange(FMath::Clamp(AnimComposite.SamplingRange.Min, 0.0f, PlayLength), FMath::Clamp(AnimComposite.SamplingRange.Max, 0.0f, PlayLength));
+	if (ClampedSamplingRange != AnimComposite.SamplingRange)
+	{
+		AnimComposite.SamplingRange = ClampedSamplingRange;
+	}
+
 	if (const TSharedPtr<UE::PoseSearch::FDatabaseAssetTreeNode> AssetTreeNode = WeakAssetTreeNode.Pin())
 	{
 		if (const TSharedPtr<UE::PoseSearch::FDatabaseViewModel> ViewModel = AssetTreeNode->EditorViewModel.Pin())
@@ -96,17 +109,14 @@ bool UPoseSearchDatabaseAnimCompositeReflection::ApplyChanges()
 			UPoseSearchDatabase* Database = ViewModel->GetPoseSearchDatabase();
 			if (IsValid(Database))
 			{
-				FInstancedStruct& DatabaseAsset = Database->GetMutableAnimationAssetStruct(AssetTreeNode->SourceAssetIdx);
-				if (FPoseSearchDatabaseAnimComposite* DatabaseAnimComposite = DatabaseAsset.GetMutablePtr<FPoseSearchDatabaseAnimComposite>())
+				if (FPoseSearchDatabaseAnimComposite* DatabaseAnimComposite = Database->GetMutableDatabaseAnimationAsset<FPoseSearchDatabaseAnimComposite>(AssetTreeNode->SourceAssetIdx))
 				{
-					AnimComposite.SamplingRange.Min = FMath::Clamp(AnimComposite.SamplingRange.Min, 0.0f, AnimComposite.GetPlayLength());
-					AnimComposite.SamplingRange.Max = FMath::Clamp(AnimComposite.SamplingRange.Max, 0.0f, AnimComposite.GetPlayLength());
-					
-					*DatabaseAnimComposite = AnimComposite;
-
-					Database->MarkPackageDirty();
-
-					return true;
+					if (*DatabaseAnimComposite != AnimComposite)
+					{
+						*DatabaseAnimComposite = AnimComposite;
+						Database->MarkPackageDirty();
+						return true;
+					}
 				}
 			}
 		}
@@ -117,6 +127,13 @@ bool UPoseSearchDatabaseAnimCompositeReflection::ApplyChanges()
 
 bool UPoseSearchDatabaseAnimMontageReflection::ApplyChanges()
 {
+	const float PlayLength = AnimMontage.GetPlayLength();
+	const FFloatInterval ClampedSamplingRange(FMath::Clamp(AnimMontage.SamplingRange.Min, 0.0f, PlayLength), FMath::Clamp(AnimMontage.SamplingRange.Max, 0.0f, PlayLength));
+	if (ClampedSamplingRange != AnimMontage.SamplingRange)
+	{
+		AnimMontage.SamplingRange = ClampedSamplingRange;
+	}
+
 	if (const TSharedPtr<UE::PoseSearch::FDatabaseAssetTreeNode> AssetTreeNode = WeakAssetTreeNode.Pin())
 	{
 		if (const TSharedPtr<UE::PoseSearch::FDatabaseViewModel> ViewModel = AssetTreeNode->EditorViewModel.Pin())
@@ -124,17 +141,14 @@ bool UPoseSearchDatabaseAnimMontageReflection::ApplyChanges()
 			UPoseSearchDatabase* Database = ViewModel->GetPoseSearchDatabase();
 			if (IsValid(Database))
 			{
-				FInstancedStruct& DatabaseAsset = Database->GetMutableAnimationAssetStruct(AssetTreeNode->SourceAssetIdx);
-				if (FPoseSearchDatabaseAnimMontage* DatabaseAnimMontage = DatabaseAsset.GetMutablePtr<FPoseSearchDatabaseAnimMontage>())
+				if (FPoseSearchDatabaseAnimMontage* DatabaseAnimMontage = Database->GetMutableDatabaseAnimationAsset<FPoseSearchDatabaseAnimMontage>(AssetTreeNode->SourceAssetIdx))
 				{
-					AnimMontage.SamplingRange.Min = FMath::Clamp(AnimMontage.SamplingRange.Min, 0.0f, AnimMontage.GetPlayLength()); 
-					AnimMontage.SamplingRange.Max = FMath::Clamp(AnimMontage.SamplingRange.Max, 0.0f, AnimMontage.GetPlayLength());
-					
-					*DatabaseAnimMontage = AnimMontage;
-
-					Database->MarkPackageDirty();
-
-					return true;
+					if (*DatabaseAnimMontage != AnimMontage)
+					{
+						*DatabaseAnimMontage = AnimMontage;
+						Database->MarkPackageDirty();
+						return true;
+					}
 				}
 			}
 		}
@@ -145,6 +159,13 @@ bool UPoseSearchDatabaseAnimMontageReflection::ApplyChanges()
 
 bool UPoseSearchDatabaseMultiSequenceReflection::ApplyChanges()
 {
+	const float PlayLength = MultiSequence.GetPlayLength();
+	const FFloatInterval ClampedSamplingRange(FMath::Clamp(MultiSequence.SamplingRange.Min, 0.0f, PlayLength), FMath::Clamp(MultiSequence.SamplingRange.Max, 0.0f, PlayLength));
+	if (ClampedSamplingRange != MultiSequence.SamplingRange)
+	{
+		MultiSequence.SamplingRange = ClampedSamplingRange;
+	}
+
 	if (const TSharedPtr<UE::PoseSearch::FDatabaseAssetTreeNode> AssetTreeNode = WeakAssetTreeNode.Pin())
 	{
 		if (const TSharedPtr<UE::PoseSearch::FDatabaseViewModel> ViewModel = AssetTreeNode->EditorViewModel.Pin())
@@ -152,13 +173,14 @@ bool UPoseSearchDatabaseMultiSequenceReflection::ApplyChanges()
 			UPoseSearchDatabase* Database = ViewModel->GetPoseSearchDatabase();
 			if (IsValid(Database))
 			{
-				FInstancedStruct& DatabaseAsset = Database->GetMutableAnimationAssetStruct(AssetTreeNode->SourceAssetIdx);
-				if (FPoseSearchDatabaseMultiSequence* DatabaseMultiSequence = DatabaseAsset.GetMutablePtr<FPoseSearchDatabaseMultiSequence>())
+				if (FPoseSearchDatabaseMultiSequence* DatabaseMultiSequence = Database->GetMutableDatabaseAnimationAsset<FPoseSearchDatabaseMultiSequence>(AssetTreeNode->SourceAssetIdx))
 				{
-					*DatabaseMultiSequence = MultiSequence;
-					Database->MarkPackageDirty();
-
-					return true;
+					if (*DatabaseMultiSequence != MultiSequence)
+					{
+						*DatabaseMultiSequence = MultiSequence;
+						Database->MarkPackageDirty();
+						return true;
+					}
 				}
 			}
 		}
@@ -176,7 +198,7 @@ void UPoseSearchDatabaseStatistics::Initialize(const UPoseSearchDatabase* PoseSe
 		const UE::PoseSearch::FSearchIndex& SearchIndex = PoseSearchDatabase->GetSearchIndex();
 		// General Information
 	
-		AnimationSequences = PoseSearchDatabase->GetAnimationAssets().Num();
+		AnimationSequences = PoseSearchDatabase->GetNumAnimationAssets();
 			
 		const int32 SampleRate = FMath::Max(1, PoseSearchDatabase->Schema->SampleRate);
 		TotalAnimationPosesInFrames = SearchIndex.GetNumPoses();
@@ -232,10 +254,11 @@ void UPoseSearchDatabaseStatistics::Initialize(const UPoseSearchDatabase* PoseSe
 		{
 			uint32 SourceAnimAssetsSizeCookedEstimateInBytes = 0;
 			TSet<const UObject*> Analyzed;
-			Analyzed.Reserve(PoseSearchDatabase->GetAnimationAssets().Num());
-			for (const FInstancedStruct& AnimAsset : PoseSearchDatabase->GetAnimationAssets())
+			const int32 NumAnimationAssets = PoseSearchDatabase->GetNumAnimationAssets();
+			Analyzed.Reserve(NumAnimationAssets);
+			for (int32 AnimationAssetIndex = 0; AnimationAssetIndex < NumAnimationAssets; ++AnimationAssetIndex)
 			{
-				if (const FPoseSearchDatabaseAnimationAssetBase* DatabaseAnimationAssetBase = AnimAsset.GetPtr<FPoseSearchDatabaseAnimationAssetBase>())
+				if (const FPoseSearchDatabaseAnimationAssetBase* DatabaseAnimationAssetBase = PoseSearchDatabase->GetDatabaseAnimationAsset<FPoseSearchDatabaseAnimationAssetBase>(AnimationAssetIndex))
 				{
 					bool bAlreadyAnalyzed = false;
 					Analyzed.Add(DatabaseAnimationAssetBase->GetAnimationAsset(), &bAlreadyAnalyzed);
diff --git a/Engine/Plugins/Animation/PoseSearch/Source/Editor/Private/PoseSearchDatabaseViewModel.cpp b/Engine/Plugins/Animation/PoseSearch/Source/Editor/Private/PoseSearchDatabaseViewModel.cpp
--- a/Engine/Plugins/Animation/PoseSearch/Source/Editor/Private/PoseSearchDatabaseViewModel.cpp
+++ b/Engine/Plugins/Animation/PoseSearch/Source/Editor/Private/PoseSearchDatabaseViewModel.cpp
@@ -54,7 +54,7 @@ bool FDatabasePreviewActor::SpawnPreviewActor(UWorld* World, const UPoseSearchDa
 
 	const FSearchIndexAsset& IndexAsset = SearchIndex.Assets[IndexAssetIdx];
 
-	const FPoseSearchDatabaseAnimationAssetBase* DatabaseAnimationAsset = PoseSearchDatabase->GetAnimationAssetBase(IndexAsset.GetSourceAssetIdx());
+	const FPoseSearchDatabaseAnimationAssetBase* DatabaseAnimationAsset = PoseSearchDatabase->GetDatabaseAnimationAsset<FPoseSearchDatabaseAnimationAssetBase>(IndexAsset.GetSourceAssetIdx());
 	UAnimationAsset* PreviewAsset = Cast<UAnimationAsset>(DatabaseAnimationAsset->GetAnimationAssetForRole(Role));
 	if (!PreviewAsset)
 	{
@@ -592,7 +592,7 @@ bool FDatabaseViewModel::DeleteFromDatabase(int32 AnimationAssetIndex)
 {
 	if (UPoseSearchDatabase* Database = GetPoseSearchDatabase())
 	{
-		if (const FPoseSearchDatabaseAnimationAssetBase* DatabaseAnimationAssetBase = Database->GetAnimationAssetBase(AnimationAssetIndex))
+		if (const FPoseSearchDatabaseAnimationAssetBase* DatabaseAnimationAssetBase = Database->GetDatabaseAnimationAsset<FPoseSearchDatabaseAnimationAssetBase>(AnimationAssetIndex))
 		{
 			Database->Modify();
 			
@@ -643,7 +643,7 @@ void FDatabaseViewModel::SetDisableReselection(int32 AnimationAssetIndex, bool b
 {
 	if (UPoseSearchDatabase* Database = GetPoseSearchDatabase())
 	{
-		if (FPoseSearchDatabaseAnimationAssetBase* DatabaseAnimationAsset = Database->GetMutableAnimationAssetBase(AnimationAssetIndex))
+		if (FPoseSearchDatabaseAnimationAssetBase* DatabaseAnimationAsset = Database->GetMutableDatabaseAnimationAsset<FPoseSearchDatabaseAnimationAssetBase>(AnimationAssetIndex))
 		{
 			Database->Modify();
 			
@@ -656,7 +656,7 @@ bool FDatabaseViewModel::IsDisableReselection(int32 AnimationAssetIndex) const
 {
 	if (const UPoseSearchDatabase* Database = GetPoseSearchDatabase())
 	{
-		if (const FPoseSearchDatabaseAnimationAssetBase* DatabaseAnimationAsset = Database->GetAnimationAssetBase(AnimationAssetIndex))
+		if (const FPoseSearchDatabaseAnimationAssetBase* DatabaseAnimationAsset = Database->GetDatabaseAnimationAsset<FPoseSearchDatabaseAnimationAssetBase>(AnimationAssetIndex))
 		{
 			return DatabaseAnimationAsset->IsDisableReselection();
 		}
@@ -669,7 +669,7 @@ void FDatabaseViewModel::SetIsEnabled(int32 AnimationAssetIndex, bool bEnabled)
 {
 	if (UPoseSearchDatabase* Database = GetPoseSearchDatabase())
 	{
-		if (FPoseSearchDatabaseAnimationAssetBase* DatabaseAnimationAsset = Database->GetMutableAnimationAssetBase(AnimationAssetIndex))
+		if (FPoseSearchDatabaseAnimationAssetBase* DatabaseAnimationAsset = Database->GetMutableDatabaseAnimationAsset<FPoseSearchDatabaseAnimationAssetBase>(AnimationAssetIndex))
 		{
 			Database->Modify();
 
@@ -682,7 +682,7 @@ bool FDatabaseViewModel::IsEnabled(int32 AnimationAssetIndex) const
 {
 	if (const UPoseSearchDatabase* Database = GetPoseSearchDatabase())
 	{
-		if (const FPoseSearchDatabaseAnimationAssetBase* DatabaseAnimationAsset = Database->GetAnimationAssetBase(AnimationAssetIndex))
+		if (const FPoseSearchDatabaseAnimationAssetBase* DatabaseAnimationAsset = Database->GetDatabaseAnimationAsset<FPoseSearchDatabaseAnimationAssetBase>(AnimationAssetIndex))
 		{
 			return DatabaseAnimationAsset->IsEnabled();
 		}
@@ -697,7 +697,7 @@ bool FDatabaseViewModel::SetAnimationAsset(int32 AnimationAssetIndex, UObject* A
 	{
 		if (UPoseSearchDatabase* Database = GetPoseSearchDatabase())
 		{
-			if (FPoseSearchDatabaseAnimationAssetBase* DatabaseAnimationAsset = Database->GetMutableAnimationAssetBase(AnimationAssetIndex))
+			if (FPoseSearchDatabaseAnimationAssetBase* DatabaseAnimationAsset = Database->GetMutableDatabaseAnimationAsset<FPoseSearchDatabaseAnimationAssetBase>(AnimationAssetIndex))
 			{
 				Database->Modify();
 
@@ -744,7 +744,7 @@ void FDatabaseViewModel::SetMirrorOption(int32 AnimationAssetIndex, EPoseSearchM
 {
 	if (UPoseSearchDatabase* Database = GetPoseSearchDatabase())
 	{
-		if (FPoseSearchDatabaseAnimationAssetBase* DatabaseAnimationAsset = Database->GetMutableAnimationAssetBase(AnimationAssetIndex))
+		if (FPoseSearchDatabaseAnimationAssetBase* DatabaseAnimationAsset = Database->GetMutableDatabaseAnimationAsset<FPoseSearchDatabaseAnimationAssetBase>(AnimationAssetIndex))
 		{
 			Database->Modify();
 
@@ -757,7 +757,7 @@ EPoseSearchMirrorOption FDatabaseViewModel::GetMirrorOption(int32 AnimationAsset
 {
 	if (const UPoseSearchDatabase* Database = GetPoseSearchDatabase())
 	{
-		if (const FPoseSearchDatabaseAnimationAssetBase* DatabaseAnimationAsset = Database->GetAnimationAssetBase(AnimationAssetIndex))
+		if (const FPoseSearchDatabaseAnimationAssetBase* DatabaseAnimationAsset = Database->GetDatabaseAnimationAsset<FPoseSearchDatabaseAnimationAssetBase>(AnimationAssetIndex))
 		{
 			return DatabaseAnimationAsset->MirrorOption;
 		}
@@ -790,7 +790,7 @@ int32 FDatabaseViewModel::SetSelectedNode(int32 PoseIdx, bool bClearSelection, b
 				if (SearchIndex.Assets.IsValidIndex(IndexAssetIndex))
 				{
 					const FSearchIndexAsset& IndexAsset = SearchIndex.Assets[IndexAssetIndex];
-					const FPoseSearchDatabaseAnimationAssetBase* DatabaseAnimationAsset = Database->GetAnimationAssetBase(IndexAsset.GetSourceAssetIdx());
+					const FPoseSearchDatabaseAnimationAssetBase* DatabaseAnimationAsset = Database->GetDatabaseAnimationAsset<FPoseSearchDatabaseAnimationAssetBase>(IndexAsset.GetSourceAssetIdx());
 					check(DatabaseAnimationAsset);
 					int32 PreviewActorGroupIndex = INDEX_NONE;
 					for (int32 RoleIndex = 0; RoleIndex < DatabaseAnimationAsset->GetNumRoles(); ++RoleIndex)
@@ -853,7 +853,7 @@ void FDatabaseViewModel::SetSelectedNodes(const TArrayView<TSharedPtr<FDatabaseA
 				const FSearchIndexAsset& IndexAsset = SearchIndex.Assets[IndexAssetIndex];
 				if (AssociatedAssetIndices.Find(IndexAsset.GetSourceAssetIdx()))
 				{
-					const FPoseSearchDatabaseAnimationAssetBase* DatabaseAnimationAsset = Database->GetAnimationAssetBase(IndexAsset.GetSourceAssetIdx());
+					const FPoseSearchDatabaseAnimationAssetBase* DatabaseAnimationAsset = Database->GetDatabaseAnimationAsset<FPoseSearchDatabaseAnimationAssetBase>(IndexAsset.GetSourceAssetIdx());
 					check(DatabaseAnimationAsset);
 					int32 PreviewActorGroupIndex = INDEX_NONE;
 					for (int32 RoleIndex = 0; RoleIndex < DatabaseAnimationAsset->GetNumRoles(); ++RoleIndex)
@@ -998,7 +998,7 @@ bool FDatabaseViewModel::GetAnimationTime(int32 SourceAssetIdx, float& CurrentPl
 					CurrentPlayTime = PlayTime + IndexAsset.GetFirstSampleTime(Database->Schema->SampleRate);
 					BlendParameters = IndexAsset.GetBlendParameters();
 
-					const bool bIsBlendSpace = Database->GetAnimationAssetStruct(IndexAsset).GetPtr<FPoseSearchDatabaseBlendSpace>() != nullptr;
+					const bool bIsBlendSpace = Database->GetDatabaseAnimationAsset<FPoseSearchDatabaseBlendSpace>(IndexAsset) != nullptr;
 					if (bIsBlendSpace && !FMath::IsNearlyEqual(MaxPreviewPlayLength, MinPreviewPlayLength))
 					{
 						CurrentPlayTime = (CurrentPlayTime - MaxPreviewPlayLength) / (MaxPreviewPlayLength - MinPreviewPlayLength);
diff --git a/Engine/Plugins/Animation/PoseSearch/Source/Editor/Private/PoseSearchDebuggerDatabaseView.cpp b/Engine/Plugins/Animation/PoseSearch/Source/Editor/Private/PoseSearchDebuggerDatabaseView.cpp
--- a/Engine/Plugins/Animation/PoseSearch/Source/Editor/Private/PoseSearchDebuggerDatabaseView.cpp
+++ b/Engine/Plugins/Animation/PoseSearch/Source/Editor/Private/PoseSearchDebuggerDatabaseView.cpp
@@ -216,8 +216,7 @@ static void AddUnfilteredDatabaseRow(const UPoseSearchDatabase* Database,
 			}
 		}
 
-		const FInstancedStruct& DatabaseAssetStruct = Database->GetAnimationAssetStruct(*SearchIndexAsset);
-		if (const FPoseSearchDatabaseAnimationAssetBase* DatabaseAsset = DatabaseAssetStruct.GetPtr<FPoseSearchDatabaseAnimationAssetBase>())
+		if (const FPoseSearchDatabaseAnimationAssetBase* DatabaseAsset = Database->GetDatabaseAnimationAsset<FPoseSearchDatabaseAnimationAssetBase>(*SearchIndexAsset))
 		{
 			const float PlayLength = DatabaseAsset->GetPlayLength();
 			UObject* AnimationAsset = DatabaseAsset->GetAnimationAsset();
diff --git a/Engine/Plugins/Animation/PoseSearch/Source/Runtime/Private/AnimNode_MotionMatching.cpp b/Engine/Plugins/Animation/PoseSearch/Source/Runtime/Private/AnimNode_MotionMatching.cpp
--- a/Engine/Plugins/Animation/PoseSearch/Source/Runtime/Private/AnimNode_MotionMatching.cpp
+++ b/Engine/Plugins/Animation/PoseSearch/Source/Runtime/Private/AnimNode_MotionMatching.cpp
@@ -194,7 +194,7 @@ void FAnimNode_MotionMatching::UpdateAssetPlayer(const FAnimationUpdateContext&
 		const UPoseSearchDatabase* CurrentResultDatabase = MotionMatchingState.CurrentSearchResult.Database.Get();
 		if (SearchIndexAsset && CurrentResultDatabase && CurrentResultDatabase->Schema)
 		{
-			const FPoseSearchDatabaseAnimationAssetBase* DatabaseAsset = CurrentResultDatabase->GetAnimationAssetBase(*SearchIndexAsset);
+			const FPoseSearchDatabaseAnimationAssetBase* DatabaseAsset = CurrentResultDatabase->GetDatabaseAnimationAsset<FPoseSearchDatabaseAnimationAssetBase>(*SearchIndexAsset);
 			check(DatabaseAsset);
 
 			if (UAnimationAsset* AnimationAsset = Cast<UAnimationAsset>(DatabaseAsset->GetAnimationAsset()))
diff --git a/Engine/Plugins/Animation/PoseSearch/Source/Runtime/Private/MotionMatchingAnimNodeLibrary.cpp b/Engine/Plugins/Animation/PoseSearch/Source/Runtime/Private/MotionMatchingAnimNodeLibrary.cpp
--- a/Engine/Plugins/Animation/PoseSearch/Source/Runtime/Private/MotionMatchingAnimNodeLibrary.cpp
+++ b/Engine/Plugins/Animation/PoseSearch/Source/Runtime/Private/MotionMatchingAnimNodeLibrary.cpp
@@ -22,7 +22,7 @@ void UMotionMatchingAnimNodeLibrary::GetMotionMatchingSearchResult(const FMotion
 			const UPoseSearchDatabase* CurrentResultDatabase = MotionMatchingState.CurrentSearchResult.Database.Get();
 			if (CurrentResultDatabase && CurrentResultDatabase->Schema)
 			{
-				if (const FPoseSearchDatabaseAnimationAssetBase* DatabaseAsset = CurrentResultDatabase->GetAnimationAssetBase(*SearchIndexAsset))
+				if (const FPoseSearchDatabaseAnimationAssetBase* DatabaseAsset = CurrentResultDatabase->GetDatabaseAnimationAsset<FPoseSearchDatabaseAnimationAssetBase>(*SearchIndexAsset))
 				{
 					Result.SelectedAnimation = DatabaseAsset->GetAnimationAsset();
 					Result.SelectedTime = MotionMatchingState.CurrentSearchResult.AssetTime;
diff --git a/Engine/Plugins/Animation/PoseSearch/Source/Runtime/Private/PoseSearchDatabase.cpp b/Engine/Plugins/Animation/PoseSearch/Source/Runtime/Private/PoseSearchDatabase.cpp
--- a/Engine/Plugins/Animation/PoseSearch/Source/Runtime/Private/PoseSearchDatabase.cpp
+++ b/Engine/Plugins/Animation/PoseSearch/Source/Runtime/Private/PoseSearchDatabase.cpp
@@ -6,7 +6,6 @@
 #include "Animation/AnimSequence.h"
 #include "Animation/BlendSpace.h"
 #include "Animation/BlendSpace1D.h"
-#include "InstancedStruct.h"
 #include "PoseSearch/PoseSearchAnimNotifies.h"
 #include "PoseSearch/PoseSearchContext.h"
 #include "PoseSearch/PoseSearchDefines.h"
@@ -52,7 +51,7 @@ static void PopulateSelectableAssetIdx(FSelectableAssetIdx& SelectableAssetIdx,
 
 		for (int32 AssetIndex = 0; AssetIndex < SearchIndex.Assets.Num(); ++AssetIndex)
 		{
-			if (const FPoseSearchDatabaseAnimationAssetBase* DatabaseAnimationAssetBase = Database->GetAnimationAssetBase(SearchIndex.Assets[AssetIndex]))
+			if (const FPoseSearchDatabaseAnimationAssetBase* DatabaseAnimationAssetBase = Database->GetDatabaseAnimationAsset<FPoseSearchDatabaseAnimationAssetBase>(SearchIndex.Assets[AssetIndex]))
 			{
 				if (AssetsToConsider.Contains(DatabaseAnimationAssetBase->GetAnimationAsset()))
 				{
@@ -707,7 +706,8 @@ const FInstancedStruct& UPoseSearchDatabase::GetAnimationAssetStruct(int32 Anima
 
 const FInstancedStruct& UPoseSearchDatabase::GetAnimationAssetStruct(const UE::PoseSearch::FSearchIndexAsset& SearchIndexAsset) const
 {
-	return GetAnimationAssetStruct(SearchIndexAsset.GetSourceAssetIdx());
+	check(AnimationAssets.IsValidIndex(SearchIndexAsset.GetSourceAssetIdx()));
+	return AnimationAssets[SearchIndexAsset.GetSourceAssetIdx()];
 }
 
 FInstancedStruct& UPoseSearchDatabase::GetMutableAnimationAssetStruct(int32 AnimationAssetIndex)
@@ -718,37 +718,8 @@ FInstancedStruct& UPoseSearchDatabase::GetMutableAnimationAssetStruct(int32 Anim
 
 FInstancedStruct& UPoseSearchDatabase::GetMutableAnimationAssetStruct(const UE::PoseSearch::FSearchIndexAsset& SearchIndexAsset)
 {
-	return GetMutableAnimationAssetStruct(SearchIndexAsset.GetSourceAssetIdx());
-}
-
-const FPoseSearchDatabaseAnimationAssetBase* UPoseSearchDatabase::GetAnimationAssetBase(int32 AnimationAssetIndex) const
-{
-	if (AnimationAssets.IsValidIndex(AnimationAssetIndex))
-	{
-		return AnimationAssets[AnimationAssetIndex].GetPtr<FPoseSearchDatabaseAnimationAssetBase>();
-	}
-
-	return nullptr;
-}
-
-const FPoseSearchDatabaseAnimationAssetBase* UPoseSearchDatabase::GetAnimationAssetBase(const UE::PoseSearch::FSearchIndexAsset& SearchIndexAsset) const
-{
-	return GetAnimationAssetBase(SearchIndexAsset.GetSourceAssetIdx());
-}
-
-FPoseSearchDatabaseAnimationAssetBase* UPoseSearchDatabase::GetMutableAnimationAssetBase(int32 AnimationAssetIndex)
-{
-	if (AnimationAssets.IsValidIndex(AnimationAssetIndex))
-	{
-		return AnimationAssets[AnimationAssetIndex].GetMutablePtr<FPoseSearchDatabaseAnimationAssetBase>();
-	}
-
-	return nullptr;
-}
-
-FPoseSearchDatabaseAnimationAssetBase* UPoseSearchDatabase::GetMutableAnimationAssetBase(const UE::PoseSearch::FSearchIndexAsset& SearchIndexAsset)
-{
-	return GetMutableAnimationAssetBase(SearchIndexAsset.GetSourceAssetIdx());
+	check(AnimationAssets.IsValidIndex(SearchIndexAsset.GetSourceAssetIdx()));
+	return AnimationAssets[SearchIndexAsset.GetSourceAssetIdx()];
 }
 
 #if WITH_EDITOR
@@ -1211,7 +1182,7 @@ UE::PoseSearch::FSearchResult UPoseSearchDatabase::SearchContinuingPose(UE::Pose
 	const FSearchIndex& SearchIndex = GetSearchIndex();
 	const int32 PoseIdx = SearchContext.GetCurrentResult().PoseIdx;
 	const FSearchIndexAsset& SearchIndexAsset = SearchIndex.GetAssetForPose(PoseIdx);
-	const FPoseSearchDatabaseAnimationAssetBase* DatabaseAnimationAssetBase = GetAnimationAssetStruct(SearchIndexAsset).GetPtr<FPoseSearchDatabaseAnimationAssetBase>();
+	const FPoseSearchDatabaseAnimationAssetBase* DatabaseAnimationAssetBase = GetDatabaseAnimationAsset<FPoseSearchDatabaseAnimationAssetBase>(SearchIndexAsset);
 	check(DatabaseAnimationAssetBase);
 
 	float UpdatedContinuingPoseCostBias = ContinuingPoseCostBias;
diff --git a/Engine/Plugins/Animation/PoseSearch/Source/Runtime/Private/PoseSearchLibrary.cpp b/Engine/Plugins/Animation/PoseSearch/Source/Runtime/Private/PoseSearchLibrary.cpp
--- a/Engine/Plugins/Animation/PoseSearch/Source/Runtime/Private/PoseSearchLibrary.cpp
+++ b/Engine/Plugins/Animation/PoseSearch/Source/Runtime/Private/PoseSearchLibrary.cpp
@@ -730,7 +730,7 @@ void UPoseSearchLibrary::MotionMatch(
 	{
 		const UPoseSearchDatabase* Database = SearchResult.Database.Get();
 		const FSearchIndexAsset* SearchIndexAsset = SearchResult.GetSearchIndexAsset();
-		if (const FPoseSearchDatabaseAnimationAssetBase* DatabaseAsset = Database->GetAnimationAssetBase(*SearchIndexAsset))
+		if (const FPoseSearchDatabaseAnimationAssetBase* DatabaseAsset = Database->GetDatabaseAnimationAsset<FPoseSearchDatabaseAnimationAssetBase>(*SearchIndexAsset))
 		{
 			Result.SelectedAnimation = DatabaseAsset->GetAnimationAsset();
 			Result.SelectedTime = SearchResult.AssetTime;
@@ -894,7 +894,7 @@ UE::PoseSearch::FSearchResult UPoseSearchLibrary::MotionMatch(
 				for (int32 AssetIndex = 0; AssetIndex < SearchIndex.Assets.Num(); ++AssetIndex)
 				{
 					const FSearchIndexAsset& SearchIndexAsset = SearchIndex.Assets[AssetIndex];
-					if (const FPoseSearchDatabaseAnimationAssetBase* DatabaseAnimationAssetBase = Database->GetAnimationAssetBase(SearchIndexAsset))
+					if (const FPoseSearchDatabaseAnimationAssetBase* DatabaseAnimationAssetBase = Database->GetDatabaseAnimationAsset<FPoseSearchDatabaseAnimationAssetBase>(SearchIndexAsset.GetSourceAssetIdx()))
 					{
 						if (PlayingAnimationAsset == DatabaseAnimationAssetBase->GetAnimationAsset())
 						{
diff --git a/Engine/Plugins/Animation/PoseSearch/Source/Runtime/Private/PoseSearchModule.cpp b/Engine/Plugins/Animation/PoseSearch/Source/Runtime/Private/PoseSearchModule.cpp
--- a/Engine/Plugins/Animation/PoseSearch/Source/Runtime/Private/PoseSearchModule.cpp
+++ b/Engine/Plugins/Animation/PoseSearch/Source/Runtime/Private/PoseSearchModule.cpp
@@ -73,7 +73,7 @@ public:
 		if (const UE::PoseSearch::FSearchIndexAsset* SearchIndexAsset = SearchResult.GetSearchIndexAsset())
 		{
 			const UPoseSearchDatabase* Database = SearchResult.Database.Get();
-			if (const FPoseSearchDatabaseAnimationAssetBase* DatabaseAnimationAssetBase = Database->GetAnimationAssetBase(*SearchIndexAsset))
+			if (const FPoseSearchDatabaseAnimationAssetBase* DatabaseAnimationAssetBase = Database->GetDatabaseAnimationAsset<FPoseSearchDatabaseAnimationAssetBase>(*SearchIndexAsset))
 			{
 				ProviderResult.SelectedAsset = DatabaseAnimationAssetBase->GetAnimationAsset();
 				ProviderResult.Dissimilarity = SearchResult.PoseCost.GetTotalCost();
diff --git a/Engine/Plugins/Animation/PoseSearch/Source/Runtime/Private/PoseSearchResult.cpp b/Engine/Plugins/Animation/PoseSearch/Source/Runtime/Private/PoseSearchResult.cpp
--- a/Engine/Plugins/Animation/PoseSearch/Source/Runtime/Private/PoseSearchResult.cpp
+++ b/Engine/Plugins/Animation/PoseSearch/Source/Runtime/Private/PoseSearchResult.cpp
@@ -18,13 +18,12 @@ void FSearchResult::Update(float NewAssetTime)
 	else
 	{
 		const FSearchIndexAsset& SearchIndexAsset = Database->GetSearchIndex().GetAssetForPose(PoseIdx);
-		const FInstancedStruct& DatabaseAsset = Database->GetAnimationAssetStruct(SearchIndexAsset);
-		if (DatabaseAsset.GetPtr<FPoseSearchDatabaseSequence>() || DatabaseAsset.GetPtr<FPoseSearchDatabaseAnimComposite>())
+		if (Database->GetDatabaseAnimationAsset<FPoseSearchDatabaseSequence>(SearchIndexAsset) || Database->GetDatabaseAnimationAsset<FPoseSearchDatabaseAnimComposite>(SearchIndexAsset))
 		{
 			PoseIdx = Database->GetPoseIndexFromTime(NewAssetTime, SearchIndexAsset);
 			AssetTime = NewAssetTime;
 		}
-		else if (const FPoseSearchDatabaseBlendSpace* DatabaseBlendSpace = DatabaseAsset.GetPtr<FPoseSearchDatabaseBlendSpace>())
+		else if (const FPoseSearchDatabaseBlendSpace* DatabaseBlendSpace = Database->GetDatabaseAnimationAsset<FPoseSearchDatabaseBlendSpace>(SearchIndexAsset))
 		{
 			TArray<FBlendSampleData> BlendSamples;
 			int32 TriangulationIndex = 0;
@@ -74,8 +73,7 @@ bool FSearchResult::CanAdvance(float DeltaTime) const
 	if (const FSearchIndexAsset* SearchIndexAsset = GetSearchIndexAsset())
 	{
 		float SteppedTime = AssetTime;
-		const FInstancedStruct& DatabaseAsset = Database->GetAnimationAssetStruct(*SearchIndexAsset);
-		if (const FPoseSearchDatabaseBlendSpace* DatabaseBlendSpace = DatabaseAsset.GetPtr<FPoseSearchDatabaseBlendSpace>())
+		if (const FPoseSearchDatabaseBlendSpace* DatabaseBlendSpace = Database->GetDatabaseAnimationAsset<FPoseSearchDatabaseBlendSpace>(*SearchIndexAsset))
 		{
 			if (const UBlendSpace* BlendSpace = DatabaseBlendSpace->BlendSpace.Get())
 			{
@@ -90,7 +88,7 @@ bool FSearchResult::CanAdvance(float DeltaTime) const
 				bCanAdvance = ETAA_Finished != FAnimationRuntime::AdvanceTime(SearchIndexAsset->IsLooping(), DeltaTime, SteppedTime, PlayLength);
 			}
 		}
-		else if (const FPoseSearchDatabaseAnimationAssetBase* DatabaseAnimationAssetBase = DatabaseAsset.GetPtr<FPoseSearchDatabaseAnimationAssetBase>())
+		else if (const FPoseSearchDatabaseAnimationAssetBase* DatabaseAnimationAssetBase = Database->GetDatabaseAnimationAsset<FPoseSearchDatabaseAnimationAssetBase>(*SearchIndexAsset))
 		{
 			const float AssetLength = DatabaseAnimationAssetBase->GetPlayLength();
 			bCanAdvance = ETAA_Finished != FAnimationRuntime::AdvanceTime(SearchIndexAsset->IsLooping(), DeltaTime, SteppedTime, AssetLength);
diff --git a/Engine/Plugins/Animation/PoseSearch/Source/Runtime/Public/PoseSearch/PoseSearchDatabase.h b/Engine/Plugins/Animation/PoseSearch/Source/Runtime/Public/PoseSearch/PoseSearchDatabase.h
--- a/Engine/Plugins/Animation/PoseSearch/Source/Runtime/Public/PoseSearch/PoseSearchDatabase.h
+++ b/Engine/Plugins/Animation/PoseSearch/Source/Runtime/Public/PoseSearch/PoseSearchDatabase.h
@@ -3,6 +3,7 @@
 #pragma once
 
 #include "Engine/DataAsset.h"
+#include "InstancedStruct.h"
 #include "PoseSearch/PoseSearchAssetSampler.h"
 #include "PoseSearch/PoseSearchCost.h"
 #include "PoseSearch/PoseSearchIndex.h"
@@ -10,7 +11,6 @@
 #include "PoseSearch/PoseSearchRole.h"
 #include "PoseSearchDatabase.generated.h"
 
-struct FInstancedStruct;
 class UAnimationAsset;
 class UAnimComposite;
 class UAnimMontage;
@@ -107,6 +107,24 @@ struct POSESEARCH_API FPoseSearchDatabaseAnimationAssetBase
 	UPROPERTY(VisibleAnywhere, Category = "Settings", meta = (DisplayPriority = 20))
 	uint32 BranchInId = 0;
 #endif // WITH_EDITORONLY_DATA
+
+	friend bool operator==(const FPoseSearchDatabaseAnimationAssetBase& A, const FPoseSearchDatabaseAnimationAssetBase& B)
+	{
+#if WITH_EDITORONLY_DATA
+		return A.bEnabled == B.bEnabled &&
+			A.bDisableReselection == B.bDisableReselection &&
+			A.MirrorOption == B.MirrorOption &&
+			A.bSynchronizeWithExternalDependency_DEPRECATED == B.bSynchronizeWithExternalDependency_DEPRECATED &&
+			A.BranchInId == B.BranchInId;
+#else // WITH_EDITORONLY_DATA
+		return true;
+#endif // WITH_EDITORONLY_DATA
+	}
+};
+
+template<> struct TStructOpsTypeTraits<FPoseSearchDatabaseAnimationAssetBase> : public TStructOpsTypeTraitsBase2<FPoseSearchDatabaseAnimationAssetBase>
+{
+	enum { WithIdenticalViaEquality = true };
 };
 
 /** A sequence entry in a UPoseSearchDatabase. */
@@ -116,14 +134,14 @@ struct POSESEARCH_API FPoseSearchDatabaseSequence : public FPoseSearchDatabaseAn
 	GENERATED_BODY()
 	virtual ~FPoseSearchDatabaseSequence() = default;
 
-	UPROPERTY(EditAnywhere, Category="Settings", meta = (DisplayPriority = 0))
+	UPROPERTY(EditAnywhere, Category = "Settings", meta = (DisplayPriority = 0))
 	TObjectPtr<UAnimSequence> Sequence;
 
 #if WITH_EDITORONLY_DATA
 	// It allows users to set a time range to an individual animation sequence in the database. 
 	// This is effectively trimming the beginning and end of the animation in the database (not in the original sequence).
 	// If set to [0, 0] it will be the entire frame range of the original sequence.
-	UPROPERTY(EditAnywhere, Category="Settings", meta = (DisplayPriority = 2))
+	UPROPERTY(EditAnywhere, Category = "Settings", meta = (ClampToMinMaxLimits, DisplayPriority = 2))
 	FFloatInterval SamplingRange = FFloatInterval(0.f, 0.f);
 
 	virtual UClass* GetAnimationAssetStaticClass() const override;
@@ -133,6 +151,21 @@ struct POSESEARCH_API FPoseSearchDatabaseSequence : public FPoseSearchDatabaseAn
 #endif // WITH_EDITORONLY_DATA
 	
 	virtual UObject* GetAnimationAsset() const override;
+
+	friend bool operator==(const FPoseSearchDatabaseSequence& A, const FPoseSearchDatabaseSequence& B)
+	{
+		return static_cast<const FPoseSearchDatabaseAnimationAssetBase&>(A) == static_cast<const FPoseSearchDatabaseAnimationAssetBase&>(B)
+			&& A.Sequence == B.Sequence
+#if WITH_EDITORONLY_DATA
+			&& A.SamplingRange == B.SamplingRange
+#endif // WITH_EDITORONLY_DATA
+			;
+	}
+};
+
+template<> struct TStructOpsTypeTraits<FPoseSearchDatabaseSequence> : public TStructOpsTypeTraitsBase2<FPoseSearchDatabaseSequence>
+{
+	enum { WithIdenticalViaEquality = true };
 };
 
 /** An blend space entry in a UPoseSearchDatabase. */
@@ -184,6 +217,26 @@ struct POSESEARCH_API FPoseSearchDatabaseBlendSpace : public FPoseSearchDatabase
 #if WITH_EDITOR
 	virtual int32 GetFrameAtTime(float Time) const override;
 #endif // WITH_EDITOR
+
+	friend bool operator==(const FPoseSearchDatabaseBlendSpace& A, const FPoseSearchDatabaseBlendSpace& B)
+	{
+		return static_cast<const FPoseSearchDatabaseAnimationAssetBase&>(A) == static_cast<const FPoseSearchDatabaseAnimationAssetBase&>(B)
+			&& A.BlendSpace == B.BlendSpace
+#if WITH_EDITORONLY_DATA
+			&& A.bUseSingleSample == B.bUseSingleSample
+			&& A.bUseGridForSampling == B.bUseGridForSampling
+			&& A.NumberOfHorizontalSamples == B.NumberOfHorizontalSamples
+			&& A.NumberOfVerticalSamples == B.NumberOfVerticalSamples
+			&& A.BlendParamX == B.BlendParamX
+			&& A.BlendParamY == B.BlendParamY
+#endif // WITH_EDITORONLY_DATA
+			;
+	}
+};
+
+template<> struct TStructOpsTypeTraits<FPoseSearchDatabaseBlendSpace> : public TStructOpsTypeTraitsBase2<FPoseSearchDatabaseBlendSpace>
+{
+	enum { WithIdenticalViaEquality = true };
 };
 
 /** An entry in a UPoseSearchDatabase. */
@@ -200,7 +253,7 @@ struct POSESEARCH_API FPoseSearchDatabaseAnimComposite : public FPoseSearchDatab
 	// It allows users to set a time range to an individual animation sequence in the database. 
 	// This is effectively trimming the beginning and end of the animation in the database (not in the original sequence).
 	// If set to [0, 0] it will be the entire frame range of the original sequence.
-	UPROPERTY(EditAnywhere, Category = "Settings", meta = (DisplayPriority = 3))
+	UPROPERTY(EditAnywhere, Category = "Settings", meta = (ClampToMinMaxLimits, DisplayPriority = 3))
 	FFloatInterval SamplingRange = FFloatInterval(0.f, 0.f);
 
 	virtual UClass* GetAnimationAssetStaticClass() const override;
@@ -210,6 +263,21 @@ struct POSESEARCH_API FPoseSearchDatabaseAnimComposite : public FPoseSearchDatab
 #endif // WITH_EDITORONLY_DATA
 
 	virtual UObject* GetAnimationAsset() const override;
+
+	friend bool operator==(const FPoseSearchDatabaseAnimComposite& A, const FPoseSearchDatabaseAnimComposite& B)
+	{
+		return static_cast<const FPoseSearchDatabaseAnimationAssetBase&>(A) == static_cast<const FPoseSearchDatabaseAnimationAssetBase&>(B)
+			&& A.AnimComposite == B.AnimComposite
+#if WITH_EDITORONLY_DATA
+			&& A.SamplingRange == B.SamplingRange
+#endif // WITH_EDITORONLY_DATA
+			;
+	}
+};
+
+template<> struct TStructOpsTypeTraits<FPoseSearchDatabaseAnimComposite> : public TStructOpsTypeTraitsBase2<FPoseSearchDatabaseAnimComposite>
+{
+	enum { WithIdenticalViaEquality = true };
 };
 
 /** An anim montage entry in a UPoseSearchDatabase. */
@@ -219,14 +287,14 @@ struct POSESEARCH_API FPoseSearchDatabaseAnimMontage : public FPoseSearchDatabas
 	GENERATED_BODY()
 	virtual ~FPoseSearchDatabaseAnimMontage() = default;
 
-	UPROPERTY(EditAnywhere, Category="Settings", meta = (DisplayPriority = 0))
+	UPROPERTY(EditAnywhere, Category = "Settings", meta = (DisplayPriority = 0))
 	TObjectPtr<UAnimMontage> AnimMontage;
 
 #if WITH_EDITORONLY_DATA
 	// It allows users to set a time range to an individual animation sequence in the database. 
 	// This is effectively trimming the beginning and end of the animation in the database (not in the original sequence).
 	// If set to [0, 0] it will be the entire frame range of the original sequence.
-	UPROPERTY(EditAnywhere, Category="Settings", meta = (DisplayPriority = 2))
+	UPROPERTY(EditAnywhere, Category = "Settings", meta = (ClampToMinMaxLimits, DisplayPriority = 2))
 	FFloatInterval SamplingRange = FFloatInterval(0.f, 0.f);
 
 	virtual UClass* GetAnimationAssetStaticClass() const override;
@@ -236,6 +304,21 @@ struct POSESEARCH_API FPoseSearchDatabaseAnimMontage : public FPoseSearchDatabas
 #endif // WITH_EDITORONLY_DATA
 
 	virtual UObject* GetAnimationAsset() const override;
+
+	friend bool operator==(const FPoseSearchDatabaseAnimMontage& A, const FPoseSearchDatabaseAnimMontage& B)
+	{
+		return static_cast<const FPoseSearchDatabaseAnimationAssetBase&>(A) == static_cast<const FPoseSearchDatabaseAnimationAssetBase&>(B)
+			&& A.AnimMontage == B.AnimMontage
+#if WITH_EDITORONLY_DATA
+			&& A.SamplingRange == B.SamplingRange
+#endif // WITH_EDITORONLY_DATA
+			;
+	}
+};
+
+template<> struct TStructOpsTypeTraits<FPoseSearchDatabaseAnimMontage> : public TStructOpsTypeTraitsBase2<FPoseSearchDatabaseAnimMontage>
+{
+	enum { WithIdenticalViaEquality = true };
 };
 
 USTRUCT(BlueprintType, Category = "Animation|Pose Search")
@@ -244,14 +327,14 @@ struct POSESEARCH_API FPoseSearchDatabaseMultiSequence : public FPoseSearchDatab
 	GENERATED_BODY()
 	virtual ~FPoseSearchDatabaseMultiSequence() = default;
 
-	UPROPERTY(EditAnywhere, Category="Settings", meta = (DisplayPriority = 0))
+	UPROPERTY(EditAnywhere, Category = "Settings", meta = (DisplayPriority = 0))
 	TObjectPtr<UPoseSearchMultiSequence> MultiSequence;
 
 #if WITH_EDITORONLY_DATA
 	// It allows users to set a time range to an individual animation sequence in the database. 
 	// This is effectively trimming the beginning and end of the animation in the database (not in the original sequence).
 	// If set to [0, 0] it will be the entire frame range of the original sequence.
-	UPROPERTY(EditAnywhere, Category="Settings", meta = (DisplayPriority = 2))
+	UPROPERTY(EditAnywhere, Category = "Settings", meta = (ClampToMinMaxLimits, DisplayPriority = 2))
 	FFloatInterval SamplingRange = FFloatInterval(0.f, 0.f);
 
 	virtual UClass* GetAnimationAssetStaticClass() const override;
@@ -271,6 +354,21 @@ struct POSESEARCH_API FPoseSearchDatabaseMultiSequence : public FPoseSearchDatab
 #if WITH_EDITOR
 	virtual int32 GetFrameAtTime(float Time) const override;
 #endif // WITH_EDITOR
+
+	friend bool operator==(const FPoseSearchDatabaseMultiSequence& A, const FPoseSearchDatabaseMultiSequence& B)
+	{
+		return static_cast<const FPoseSearchDatabaseAnimationAssetBase&>(A) == static_cast<const FPoseSearchDatabaseAnimationAssetBase&>(B)
+			&& A.MultiSequence == B.MultiSequence
+#if WITH_EDITORONLY_DATA
+			&& A.SamplingRange == B.SamplingRange
+#endif // WITH_EDITORONLY_DATA
+			;
+	}
+};
+
+template<> struct TStructOpsTypeTraits<FPoseSearchDatabaseMultiSequence> : public TStructOpsTypeTraitsBase2<FPoseSearchDatabaseMultiSequence>
+{
+	enum { WithIdenticalViaEquality = true };
 };
 
 /** A data asset for indexing a collection of animation sequences. */
@@ -392,15 +490,33 @@ public:
 	void AddAnimationAsset(FInstancedStruct AnimationAsset);
 	void RemoveAnimationAssetAt(int32 AnimationAssetIndex);
 
+	template<typename TDatabaseAnimationAsset> const TDatabaseAnimationAsset* GetDatabaseAnimationAsset(int32 AnimationAssetIndex) const;
+	template<typename TDatabaseAnimationAsset> const TDatabaseAnimationAsset* GetDatabaseAnimationAsset(const UE::PoseSearch::FSearchIndexAsset& SearchIndexAsset) const;
+	template<typename TDatabaseAnimationAsset> TDatabaseAnimationAsset* GetMutableDatabaseAnimationAsset(int32 AnimationAssetIndex);
+	template<typename TDatabaseAnimationAsset> TDatabaseAnimationAsset* GetMutableDatabaseAnimationAsset(const UE::PoseSearch::FSearchIndexAsset& SearchIndexAsset);
+
+	// @todo: deprecate GetAnimationAssets
+	// UE_DEPRECATED(5.4, "Use GetNumAnimationAssets to iterate over GetDatabaseAnimationAsset instead")
 	const TArray<FInstancedStruct>& GetAnimationAssets() const { return AnimationAssets; }
+
+	UE_DEPRECATED(5.4, "Use GetDatabaseAnimationAsset instead")
 	const FInstancedStruct& GetAnimationAssetStruct(int32 AnimationAssetIndex) const;
+	UE_DEPRECATED(5.4, "Use GetDatabaseAnimationAsset instead")
 	const FInstancedStruct& GetAnimationAssetStruct(const UE::PoseSearch::FSearchIndexAsset& SearchIndexAsset) const;
+	UE_DEPRECATED(5.4, "Use GetMutableDatabaseAnimationAsset instead")
 	FInstancedStruct& GetMutableAnimationAssetStruct(int32 AnimationAssetIndex);
+	UE_DEPRECATED(5.4, "Use GetMutableDatabaseAnimationAsset instead")
 	FInstancedStruct& GetMutableAnimationAssetStruct(const UE::PoseSearch::FSearchIndexAsset& SearchIndexAsset);
-	const FPoseSearchDatabaseAnimationAssetBase* GetAnimationAssetBase(int32 AnimationAssetIndex) const;
-	const FPoseSearchDatabaseAnimationAssetBase* GetAnimationAssetBase(const UE::PoseSearch::FSearchIndexAsset& SearchIndexAsset) const;
-	FPoseSearchDatabaseAnimationAssetBase* GetMutableAnimationAssetBase(int32 AnimationAssetIndex);
-	FPoseSearchDatabaseAnimationAssetBase* GetMutableAnimationAssetBase(const UE::PoseSearch::FSearchIndexAsset& SearchIndexAsset);
+	
+	UE_DEPRECATED(5.4, "Use GetDatabaseAnimationAsset instead")
+	const FPoseSearchDatabaseAnimationAssetBase* GetAnimationAssetBase(int32 AnimationAssetIndex) const { return GetDatabaseAnimationAsset<FPoseSearchDatabaseAnimationAssetBase>(AnimationAssetIndex); }
+	UE_DEPRECATED(5.4, "Use GetDatabaseAnimationAsset instead")
+	const FPoseSearchDatabaseAnimationAssetBase* GetAnimationAssetBase(const UE::PoseSearch::FSearchIndexAsset& SearchIndexAsset) const { return GetDatabaseAnimationAsset<FPoseSearchDatabaseAnimationAssetBase>(SearchIndexAsset); }
+	UE_DEPRECATED(5.4, "Use GetMutableDatabaseAnimationAsset instead")
+	FPoseSearchDatabaseAnimationAssetBase* GetMutableAnimationAssetBase(int32 AnimationAssetIndex) { return GetMutableDatabaseAnimationAsset<FPoseSearchDatabaseAnimationAssetBase>(AnimationAssetIndex); }
+	UE_DEPRECATED(5.4, "Use GetMutableDatabaseAnimationAsset instead")
+	FPoseSearchDatabaseAnimationAssetBase* GetMutableAnimationAssetBase(const UE::PoseSearch::FSearchIndexAsset& SearchIndexAsset) { return GetMutableDatabaseAnimationAsset<FPoseSearchDatabaseAnimationAssetBase>(SearchIndexAsset); }
+
 	float GetRealAssetTime(int32 PoseIdx) const;
 	float GetNormalizedAssetTime(int32 PoseIdx) const;
 
@@ -451,3 +567,35 @@ private:
 	UE::PoseSearch::FSearchResult SearchVPTree(UE::PoseSearch::FSearchContext& SearchContext) const;
 	UE::PoseSearch::FSearchResult SearchBruteForce(UE::PoseSearch::FSearchContext& SearchContext) const;
 };
+
+template<typename TDatabaseAnimationAsset>
+inline const TDatabaseAnimationAsset* UPoseSearchDatabase::GetDatabaseAnimationAsset(int32 AnimationAssetIndex) const
+{
+	if (AnimationAssets.IsValidIndex(AnimationAssetIndex))
+	{
+		return AnimationAssets[AnimationAssetIndex].GetPtr<TDatabaseAnimationAsset>();
+	}
+	return nullptr;
+}
+
+template<typename TDatabaseAnimationAsset>
+inline const TDatabaseAnimationAsset* UPoseSearchDatabase::GetDatabaseAnimationAsset(const UE::PoseSearch::FSearchIndexAsset& SearchIndexAsset) const
+{
+	return GetDatabaseAnimationAsset<TDatabaseAnimationAsset>(SearchIndexAsset.GetSourceAssetIdx());
+}
+
+template<typename TDatabaseAnimationAsset>
+inline TDatabaseAnimationAsset* UPoseSearchDatabase::GetMutableDatabaseAnimationAsset(int32 AnimationAssetIndex)
+{
+	if (AnimationAssets.IsValidIndex(AnimationAssetIndex))
+	{
+		return AnimationAssets[AnimationAssetIndex].GetMutablePtr<TDatabaseAnimationAsset>();
+	}
+	return nullptr;
+}
+
+template<typename TDatabaseAnimationAsset>
+inline TDatabaseAnimationAsset* UPoseSearchDatabase::GetMutableDatabaseAnimationAsset(const UE::PoseSearch::FSearchIndexAsset& SearchIndexAsset)
+{
+	return GetMutableDatabaseAnimationAsset<TDatabaseAnimationAsset>(SearchIndexAsset.GetSourceAssetIdx());
+}
