motion matching - sampler determinism test

diff --git a/Engine/Plugins/Animation/PoseSearch/Source/Runtime/Private/PoseSearchDerivedData.cpp b/Engine/Plugins/Animation/PoseSearch/Source/Runtime/Private/PoseSearchDerivedData.cpp
--- a/Engine/Plugins/Animation/PoseSearch/Source/Runtime/Private/PoseSearchDerivedData.cpp
+++ b/Engine/Plugins/Animation/PoseSearch/Source/Runtime/Private/PoseSearchDerivedData.cpp
@@ -13,6 +13,9 @@
 #include "DerivedDataRequestOwner.h"
 #include "InstancedStruct.h"
 #include "Misc/CoreDelegates.h"
+#if ENABLE_ANIM_DEBUG
+#include "Misc/FileHelper.h"
+#endif //ENABLE_ANIM_DEBUG
 #include "PoseSearch/PoseSearchAnimNotifies.h"
 #include "PoseSearch/PoseSearchAssetIndexer.h"
 #include "PoseSearch/PoseSearchDatabase.h"
@@ -75,6 +78,12 @@ enum EMotionMatchTestFlags
 
 	// validating SynchronizeWithExternalDependencies doesn't alter the database AnimationAssets order
 	ValidateSynchronizeWithExternalDependenciesDeterminism = 1 << 11,
+
+	// test FAnimationAssetSampler determinism
+	TestAssetSamplerDeterminism = 1 << 12,
+
+	// test FAnimationAssetSampler determinism across multiple editro executions. It'll store some bin files in \Engine\TestAssetSamplerDeterminism
+	TestAssetSamplerDeterminismFromPreviousExecution = 1 << 13,
 };
 static TAutoConsoleVariable<int32> CVarMotionMatchTestFlags(TEXT("a.MotionMatch.TestFlags"), EMotionMatchTestFlags::None, TEXT("Test Motion Matching using EMotionMatchTestFlags"));
 static TAutoConsoleVariable<int32> CVarMotionMatchTestNumIterations(TEXT("a.MotionMatch.TestNumIterations"), 10, TEXT("Test Motion Matching Num Iterations"));
@@ -501,6 +510,28 @@ static void PreprocessSearchIndexWeights(FSearchIndex& SearchIndex, const UPoseS
 // it calculates Mean, PCAValues, and PCAProjectionMatrix
 static Eigen::ComputationInfo PreprocessSearchIndexPCAData(FSearchIndex& SearchIndex, int32 NumDimensions, int32 NumberOfPrincipalComponents, EPoseSearchMode PoseSearchMode)
 {
+#if ENABLE_ANIM_DEBUG
+	if (AnyTestFlags(EMotionMatchTestFlags::ValidateKDTreeConstruct))
+	{
+		// @todo: move this into a unit test.
+		// this code will fail with nanoflann 1.5.5
+
+		int NumPoses = 61;
+		int DataCardinality = 8;
+		TArray<float> Values;
+		Values.SetNumZeroed(NumPoses * DataCardinality);
+
+		for (int PoseIndex = 0; PoseIndex < NumPoses; ++PoseIndex)
+		{
+			Values[PoseIndex * DataCardinality + 0] = -5.54383405e-07;
+			Values[PoseIndex * DataCardinality + 1] = 2.77555756e-16;
+		}
+
+		FKDTree KDTree(NumPoses, DataCardinality, Values.GetData());
+	}
+#endif // ENABLE_ANIM_DEBUG
+
+
 	// binding SearchIndex.Values and SearchIndex.PCAValues Eigen row major matrix maps
 	const int32 NumPoses = SearchIndex.GetNumPoses();
 
@@ -668,14 +699,14 @@ static void PreprocessSearchIndexKDTree(FSearchIndex& SearchIndex, const UPoseSe
 
 			TArray<int32> ResultIndexes;
 			TArray<float> ResultDistanceSqr;
-			ResultIndexes.SetNum(NumPCAValuesVectors + 1);
-			ResultDistanceSqr.SetNum(NumPCAValuesVectors + 1);
+			ResultIndexes.SetNum(KDTreeQueryNumNeighbors + 1);
+			ResultDistanceSqr.SetNum(KDTreeQueryNumNeighbors + 1);
 
 			int32 MaxNumNeighborToFindAPoint = 0;
 			for (int32 PointIndex = 0; PointIndex < NumPCAValuesVectors; ++PointIndex)
 			{
 				// searching the kdtree for PointIndex
-				FKDTree::FRadiusResultSet ResultSet(UE_SMALL_NUMBER, NumPCAValuesVectors, ResultIndexes, ResultDistanceSqr);
+				FKDTree::FRadiusResultSet ResultSet(UE_SMALL_NUMBER, KDTreeQueryNumNeighbors, ResultIndexes, ResultDistanceSqr);
 				SearchIndex.KDTree.FindNeighbors(ResultSet, MakeArrayView(&SearchIndex.PCAValues[PointIndex * NumberOfPrincipalComponents], NumberOfPrincipalComponents));
 
 				bool bFound = false;
@@ -1942,28 +1973,28 @@ void FAsyncPoseSearchDatabasesManagement::SynchronizeDatabases()
 
 	if (!DatabasesToSynchronize.IsEmpty())
 	{
-	// copying DatabasesToSynchronize because modifying the database will call OnObjectModified that could populate DatabasesToSynchronize again
-	const TDatabasesToSynchronize DatabasesToSynchronizeCopy = DatabasesToSynchronize;
-	DatabasesToSynchronize.Reset();
+		// copying DatabasesToSynchronize because modifying the database will call OnObjectModified that could populate DatabasesToSynchronize again
+		const TDatabasesToSynchronize DatabasesToSynchronizeCopy = DatabasesToSynchronize;
+		DatabasesToSynchronize.Reset();
 
-	for (const TDatabasesToSynchronizePair& Pair : DatabasesToSynchronizeCopy)
-	{
-		if (Pair.Key.IsValid())
+		for (const TDatabasesToSynchronizePair& Pair : DatabasesToSynchronizeCopy)
 		{
-			TArray<UAnimSequenceBase*> SequencesBase;
-			for (const TWeakObjectPtr<UAnimSequenceBase>& SequenceBase : Pair.Value)
+			if (Pair.Key.IsValid())
 			{
-				if (SequenceBase.IsValid())
+				TArray<UAnimSequenceBase*> SequencesBase;
+				for (const TWeakObjectPtr<UAnimSequenceBase>& SequenceBase : Pair.Value)
 				{
-					SequencesBase.Add(SequenceBase.Get());
+					if (SequenceBase.IsValid())
+					{
+						SequencesBase.Add(SequenceBase.Get());
+					}
 				}
-			}
 
-			Pair.Key->SynchronizeWithExternalDependencies(SequencesBase);
+				Pair.Key->SynchronizeWithExternalDependencies(SequencesBase);
+			}
 		}
 	}
 }
-}
 
 // we're listening to OnObjectModified to cancel any pending Task indexing databases depending from Object to avoid multi threading issues
 void FAsyncPoseSearchDatabasesManagement::OnObjectModified(UObject* Object)
@@ -2080,6 +2111,177 @@ void FAsyncPoseSearchDatabasesManagement::Tick(float DeltaTime)
 	check(IsInGameThread());
 
 #if ENABLE_ANIM_DEBUG
+	// testing sampler determinism
+	const bool bTestAssetSamplerDeterminism = AnyTestFlags(EMotionMatchTestFlags::TestAssetSamplerDeterminism);
+	const bool bTestAssetSamplerDeterminismFromPreviousExecution = AnyTestFlags(EMotionMatchTestFlags::TestAssetSamplerDeterminismFromPreviousExecution);
+	if (bTestAssetSamplerDeterminism || bTestAssetSamplerDeterminismFromPreviousExecution)
+	{
+		const int32 NumIterations = CVarMotionMatchTestNumIterations.GetValueOnAnyThread();
+
+		struct FTestSample
+		{
+			FBoneContainer BoneContainer;
+			const UPoseSearchDatabase* Database = nullptr;
+			int32 AnimationAssetIndex = INDEX_NONE;
+			const UAnimationAsset* AnimationAsset = nullptr;
+			int32 SampleIndex = INDEX_NONE;
+			float SampleTime = 0.f;
+
+			TAlignedArray<uint8> SerializedData;
+
+			FString GetFileName() const
+			{
+				return FString::Printf(TEXT("%s//TestAssetSamplerDeterminism//%s_%d_%s_%d.bin"), *FPaths::EngineDir(), *GetNameSafe(Database), AnimationAssetIndex, *GetNameSafe(AnimationAsset), SampleIndex);
+			}
+		};
+
+		TArray<FTestSample> TestSamples;
+		for (const TUniquePtr<FPoseSearchDatabaseAsyncCacheTask>& TaskPtr : Tasks)
+		{
+			if (const UPoseSearchDatabase* Database = TaskPtr->GetDatabase())
+			{
+				if (Database->Schema)
+				{
+					TMap<FRole, FBoneContainer> RoledBoneContainers;
+					Database->Schema->InitBoneContainersFromRoledSkeleton(RoledBoneContainers);
+
+					for (int32 AnimationAssetIndex = 0; AnimationAssetIndex < Database->GetAnimationAssets().Num(); ++AnimationAssetIndex)
+					{
+						if (const FPoseSearchDatabaseAnimationAssetBase* DatabaseAsset = Database->GetAnimationAssets()[AnimationAssetIndex].GetPtr<FPoseSearchDatabaseAnimationAssetBase>())
+						{
+							for (int32 RoleIndex = 0; RoleIndex < DatabaseAsset->GetNumRoles(); ++RoleIndex)
+							{
+								const FRole Role = DatabaseAsset->GetRole(RoleIndex);
+								if (const UAnimationAsset* AnimationAsset = DatabaseAsset->GetAnimationAssetForRole(Role))
+								{
+									if (const FBoneContainer* BoneContainer = RoledBoneContainers.Find(Role))
+									{
+										static float SAMPLING_TIME = 1 / 30.f;
+										const float PlayLength = AnimationAsset->GetPlayLength();
+										const int32 NumSamples = FMath::FloorToInt(PlayLength / SAMPLING_TIME);
+
+										for (int32 SampleIndex = 0; SampleIndex < NumSamples; ++SampleIndex)
+										{
+											FTestSample& TestSample = TestSamples.AddDefaulted_GetRef();
+											TestSample.BoneContainer = *BoneContainer;
+											TestSample.Database = Database;
+											TestSample.AnimationAssetIndex = AnimationAssetIndex;
+											TestSample.AnimationAsset = AnimationAsset;
+											TestSample.SampleIndex = SampleIndex;
+											TestSample.SampleTime = SampleIndex * SAMPLING_TIME;
+										}
+									}
+								}
+							}
+						}
+					}
+				}
+			}
+		}
+
+		ParallelFor(TestSamples.Num(), [&TestSamples, NumIterations, bTestAssetSamplerDeterminism, bTestAssetSamplerDeterminismFromPreviousExecution](int32 TestSampleIndex)
+			{
+				FMemMark Mark(FMemStack::Get());
+
+				FTestSample& TestSample = TestSamples[TestSampleIndex];
+				const FAnimationAssetSampler AssetSampler(TestSample.AnimationAsset);
+
+				FCompactPose Pose;
+				Pose.SetBoneContainer(&TestSample.BoneContainer);
+				AssetSampler.ExtractPose(TestSample.SampleTime, Pose);
+				const FTransform RootTransform = AssetSampler.ExtractRootTransform(TestSample.SampleTime);
+
+				const TConstArrayView<FTransform> Bones = Pose.GetBones();
+
+				if (bTestAssetSamplerDeterminism)
+				{
+					for (int32 IterationIndex = 0; IterationIndex < NumIterations; ++IterationIndex)
+					{
+						FCompactPose TestPose;
+						TestPose.SetBoneContainer(&TestSample.BoneContainer);
+						AssetSampler.ExtractPose(TestSample.SampleTime, TestPose);
+						const FTransform TestRootTransform = AssetSampler.ExtractRootTransform(TestSample.SampleTime);
+
+						if (FMemory::Memcmp(&RootTransform, &TestRootTransform, sizeof(FTransform)) != 0)
+						{
+							UE_LOG(LogPoseSearch, Error, TEXT("FAnimationAssetSampler - ExtractRootTransform is not deterministic"));
+						}
+
+						const TConstArrayView<FTransform> TestBones = TestPose.GetBones();
+						if (Bones.Num() != TestBones.Num())
+						{
+							UE_LOG(LogPoseSearch, Error, TEXT("FAnimationAssetSampler - ExtractPose is not deterministic"));
+						}
+						else
+						{
+							for (int32 BoneIndex = 0; BoneIndex < Bones.Num(); ++BoneIndex)
+							{
+								if (FMemory::Memcmp(&Bones[BoneIndex], &TestBones[BoneIndex], sizeof(FTransform)) != 0)
+								{
+									UE_LOG(LogPoseSearch, Error, TEXT("FAnimationAssetSampler - ExtractPose is not deterministic"));
+								}
+							}
+						}
+					}
+				}
+
+				if (bTestAssetSamplerDeterminismFromPreviousExecution)
+				{
+					const int32 RootTransformSize = sizeof(FTransform);
+					const int32 BonesSize = sizeof(FTransform) * Bones.Num();
+
+					TestSample.SerializedData.Reserve(RootTransformSize + BonesSize);
+					TestSample.SerializedData.Append(reinterpret_cast<const uint8*>(&RootTransform), RootTransformSize);
+					TestSample.SerializedData.Append(reinterpret_cast<const uint8*>(Bones.GetData()), BonesSize);
+				}
+			}, ParallelForFlags);
+
+		if (bTestAssetSamplerDeterminismFromPreviousExecution)
+		{
+			for (const FTestSample& TestSample : TestSamples)
+			{
+				FString FileName = TestSample.GetFileName();
+				TArray<uint8> LoadedData;
+				const bool bLoadedFile = FFileHelper::LoadFileToArray(LoadedData, *FileName, FILEREAD_Silent);
+				if (bLoadedFile)
+				{
+					if (LoadedData.Num() <= 0 || LoadedData.Num() % sizeof(FTransform) != 0)
+					{
+						UE_LOG(LogPoseSearch, Error, TEXT("FAnimationAssetSampler - Loaded the wrong amount of data!"));
+					}
+					else if (TestSample.SerializedData.Num() != LoadedData.Num() || FMemory::Memcmp(TestSample.SerializedData.GetData(), LoadedData.GetData(), TestSample.SerializedData.Num()) != 0)
+					{
+						const int32 NumTransforms = LoadedData.Num() / sizeof(FTransform);
+
+						// copying the data into an aligned buffer, so we can cast it to FTransform
+						TAlignedArray<uint8> LoadedDataAligned(LoadedData.GetData(), LoadedData.Num());
+						TArrayView<const FTransform> LoadedTransforms(reinterpret_cast<const FTransform*>(LoadedDataAligned.GetData()), NumTransforms);
+						TArrayView<const FTransform> SerializedTransforms(reinterpret_cast<const FTransform*>(TestSample.SerializedData.GetData()), NumTransforms);
+
+						for (int32 TransformIndex = 0; TransformIndex < NumTransforms; ++TransformIndex)
+						{
+							if (FMemory::Memcmp(&LoadedTransforms[TransformIndex], &SerializedTransforms[TransformIndex], sizeof(FTransform)) != 0)
+							{
+								// NoTe:	TransformIndex == 0 for the root
+								//			TransformIndex > 0 are the bones, where BoneIndex = TransformIndex - 1
+								const int32 BoneIndex = TransformIndex - 1;
+								UE_LOG(LogPoseSearch, Error, TEXT("FAnimationAssetSampler - ExtractPose is not deterministic for %s for Bone %d"), *FileName, BoneIndex);
+							}
+						}
+					}
+				}
+				else
+				{
+					bool bSavedFile = FFileHelper::SaveArrayToFile(TestSample.SerializedData, *FileName);
+					if (!bSavedFile)
+					{
+						UE_LOG(LogPoseSearch, Error, TEXT("FAnimationAssetSampler - Failed to save comparison file!"));
+					}
+				}
+			}
+		}
+	}
+
 	if (AnyTestFlags(EMotionMatchTestFlags::InvalidateCache))
 	{
 		if (AnyTestFlags(EMotionMatchTestFlags::WaitForTaskCompletion))
