motion matching - multi character improvements

diff --git a/Engine/Plugins/Animation/PoseSearch/Source/Editor/Private/PoseSearchDatabaseAssetBrowser.cpp b/Engine/Plugins/Animation/PoseSearch/Source/Editor/Private/PoseSearchDatabaseAssetBrowser.cpp
--- a/Engine/Plugins/Animation/PoseSearch/Source/Editor/Private/PoseSearchDatabaseAssetBrowser.cpp
+++ b/Engine/Plugins/Animation/PoseSearch/Source/Editor/Private/PoseSearchDatabaseAssetBrowser.cpp
@@ -46,7 +46,7 @@ void SPoseSearchDatabaseAssetBrowser::Construct(const FArguments& InArgs, TShare
 
 void SPoseSearchDatabaseAssetBrowser::RefreshView()
 {
-	// @TODO: Add support for MultiSequence.
+	// @TODO: Add support for MultiAnimAsset.
 	
 	FAssetPickerConfig AssetPickerConfig;
 	
diff --git a/Engine/Plugins/Animation/PoseSearch/Source/Editor/Private/PoseSearchDatabaseAssetTree.cpp b/Engine/Plugins/Animation/PoseSearch/Source/Editor/Private/PoseSearchDatabaseAssetTree.cpp
--- a/Engine/Plugins/Animation/PoseSearch/Source/Editor/Private/PoseSearchDatabaseAssetTree.cpp
+++ b/Engine/Plugins/Animation/PoseSearch/Source/Editor/Private/PoseSearchDatabaseAssetTree.cpp
@@ -9,20 +9,20 @@
 #include "AssetSelection.h"
 #include "ClassIconFinder.h"
 #include "DetailColumnSizeData.h"
-#include "PoseSearchDatabaseEditorClipboard.h"
 #include "DragAndDrop/AssetDragDropOp.h"
 #include "Framework/Commands/GenericCommands.h"
 #include "Framework/MultiBox/MultiBoxBuilder.h"
 #include "Misc/FeedbackContext.h"
 #include "Misc/TransactionObjectEvent.h"
+#include "PoseSearch/MultiAnimAsset.h"
 #include "PoseSearch/PoseSearchAnimNotifies.h"
 #include "PoseSearch/PoseSearchDatabase.h"
+#include "PoseSearchDatabaseEditorClipboard.h"
 #include "PoseSearchDatabaseViewModel.h"
 #include "PoseSearchEditor.h"
 #include "SPositiveActionButton.h"
 #include "Styling/AppStyle.h"
 #include "ScopedTransaction.h"
-#include "PoseSearch/PoseSearchMultiSequence.h"
 #include "Styling/StyleColors.h"
 #include "Widgets/Text/SRichTextBlock.h"
 #include "Widgets/Input/SCheckBox.h"
@@ -140,7 +140,7 @@ namespace UE::PoseSearch
 								AssetClass->IsChildOf(UAnimComposite::StaticClass()) ||
 								AssetClass->IsChildOf(UBlendSpace::StaticClass()) ||
 								AssetClass->IsChildOf(UAnimMontage::StaticClass()) ||
-								AssetClass->IsChildOf(UPoseSearchMultiSequence::StaticClass()))
+								AssetClass->IsChildOf(UMultiAnimAsset::StaticClass()))
 							{
 								Reply = FReply::Handled();
 								break;
@@ -247,7 +247,16 @@ namespace UE::PoseSearch
 			IndexArray.Sort([Database](int32 SequenceIdxA, int32 SequenceIdxB)
 			{
 				const FPoseSearchDatabaseAnimationAssetBase* A = Database->GetDatabaseAnimationAsset<FPoseSearchDatabaseAnimationAssetBase>(SequenceIdxA);
+				if (!A)
+				{
+					return false;
+				}
+
 				const FPoseSearchDatabaseAnimationAssetBase* B = Database->GetDatabaseAnimationAsset<FPoseSearchDatabaseAnimationAssetBase>(SequenceIdxB);
+				if (!B)
+				{
+					return true;
+				}
 
 				//If its null add it to the end of the list 
 				if (!B->GetAnimationAsset())
@@ -344,7 +353,7 @@ namespace UE::PoseSearch
 						AssetClass->IsChildOf(UAnimComposite::StaticClass()) ||
 						AssetClass->IsChildOf(UBlendSpace::StaticClass()) ||
 						AssetClass->IsChildOf(UAnimMontage::StaticClass()) ||
-						AssetClass->IsChildOf(UPoseSearchMultiSequence::StaticClass()))
+						AssetClass->IsChildOf(UMultiAnimAsset::StaticClass()))
 					{
 						ReturnedDropZone = EItemDropZone::OntoItem;
 						break;
@@ -424,9 +433,9 @@ namespace UE::PoseSearch
 					ViewModel->AddAnimMontageToDatabase(Cast<UAnimMontage>(Asset));
 					++AddedAssets;
 				}
-				else if (AssetClass->IsChildOf(UPoseSearchMultiSequence::StaticClass()))
+				else if (AssetClass->IsChildOf(UMultiAnimAsset::StaticClass()))
 				{
-					ViewModel->AddMultiSequenceToDatabase(Cast<UPoseSearchMultiSequence>(Asset));
+					ViewModel->AddMultiAnimAssetToDatabase(Cast<UMultiAnimAsset>(Asset));
 					++AddedAssets;
 				}
 			}
@@ -482,10 +491,10 @@ namespace UE::PoseSearch
 				EUserInterfaceActionType::Button);
 
 			AddOptions.AddMenuEntry(
-				LOCTEXT("AddMultiSequenceOption", "Multi Sequence"),
-				LOCTEXT("AddMultiSequenceToDatabaseTooltip", "Add new multi sequence to the database"),
+				LOCTEXT("AddMultiAnimAssetOption", "Multi Anim Asset"),
+				LOCTEXT("AddMultiAnimAssetToDatabaseTooltip", "Add new multi anim asset to the database"),
 				FSlateIcon(),
-				FUIAction(FExecuteAction::CreateSP(this, &SDatabaseAssetTree::OnAddMultiSequence, true)),
+				FUIAction(FExecuteAction::CreateSP(this, &SDatabaseAssetTree::OnAddMultiAnimAsset, true)),
 				NAME_None,
 				EUserInterfaceActionType::Button);
 		}
@@ -757,11 +766,11 @@ namespace UE::PoseSearch
 		}
 	}
 
-	void SDatabaseAssetTree::OnAddMultiSequence(bool bFinalizeChanges)
+	void SDatabaseAssetTree::OnAddMultiAnimAsset(bool bFinalizeChanges)
 	{
-		FScopedTransaction Transaction(LOCTEXT("AddMultiSequenceTransaction", "Add Multi Sequence"));
+		FScopedTransaction Transaction(LOCTEXT("AddMultiAnimAssetTransaction", "Add Multi Anim Asset"));
 
-		EditorViewModel.Pin()->AddMultiSequenceToDatabase(nullptr);
+		EditorViewModel.Pin()->AddMultiAnimAssetToDatabase(nullptr);
 
 		if (bFinalizeChanges)
 		{
diff --git a/Engine/Plugins/Animation/PoseSearch/Source/Editor/Private/PoseSearchDatabaseEditor.cpp b/Engine/Plugins/Animation/PoseSearch/Source/Editor/Private/PoseSearchDatabaseEditor.cpp
--- a/Engine/Plugins/Animation/PoseSearch/Source/Editor/Private/PoseSearchDatabaseEditor.cpp
+++ b/Engine/Plugins/Animation/PoseSearch/Source/Editor/Private/PoseSearchDatabaseEditor.cpp
@@ -615,13 +615,13 @@ namespace UE::PoseSearch
 
 					SelectionWidget.SelectedReflections.Add(NewSelectionReflection);
 				}
-				else if (const FPoseSearchDatabaseMultiSequence* DatabaseMultiSequence = DatabaseAsset.GetPtr<FPoseSearchDatabaseMultiSequence>())
+				else if (const FPoseSearchDatabaseMultiAnimAsset* DatabaseMultiAnimAsset = DatabaseAsset.GetPtr<FPoseSearchDatabaseMultiAnimAsset>())
 				{
-					UPoseSearchDatabaseMultiSequenceReflection* NewSelectionReflection = NewObject<UPoseSearchDatabaseMultiSequenceReflection>();
+					UPoseSearchDatabaseMultiAnimAssetReflection* NewSelectionReflection = NewObject<UPoseSearchDatabaseMultiAnimAssetReflection>();
 					NewSelectionReflection->AddToRoot();
-					static_cast<FPoseSearchDatabaseMultiSequence&>(NewSelectionReflection->MultiSequence) = *DatabaseMultiSequence;
-					NewSelectionReflection->MultiSequence.bLooping = DatabaseMultiSequence->IsLooping();
-					NewSelectionReflection->MultiSequence.bHasRootMotion = DatabaseMultiSequence->IsRootMotionEnabled();
+					static_cast<FPoseSearchDatabaseMultiAnimAsset&>(NewSelectionReflection->MultiAnimAsset) = *DatabaseMultiAnimAsset;
+					NewSelectionReflection->MultiAnimAsset.bLooping = DatabaseMultiAnimAsset->IsLooping();
+					NewSelectionReflection->MultiAnimAsset.bHasRootMotion = DatabaseMultiAnimAsset->IsRootMotionEnabled();
 					NewSelectionReflection->SetSourceLink(SelectedItem, AssetTreeWidget);
 					NewSelectionReflection->SetFlags(RF_Transactional);
 
diff --git a/Engine/Plugins/Animation/PoseSearch/Source/Editor/Private/PoseSearchDatabaseEditorClipboard.cpp b/Engine/Plugins/Animation/PoseSearch/Source/Editor/Private/PoseSearchDatabaseEditorClipboard.cpp
--- a/Engine/Plugins/Animation/PoseSearch/Source/Editor/Private/PoseSearchDatabaseEditorClipboard.cpp
+++ b/Engine/Plugins/Animation/PoseSearch/Source/Editor/Private/PoseSearchDatabaseEditorClipboard.cpp
@@ -1,4 +1,4 @@
-﻿// Copyright Epic Games, Inc. All Rights Reserved.
+// Copyright Epic Games, Inc. All Rights Reserved.
 
 #include "PoseSearchDatabaseEditorClipboard.h"
 
@@ -8,7 +8,7 @@
 #include "Animation/AnimSequence.h"
 #include "Animation/AnimMontage.h"
 #include "Animation/BlendSpace.h"
-#include "PoseSearch/PoseSearchMultiSequence.h"
+#include "PoseSearch/MultiAnimAsset.h"
 #include "PoseSearch/PoseSearchDatabase.h"
 
 #include "UnrealExporter.h"
@@ -39,9 +39,9 @@ void UPoseSearchDatabaseEditorClipboardContent::CopyDatabaseItem(const FPoseSear
 	{
 		FPoseSearchDatabaseBlendSpace::StaticStruct()->ExportText(CopyObj->Content, InItem, nullptr, nullptr, PPF_None, nullptr);
 	}
-	else if (UPoseSearchMultiSequence::StaticClass() == InItem->GetAnimationAssetStaticClass())
+	else if (UMultiAnimAsset::StaticClass() == InItem->GetAnimationAssetStaticClass())
 	{
-		FPoseSearchDatabaseMultiSequence::StaticStruct()->ExportText(CopyObj->Content, InItem, nullptr, nullptr, PPF_None, nullptr);
+		FPoseSearchDatabaseMultiAnimAsset::StaticStruct()->ExportText(CopyObj->Content, InItem, nullptr, nullptr, PPF_None, nullptr);
 	}
 
 	DatabaseItems.Add(CopyObj);
@@ -95,10 +95,10 @@ void UPoseSearchDatabaseEditorClipboardContent::PasteToDatabase(UPoseSearchDatab
 
 			InTargetDatabase->AddAnimationAsset(FInstancedStruct::Make(NewAsset));
 		}
-		else if (Item->ClassName == UPoseSearchMultiSequence::StaticClass()->GetFName())
+		else if (Item->ClassName == UMultiAnimAsset::StaticClass()->GetFName())
 		{
-			FPoseSearchDatabaseMultiSequence NewAsset;
-			FPoseSearchDatabaseMultiSequence::StaticStruct()->ImportText(*Item->Content, &NewAsset, nullptr, PPF_None, GLog, FPoseSearchDatabaseMultiSequence::StaticStruct()->GetName());
+			FPoseSearchDatabaseMultiAnimAsset NewAsset;
+			FPoseSearchDatabaseMultiAnimAsset::StaticStruct()->ImportText(*Item->Content, &NewAsset, nullptr, PPF_None, GLog, FPoseSearchDatabaseMultiAnimAsset::StaticStruct()->GetName());
 				
 			InTargetDatabase->AddAnimationAsset(FInstancedStruct::Make(NewAsset));
 		}
diff --git a/Engine/Plugins/Animation/PoseSearch/Source/Editor/Private/PoseSearchDatabaseEditorReflection.cpp b/Engine/Plugins/Animation/PoseSearch/Source/Editor/Private/PoseSearchDatabaseEditorReflection.cpp
--- a/Engine/Plugins/Animation/PoseSearch/Source/Editor/Private/PoseSearchDatabaseEditorReflection.cpp
+++ b/Engine/Plugins/Animation/PoseSearch/Source/Editor/Private/PoseSearchDatabaseEditorReflection.cpp
@@ -157,13 +157,13 @@ bool UPoseSearchDatabaseAnimMontageReflection::ApplyChanges()
 	return false;
 }
 
-bool UPoseSearchDatabaseMultiSequenceReflection::ApplyChanges()
+bool UPoseSearchDatabaseMultiAnimAssetReflection::ApplyChanges()
 {
-	const float PlayLength = MultiSequence.GetPlayLength();
-	const FFloatInterval ClampedSamplingRange(FMath::Clamp(MultiSequence.SamplingRange.Min, 0.0f, PlayLength), FMath::Clamp(MultiSequence.SamplingRange.Max, 0.0f, PlayLength));
-	if (ClampedSamplingRange != MultiSequence.SamplingRange)
+	const float PlayLength = MultiAnimAsset.GetPlayLength();
+	const FFloatInterval ClampedSamplingRange(FMath::Clamp(MultiAnimAsset.SamplingRange.Min, 0.0f, PlayLength), FMath::Clamp(MultiAnimAsset.SamplingRange.Max, 0.0f, PlayLength));
+	if (ClampedSamplingRange != MultiAnimAsset.SamplingRange)
 	{
-		MultiSequence.SamplingRange = ClampedSamplingRange;
+		MultiAnimAsset.SamplingRange = ClampedSamplingRange;
 	}
 
 	if (const TSharedPtr<UE::PoseSearch::FDatabaseAssetTreeNode> AssetTreeNode = WeakAssetTreeNode.Pin())
@@ -173,11 +173,11 @@ bool UPoseSearchDatabaseMultiSequenceReflection::ApplyChanges()
 			UPoseSearchDatabase* Database = ViewModel->GetPoseSearchDatabase();
 			if (IsValid(Database))
 			{
-				if (FPoseSearchDatabaseMultiSequence* DatabaseMultiSequence = Database->GetMutableDatabaseAnimationAsset<FPoseSearchDatabaseMultiSequence>(AssetTreeNode->SourceAssetIdx))
+				if (FPoseSearchDatabaseMultiAnimAsset* DatabaseMultiAnimAsset = Database->GetMutableDatabaseAnimationAsset<FPoseSearchDatabaseMultiAnimAsset>(AssetTreeNode->SourceAssetIdx))
 				{
-					if (*DatabaseMultiSequence != MultiSequence)
+					if (*DatabaseMultiAnimAsset != MultiAnimAsset)
 					{
-						*DatabaseMultiSequence = MultiSequence;
+						*DatabaseMultiAnimAsset = MultiAnimAsset;
 						Database->MarkPackageDirty();
 						return true;
 					}
diff --git a/Engine/Plugins/Animation/PoseSearch/Source/Editor/Private/PoseSearchDatabaseViewModel.cpp b/Engine/Plugins/Animation/PoseSearch/Source/Editor/Private/PoseSearchDatabaseViewModel.cpp
--- a/Engine/Plugins/Animation/PoseSearch/Source/Editor/Private/PoseSearchDatabaseViewModel.cpp
+++ b/Engine/Plugins/Animation/PoseSearch/Source/Editor/Private/PoseSearchDatabaseViewModel.cpp
@@ -14,6 +14,7 @@
 #include "GameFramework/PlayerController.h"
 #include "InstancedStruct.h"
 #include "Modules/ModuleManager.h"
+#include "PoseSearch/MultiAnimAsset.h"
 #include "PoseSearch/PoseSearchAnimNotifies.h"
 #include "PoseSearch/PoseSearchContext.h"
 #include "PoseSearch/PoseSearchDatabase.h"
@@ -27,7 +28,6 @@
 #include "PoseSearchDatabasePreviewScene.h"
 #include "PoseSearchEditor.h"
 #include "PropertyEditorModule.h"
-#include "PoseSearch/PoseSearchMultiSequence.h"
 
 namespace UE::PoseSearch
 {
@@ -576,14 +576,14 @@ void FDatabaseViewModel::AddAnimMontageToDatabase(UAnimMontage* AnimMontage)
 	}
 }
 
-void FDatabaseViewModel::AddMultiSequenceToDatabase(UPoseSearchMultiSequence* MultiSequence)
+void FDatabaseViewModel::AddMultiAnimAssetToDatabase(UMultiAnimAsset* MultiAnimAsset)
 {
 	if (UPoseSearchDatabase* Database = GetPoseSearchDatabase())
 	{
 		Database->Modify();
 		
-		FPoseSearchDatabaseMultiSequence NewAsset;
-		NewAsset.MultiSequence = MultiSequence;
+		FPoseSearchDatabaseMultiAnimAsset NewAsset;
+		NewAsset.MultiAnimAsset = MultiAnimAsset;
 		Database->AddAnimationAsset(FInstancedStruct::Make(NewAsset));
 	}
 }
@@ -725,10 +725,10 @@ bool FDatabaseViewModel::SetAnimationAsset(int32 AnimationAssetIndex, UObject* A
 						FPoseSearchDatabaseBlendSpace* DatabaseBlendSpaceAsset = static_cast<FPoseSearchDatabaseBlendSpace*>(DatabaseAnimationAsset);
 						DatabaseBlendSpaceAsset->BlendSpace = Cast<UBlendSpace>(AnimAsset);
 					}
-					else if (AssetClass->IsChildOf(UPoseSearchMultiSequence::StaticClass()))
+					else if (AssetClass->IsChildOf(UMultiAnimAsset::StaticClass()))
 					{
-						FPoseSearchDatabaseMultiSequence* DatabaseMultiSequenceAsset = static_cast<FPoseSearchDatabaseMultiSequence*>(DatabaseAnimationAsset);
-						DatabaseMultiSequenceAsset->MultiSequence = Cast<UPoseSearchMultiSequence>(AnimAsset);
+						FPoseSearchDatabaseMultiAnimAsset* DatabaseMultiAnimAssetAsset = static_cast<FPoseSearchDatabaseMultiAnimAsset*>(DatabaseAnimationAsset);
+						DatabaseMultiAnimAssetAsset->MultiAnimAsset = Cast<UMultiAnimAsset>(AnimAsset);
 					}
 
 					return true;
@@ -860,7 +860,7 @@ void FDatabaseViewModel::SetSelectedNodes(const TArrayView<TSharedPtr<FDatabaseA
 					{
 						FDatabasePreviewActor PreviewActor;
 						const UE::PoseSearch::FRole Role = DatabaseAnimationAsset->GetRole(RoleIndex);
-						const FTransform& RootTransformOrigin = DatabaseAnimationAsset->GetRootTransformOriginForRole(Role);
+						const FTransform RootTransformOrigin = DatabaseAnimationAsset->GetRootTransformOriginForRole(Role);
 						if (PreviewActor.SpawnPreviewActor(GetWorld(), Database, IndexAssetIndex, Role, RootTransformOrigin))
 						{
 							if (PreviewActorGroupIndex == INDEX_NONE)
diff --git a/Engine/Plugins/Animation/PoseSearch/Source/Editor/Public/PoseSearchDatabaseAssetTree.h b/Engine/Plugins/Animation/PoseSearch/Source/Editor/Public/PoseSearchDatabaseAssetTree.h
--- a/Engine/Plugins/Animation/PoseSearch/Source/Editor/Public/PoseSearchDatabaseAssetTree.h
+++ b/Engine/Plugins/Animation/PoseSearch/Source/Editor/Public/PoseSearchDatabaseAssetTree.h
@@ -88,7 +88,7 @@ namespace UE::PoseSearch
 		void OnAddBlendSpace(bool bFinalizeChanges = true);
 		void OnAddAnimComposite(bool bFinalizeChanges = true);
 		void OnAddAnimMontage(bool bFinalizeChanges = true);
-		void OnAddMultiSequence(bool bFinalizeChanges = true);
+		void OnAddMultiAnimAsset(bool bFinalizeChanges = true);
 
 		void OnDeleteAsset(TSharedPtr<FDatabaseAssetTreeNode> Node, bool bFinalizeChanges = true);
 		void CreateCommandList();
diff --git a/Engine/Plugins/Animation/PoseSearch/Source/Editor/Public/PoseSearchDatabaseEditorReflection.h b/Engine/Plugins/Animation/PoseSearch/Source/Editor/Public/PoseSearchDatabaseEditorReflection.h
--- a/Engine/Plugins/Animation/PoseSearch/Source/Editor/Public/PoseSearchDatabaseEditorReflection.h
+++ b/Engine/Plugins/Animation/PoseSearch/Source/Editor/Public/PoseSearchDatabaseEditorReflection.h
@@ -141,7 +141,7 @@ public:
 };
 
 USTRUCT()
-struct FPoseSearchDatabaseMultiSequenceEx : public FPoseSearchDatabaseMultiSequence
+struct FPoseSearchDatabaseMultiAnimAssetEx : public FPoseSearchDatabaseMultiAnimAsset
 {
 	GENERATED_BODY()
 
@@ -156,13 +156,13 @@ struct FPoseSearchDatabaseMultiSequenceEx : public FPoseSearchDatabaseMultiSeque
 };
 
 UCLASS()
-class UPoseSearchDatabaseMultiSequenceReflection : public UPoseSearchDatabaseReflectionBase
+class UPoseSearchDatabaseMultiAnimAssetReflection : public UPoseSearchDatabaseReflectionBase
 {
 	GENERATED_BODY()
 
 public:
 	UPROPERTY(EditAnywhere, Category = "Selected Sequence")
-	FPoseSearchDatabaseMultiSequenceEx MultiSequence;
+	FPoseSearchDatabaseMultiAnimAssetEx MultiAnimAsset;
 	
 	virtual bool ApplyChanges() override;
 };
diff --git a/Engine/Plugins/Animation/PoseSearch/Source/Editor/Public/PoseSearchDatabaseViewModel.h b/Engine/Plugins/Animation/PoseSearch/Source/Editor/Public/PoseSearchDatabaseViewModel.h
--- a/Engine/Plugins/Animation/PoseSearch/Source/Editor/Public/PoseSearchDatabaseViewModel.h
+++ b/Engine/Plugins/Animation/PoseSearch/Source/Editor/Public/PoseSearchDatabaseViewModel.h
@@ -21,7 +21,7 @@ class UAnimComposite;
 class UAnimSequence;
 class UBlendSpace;
 class UMirrorDataTable;
-class UPoseSearchMultiSequence;
+class UMultiAnimAsset;
 
 namespace UE::PoseSearch
 {
@@ -109,7 +109,7 @@ namespace UE::PoseSearch
 		void AddBlendSpaceToDatabase(UBlendSpace* BlendSpace);
 		void AddAnimCompositeToDatabase(UAnimComposite* AnimComposite);
 		void AddAnimMontageToDatabase(UAnimMontage* AnimMontage);
-		void AddMultiSequenceToDatabase(UPoseSearchMultiSequence* MultiSequence);
+		void AddMultiAnimAssetToDatabase(UMultiAnimAsset* MultiAnimAsset);
 
 		bool DeleteFromDatabase(int32 AnimationAssetIndex);
 
diff --git a/Engine/Plugins/Animation/PoseSearch/Source/Runtime/Private/PoseSearchContext.cpp b/Engine/Plugins/Animation/PoseSearch/Source/Runtime/Private/PoseSearchContext.cpp
--- a/Engine/Plugins/Animation/PoseSearch/Source/Runtime/Private/PoseSearchContext.cpp
+++ b/Engine/Plugins/Animation/PoseSearch/Source/Runtime/Private/PoseSearchContext.cpp
@@ -455,6 +455,10 @@ FSearchContext::FSearchContext(float InDesiredPermutationTimeOffset, const FPose
 , bUseCachedChannelData(bInUseCachedChannelData)
 {
 	UpdateCurrentResultPoseVector();
+
+#if UE_POSE_SEARCH_TRACE_ENABLED
+	BestPoseCandidatesMap.Reserve(16);
+#endif // UE_POSE_SEARCH_TRACE_ENABLED
 }
 
 void FSearchContext::AddRole(const FRole& Role, const UAnimInstance* AnimInstance, const IPoseHistory* PoseHistory)
diff --git a/Engine/Plugins/Animation/PoseSearch/Source/Runtime/Private/PoseSearchDatabase.cpp b/Engine/Plugins/Animation/PoseSearch/Source/Runtime/Private/PoseSearchDatabase.cpp
--- a/Engine/Plugins/Animation/PoseSearch/Source/Runtime/Private/PoseSearchDatabase.cpp
+++ b/Engine/Plugins/Animation/PoseSearch/Source/Runtime/Private/PoseSearchDatabase.cpp
@@ -6,13 +6,13 @@
 #include "Animation/AnimSequence.h"
 #include "Animation/BlendSpace.h"
 #include "Animation/BlendSpace1D.h"
+#include "PoseSearch/MultiAnimAsset.h"
 #include "PoseSearch/PoseSearchAnimNotifies.h"
 #include "PoseSearch/PoseSearchContext.h"
 #include "PoseSearch/PoseSearchDefines.h"
 #include "PoseSearch/PoseSearchDerivedData.h"
 #include "PoseSearch/PoseSearchFeatureChannel.h"
 #include "PoseSearch/PoseSearchHistory.h"
-#include "PoseSearch/PoseSearchMultiSequence.h"
 #include "PoseSearch/PoseSearchSchema.h"
 #include "Serialization/ArchiveCountMem.h"
 #include "UObject/ObjectSaveContext.h"
@@ -376,7 +376,7 @@ UAnimationAsset* FPoseSearchDatabaseAnimationAssetBase::GetAnimationAssetForRole
 	return CastChecked<UAnimationAsset>(GetAnimationAsset());
 }
 
-const FTransform& FPoseSearchDatabaseAnimationAssetBase::GetRootTransformOriginForRole(const UE::PoseSearch::FRole& Role) const
+FTransform FPoseSearchDatabaseAnimationAssetBase::GetRootTransformOriginForRole(const UE::PoseSearch::FRole& Role) const
 {
 	check(GetNumRoles() == 1);
 	return FTransform::Identity;
@@ -606,61 +606,61 @@ bool FPoseSearchDatabaseAnimMontage::IsRootMotionEnabled() const
 #endif // WITH_EDITORONLY_DATA
 
 //////////////////////////////////////////////////////////////////////////
-// FPoseSearchDatabaseMultiSequence
-UObject* FPoseSearchDatabaseMultiSequence::GetAnimationAsset() const
+// FPoseSearchDatabaseMultiAnimAsset
+UObject* FPoseSearchDatabaseMultiAnimAsset::GetAnimationAsset() const
 {
-	return MultiSequence.Get();
+	return MultiAnimAsset.Get();
 }
 
-float FPoseSearchDatabaseMultiSequence::GetPlayLength() const
+float FPoseSearchDatabaseMultiAnimAsset::GetPlayLength() const
 {
-	return MultiSequence ? MultiSequence->GetPlayLength() : 0.f;
+	return MultiAnimAsset ? MultiAnimAsset->GetPlayLength() : 0.f;
 }
 
 #if WITH_EDITOR
-int32 FPoseSearchDatabaseMultiSequence::GetFrameAtTime(float Time) const
+int32 FPoseSearchDatabaseMultiAnimAsset::GetFrameAtTime(float Time) const
 {
-	return MultiSequence ? MultiSequence->GetFrameAtTime(Time) : 0;
+	return MultiAnimAsset ? MultiAnimAsset->GetFrameAtTime(Time) : 0;
 }
 #endif // WITH_EDITOR
 
-int32 FPoseSearchDatabaseMultiSequence::GetNumRoles() const
+int32 FPoseSearchDatabaseMultiAnimAsset::GetNumRoles() const
 {
-	return MultiSequence ? MultiSequence->GetNumRoles() : 0;
+	return MultiAnimAsset ? MultiAnimAsset->GetNumRoles() : 0;
 }
 
-UE::PoseSearch::FRole FPoseSearchDatabaseMultiSequence::GetRole(int32 RoleIndex) const
+UE::PoseSearch::FRole FPoseSearchDatabaseMultiAnimAsset::GetRole(int32 RoleIndex) const
 {
-	return MultiSequence ? MultiSequence->GetRole(RoleIndex) : UE::PoseSearch::DefaultRole;
+	return MultiAnimAsset ? MultiAnimAsset->GetRole(RoleIndex) : UE::PoseSearch::DefaultRole;
 }
 
-UAnimationAsset* FPoseSearchDatabaseMultiSequence::GetAnimationAssetForRole(const UE::PoseSearch::FRole& Role) const
+UAnimationAsset* FPoseSearchDatabaseMultiAnimAsset::GetAnimationAssetForRole(const UE::PoseSearch::FRole& Role) const
 {
-	return MultiSequence ? MultiSequence->GetSequence(Role) : nullptr;
+	return MultiAnimAsset ? MultiAnimAsset->GetAnimationAsset(Role) : nullptr;
 }
 
-const FTransform& FPoseSearchDatabaseMultiSequence::GetRootTransformOriginForRole(const UE::PoseSearch::FRole& Role) const
+FTransform FPoseSearchDatabaseMultiAnimAsset::GetRootTransformOriginForRole(const UE::PoseSearch::FRole& Role) const
 {
-	return MultiSequence ? MultiSequence->GetOrigin(Role) : FTransform::Identity;
+	return MultiAnimAsset ? MultiAnimAsset->GetOrigin(Role) : FTransform::Identity;
 }
 
 #if WITH_EDITORONLY_DATA
-UClass* FPoseSearchDatabaseMultiSequence::GetAnimationAssetStaticClass() const
+UClass* FPoseSearchDatabaseMultiAnimAsset::GetAnimationAssetStaticClass() const
 {
-	return UPoseSearchMultiSequence::StaticClass();
+	return UMultiAnimAsset::StaticClass();
 }
 
-bool FPoseSearchDatabaseMultiSequence::IsLooping() const
+bool FPoseSearchDatabaseMultiAnimAsset::IsLooping() const
 {
-	return MultiSequence &&
-		MultiSequence->IsLooping() &&
+	return MultiAnimAsset &&
+		MultiAnimAsset->IsLooping() &&
 		SamplingRange.Min == 0.f &&
 		SamplingRange.Max == 0.f;
 }
 
-bool FPoseSearchDatabaseMultiSequence::IsRootMotionEnabled() const
+bool FPoseSearchDatabaseMultiAnimAsset::IsRootMotionEnabled() const
 {
-	return MultiSequence ? MultiSequence->HasRootMotion() : false;
+	return MultiAnimAsset ? MultiAnimAsset->HasRootMotion() : false;
 }
 #endif // WITH_EDITORONLY_DATA
 
diff --git a/Engine/Plugins/Animation/PoseSearch/Source/Runtime/Private/PoseSearchDerivedData.cpp b/Engine/Plugins/Animation/PoseSearch/Source/Runtime/Private/PoseSearchDerivedData.cpp
--- a/Engine/Plugins/Animation/PoseSearch/Source/Runtime/Private/PoseSearchDerivedData.cpp
+++ b/Engine/Plugins/Animation/PoseSearch/Source/Runtime/Private/PoseSearchDerivedData.cpp
@@ -22,7 +22,6 @@
 #include "PoseSearch/PoseSearchDefines.h"
 #include "PoseSearch/PoseSearchDerivedDataKey.h"
 #include "PoseSearch/PoseSearchFeatureChannel.h"
-#include "PoseSearch/PoseSearchMultiSequence.h"
 #include "PoseSearch/PoseSearchNormalizationSet.h"
 #include "PoseSearch/PoseSearchSchema.h"
 #include "PoseSearchEigenHelper.h"
@@ -363,7 +362,7 @@ static bool InitSearchIndexAssets(FSearchIndexBase& SearchIndex, const TArray<FI
 				continue;
 			}
 
-			// checking for valid roles in DatabaseMultiSequence against the Schema
+			// checking for valid roles in DatabaseMultiAnimAsset against the Schema
 			bool bAreAllRolesSupported = true;
 			for (const FPoseSearchRoledSkeleton& RoledSkeleton : Schema->GetRoledSkeletons())
 			{
@@ -429,7 +428,7 @@ static bool InitSearchIndexAssets(FSearchIndexBase& SearchIndex, const TArray<FI
 					}
 				}
 			}
-			// support for FPoseSearchDatabaseSequence, FPoseSearchDatabaseAnimComposite, FPoseSearchDatabaseAnimMontage, FPoseSearchDatabaseMultiSequence
+			// support for FPoseSearchDatabaseSequence, FPoseSearchDatabaseAnimComposite, FPoseSearchDatabaseAnimMontage, FPoseSearchDatabaseMultiAnimAsset
 			else
 			{
 				ValidRanges.Reset();
@@ -907,7 +906,7 @@ static bool IndexDatabase(FSearchIndexBase& SearchIndexBase, const TArray<FInsta
 						const FVector BlendParameters = DatabaseBlendSpace->BlendParameterForSampleRanges(HorizontalIndex, VerticalIndex);
 
 						check(DatabaseBlendSpace->GetNumRoles() == 1);
-						const FTransform& RootTransformOrigin = DatabaseBlendSpace->GetRootTransformOriginForRole(DatabaseBlendSpace->GetRole(0));
+						const FTransform RootTransformOrigin = DatabaseBlendSpace->GetRootTransformOriginForRole(DatabaseBlendSpace->GetRole(0));
 						const FSamplerMapKey SamplerMapKey(DatabaseBlendSpace->BlendSpace, RootTransformOrigin, BlendParameters);
 						if (!SamplerMap.Contains(SamplerMapKey))
 						{
@@ -927,7 +926,7 @@ static bool IndexDatabase(FSearchIndexBase& SearchIndexBase, const TArray<FInsta
 				{
 					const FRole& Role = DatabaseAnimationAssetBase->GetRole(RoleIndex);
 					const UAnimationAsset* AnimationAsset = DatabaseAnimationAssetBase->GetAnimationAssetForRole(Role);
-					const FTransform& RootTransformOrigin = DatabaseAnimationAssetBase->GetRootTransformOriginForRole(Role);
+					const FTransform RootTransformOrigin = DatabaseAnimationAssetBase->GetRootTransformOriginForRole(Role);
 					const FSamplerMapKey SamplerMapKey(AnimationAsset, RootTransformOrigin);
 					if (!SamplerMap.Contains(SamplerMapKey))
 					{
@@ -985,7 +984,7 @@ static bool IndexDatabase(FSearchIndexBase& SearchIndexBase, const TArray<FInsta
 			if (const FMirrorDataCache* MirrorDataCache = RoledMirrorDataCaches.Find(Role))
 			{
 				UAnimationAsset* AnimationAsset = DatabaseAnimationAssetBase->GetAnimationAssetForRole(Role);
-				const FTransform& RootTransformOrigin = DatabaseAnimationAssetBase->GetRootTransformOriginForRole(Role);
+				const FTransform RootTransformOrigin = DatabaseAnimationAssetBase->GetRootTransformOriginForRole(Role);
 				const FVector BlendParameters = SearchIndexAsset.GetBlendParameters();
 				const int32 SamplerIndex = SamplerMap[{ AnimationAsset, RootTransformOrigin, BlendParameters }];
 				TempAssetSamplers.AnimationAssetSamplers.Emplace(&Samplers[SamplerIndex]);
diff --git a/Engine/Plugins/Animation/PoseSearch/Source/Runtime/Private/PoseSearchLibrary.cpp b/Engine/Plugins/Animation/PoseSearch/Source/Runtime/Private/PoseSearchLibrary.cpp
--- a/Engine/Plugins/Animation/PoseSearch/Source/Runtime/Private/PoseSearchLibrary.cpp
+++ b/Engine/Plugins/Animation/PoseSearch/Source/Runtime/Private/PoseSearchLibrary.cpp
@@ -16,10 +16,10 @@
 #include "InstancedStruct.h"
 #include "PoseSearch/AnimNode_MotionMatching.h"
 #include "PoseSearch/AnimNode_PoseSearchHistoryCollector.h"
+#include "PoseSearch/MultiAnimAsset.h"
 #include "PoseSearch/PoseSearchAnimNotifies.h"
 #include "PoseSearch/PoseSearchDatabase.h"
 #include "PoseSearch/PoseSearchDerivedData.h"
-#include "PoseSearch/PoseSearchMultiSequence.h"
 #include "PoseSearch/PoseSearchSchema.h"
 #include "PoseSearchFeatureChannel_Trajectory.h"
 #include "PoseSearch/Trace/PoseSearchTraceLogger.h"
@@ -812,9 +812,9 @@ UE::PoseSearch::FSearchResult UPoseSearchLibrary::MotionMatch(
 			const UAnimationAsset* AnimationAsset = Cast<UAnimationAsset>(Future.Animation);
 			if (!AnimationAsset)
 			{
-				if (const UPoseSearchMultiSequence* MultiSequence = Cast<UPoseSearchMultiSequence>(Future.Animation))
+				if (const UMultiAnimAsset* MultiAnimAsset = Cast<UMultiAnimAsset>(Future.Animation))
 				{
-					AnimationAsset = MultiSequence->GetSequence(Roles[RoleIndex]);
+					AnimationAsset = MultiAnimAsset->GetAnimationAsset(Roles[RoleIndex]);
 				}
 				else
 				{
diff --git a/Engine/Plugins/Animation/PoseSearch/Source/Runtime/Private/PoseSearchMultiSequence.cpp b/Engine/Plugins/Animation/PoseSearch/Source/Runtime/Private/PoseSearchMultiSequence.cpp
deleted file mode 100644
diff --git a/Engine/Plugins/Animation/PoseSearch/Source/Runtime/Public/PoseSearch/MultiAnimAsset.h b/Engine/Plugins/Animation/PoseSearch/Source/Runtime/Public/PoseSearch/MultiAnimAsset.h
new file mode 100644
--- /dev/null
+++ b/Engine/Plugins/Animation/PoseSearch/Source/Runtime/Public/PoseSearch/MultiAnimAsset.h
@@ -0,0 +1,28 @@
+// Copyright Epic Games, Inc. All Rights Reserved.
+
+#pragma once
+
+#include "MultiAnimAsset.generated.h"
+
+class UAnimationAsset;
+
+// @todo: move to Engine
+UCLASS(Abstract, Experimental, BlueprintType, Category = "Animation")
+class POSESEARCH_API UMultiAnimAsset : public UObject
+{
+	GENERATED_BODY()
+public:
+
+	[[nodiscard]] virtual bool IsLooping() const PURE_VIRTUAL(UMultiAnimAsset::IsLooping, return false;);
+	[[nodiscard]] virtual bool HasRootMotion() const PURE_VIRTUAL(UMultiAnimAsset::HasRootMotion, return false;);
+	[[nodiscard]] virtual float GetPlayLength() const PURE_VIRTUAL(UMultiAnimAsset::GetPlayLength, return 0.f;);
+
+#if WITH_EDITOR
+	[[nodiscard]] virtual int32 GetFrameAtTime(float Time) const PURE_VIRTUAL(UMultiAnimAsset::GetFrameAtTime, return 0;);
+#endif // WITH_EDITOR
+
+	[[nodiscard]] virtual int32 GetNumRoles() const PURE_VIRTUAL(UMultiAnimAsset::GetNumRoles, return 0;);
+	[[nodiscard]] virtual FName GetRole(int32 RoleIndex) const PURE_VIRTUAL(UMultiAnimAsset::GetRole, return FName(););
+	[[nodiscard]] virtual UAnimationAsset* GetAnimationAsset(const FName& Role) const PURE_VIRTUAL(UMultiAnimAsset::GetAnimationAsset, return nullptr;);
+	[[nodiscard]] virtual FTransform GetOrigin(const FName& Role) const PURE_VIRTUAL(UMultiAnimAsset::GetOrigin, return FTransform::Identity;);
+};
\ No newline at end of file
diff --git a/Engine/Plugins/Animation/PoseSearch/Source/Runtime/Public/PoseSearch/PoseSearchAssetSampler.h b/Engine/Plugins/Animation/PoseSearch/Source/Runtime/Public/PoseSearch/PoseSearchAssetSampler.h
--- a/Engine/Plugins/Animation/PoseSearch/Source/Runtime/Public/PoseSearch/PoseSearchAssetSampler.h
+++ b/Engine/Plugins/Animation/PoseSearch/Source/Runtime/Public/PoseSearch/PoseSearchAssetSampler.h
@@ -52,6 +52,7 @@ struct POSESEARCH_API FAnimationAssetSampler
 	static float GetPlayLength(const UAnimationAsset* AnimAsset, const FVector& BlendParameters = FVector::ZeroVector);
 
 	void SetRootTransformOrigin(const FTransform& InRootTransformOrigin) { RootTransformOrigin = InRootTransformOrigin; }
+	const FTransform& GetRootTransformOrigin() const { return RootTransformOrigin; }
 
 protected:
 	TWeakObjectPtr<const UAnimationAsset> AnimationAssetPtr;
diff --git a/Engine/Plugins/Animation/PoseSearch/Source/Runtime/Public/PoseSearch/PoseSearchDatabase.h b/Engine/Plugins/Animation/PoseSearch/Source/Runtime/Public/PoseSearch/PoseSearchDatabase.h
--- a/Engine/Plugins/Animation/PoseSearch/Source/Runtime/Public/PoseSearch/PoseSearchDatabase.h
+++ b/Engine/Plugins/Animation/PoseSearch/Source/Runtime/Public/PoseSearch/PoseSearchDatabase.h
@@ -15,7 +15,7 @@ class UAnimationAsset;
 class UAnimComposite;
 class UAnimMontage;
 class UBlendSpace;
-class UPoseSearchMultiSequence;
+class UMultiAnimAsset;
 
 #if WITH_EDITORONLY_DATA
 class UPoseSearchNormalizationSet;
@@ -57,7 +57,7 @@ struct POSESEARCH_API FPoseSearchDatabaseAnimationAssetBase
 	virtual int32 GetNumRoles() const { return 1; }
 	virtual UE::PoseSearch::FRole GetRole(int32 RoleIndex) const { return UE::PoseSearch::DefaultRole; }
 	virtual UAnimationAsset* GetAnimationAssetForRole(const UE::PoseSearch::FRole& Role) const;
-	virtual const FTransform& GetRootTransformOriginForRole(const UE::PoseSearch::FRole& Role) const;
+	virtual FTransform GetRootTransformOriginForRole(const UE::PoseSearch::FRole& Role) const;
 
 #if WITH_EDITOR
 	virtual int32 GetFrameAtTime(float Time) const;
@@ -322,13 +322,13 @@ template<> struct TStructOpsTypeTraits<FPoseSearchDatabaseAnimMontage> : public
 };
 
 USTRUCT(BlueprintType, Category = "Animation|Pose Search")
-struct POSESEARCH_API FPoseSearchDatabaseMultiSequence : public FPoseSearchDatabaseAnimationAssetBase
+struct POSESEARCH_API FPoseSearchDatabaseMultiAnimAsset : public FPoseSearchDatabaseAnimationAssetBase
 {
 	GENERATED_BODY()
-	virtual ~FPoseSearchDatabaseMultiSequence() = default;
+	virtual ~FPoseSearchDatabaseMultiAnimAsset() = default;
 
 	UPROPERTY(EditAnywhere, Category = "Settings", meta = (DisplayPriority = 0))
-	TObjectPtr<UPoseSearchMultiSequence> MultiSequence;
+	TObjectPtr<UMultiAnimAsset> MultiAnimAsset;
 
 #if WITH_EDITORONLY_DATA
 	// It allows users to set a time range to an individual animation sequence in the database. 
@@ -349,16 +349,16 @@ struct POSESEARCH_API FPoseSearchDatabaseMultiSequence : public FPoseSearchDatab
 	virtual int32 GetNumRoles() const override;
 	virtual UE::PoseSearch::FRole GetRole(int32 RoleIndex) const override;
 	virtual UAnimationAsset* GetAnimationAssetForRole(const UE::PoseSearch::FRole& Role) const override;
-	virtual const FTransform& GetRootTransformOriginForRole(const UE::PoseSearch::FRole& Role) const override;
+	virtual FTransform GetRootTransformOriginForRole(const UE::PoseSearch::FRole& Role) const override;
 
 #if WITH_EDITOR
 	virtual int32 GetFrameAtTime(float Time) const override;
 #endif // WITH_EDITOR
 
-	friend bool operator==(const FPoseSearchDatabaseMultiSequence& A, const FPoseSearchDatabaseMultiSequence& B)
+	friend bool operator==(const FPoseSearchDatabaseMultiAnimAsset& A, const FPoseSearchDatabaseMultiAnimAsset& B)
 	{
 		return static_cast<const FPoseSearchDatabaseAnimationAssetBase&>(A) == static_cast<const FPoseSearchDatabaseAnimationAssetBase&>(B)
-			&& A.MultiSequence == B.MultiSequence
+			&& A.MultiAnimAsset == B.MultiAnimAsset
 #if WITH_EDITORONLY_DATA
 			&& A.SamplingRange == B.SamplingRange
 #endif // WITH_EDITORONLY_DATA
@@ -366,7 +366,7 @@ struct POSESEARCH_API FPoseSearchDatabaseMultiSequence : public FPoseSearchDatab
 	}
 };
 
-template<> struct TStructOpsTypeTraits<FPoseSearchDatabaseMultiSequence> : public TStructOpsTypeTraitsBase2<FPoseSearchDatabaseMultiSequence>
+template<> struct TStructOpsTypeTraits<FPoseSearchDatabaseMultiAnimAsset> : public TStructOpsTypeTraitsBase2<FPoseSearchDatabaseMultiAnimAsset>
 {
 	enum { WithIdenticalViaEquality = true };
 };
diff --git a/Engine/Plugins/Animation/PoseSearch/Source/Runtime/Public/PoseSearch/PoseSearchDefines.h b/Engine/Plugins/Animation/PoseSearch/Source/Runtime/Public/PoseSearch/PoseSearchDefines.h
--- a/Engine/Plugins/Animation/PoseSearch/Source/Runtime/Public/PoseSearch/PoseSearchDefines.h
+++ b/Engine/Plugins/Animation/PoseSearch/Source/Runtime/Public/PoseSearch/PoseSearchDefines.h
@@ -32,7 +32,7 @@ static constexpr int32 PreallocatedRolesNum = 4;
 static constexpr int32 PreallocatedCachedQueriesNum = 8;
 static constexpr int32 PreallocatedCachedChannelDataNum = 64;
 static constexpr float FiniteDelta = 1 / 60.0f; // Time delta used for computing pose derivatives
-static constexpr int32 MaxNumberOfCollectedPoseCandidatesPerDatabase = 200;
+static constexpr int32 MaxNumberOfCollectedPoseCandidatesPerDatabase = 256;
 static constexpr int32 InvalidSearchId = INDEX_NONE; // @todo: it should be 0, since the search id is an hash!
 
 template<uint32 Alignment = DEFAULT_ALIGNMENT> class TMemStackSetAllocator : public TSetAllocator<TSparseArrayAllocator<TMemStackAllocator<Alignment>, TMemStackAllocator<Alignment>>, TMemStackAllocator<Alignment>> {};
diff --git a/Engine/Plugins/Animation/PoseSearch/Source/Runtime/Public/PoseSearch/PoseSearchMultiSequence.h b/Engine/Plugins/Animation/PoseSearch/Source/Runtime/Public/PoseSearch/PoseSearchMultiSequence.h
deleted file mode 100644
