motion matching - fix for misleading Warning when validating level that contains character using MotionMatching

diff --git a/Engine/Plugins/Animation/PoseSearch/Source/Runtime/Private/PoseSearchLibrary.cpp b/Engine/Plugins/Animation/PoseSearch/Source/Runtime/Private/PoseSearchLibrary.cpp
--- a/Engine/Plugins/Animation/PoseSearch/Source/Runtime/Private/PoseSearchLibrary.cpp
+++ b/Engine/Plugins/Animation/PoseSearch/Source/Runtime/Private/PoseSearchLibrary.cpp
@@ -531,11 +531,11 @@ void UPoseSearchLibrary::UpdateMotionMatchingState(
 			StringBuilder << bCanAdvance;
 			StringBuilder << "], Indexing [";
 
+			bool bIsIndexing = false;
 #if WITH_EDITOR
-			StringBuilder << SearchContext.IsAsyncBuildIndexInProgress();
-#else // WITH_EDITOR
-			StringBuilder << false;
+			bIsIndexing = SearchContext.IsAsyncBuildIndexInProgress();
 #endif // WITH_EDITOR
+			StringBuilder << bIsIndexing;
 
 			StringBuilder << "], Databases [";
 
@@ -551,7 +551,15 @@ void UPoseSearchLibrary::UpdateMotionMatchingState(
 			StringBuilder << "] ";
 
 			FString String = StringBuilder.ToString();
-			UE_LOG(LogPoseSearch, Warning, TEXT("%s"), *String);
+
+			if (bIsIndexing)
+			{
+				UE_LOG(LogPoseSearch, Log, TEXT("%s"), *String);
+			}
+			else
+			{
+				UE_LOG(LogPoseSearch, Warning, TEXT("%s"), *String);
+			}
 		}
 #endif // !NO_LOGGING
 
