motion matching - multi character editor improvements

diff --git a/Engine/Plugins/Animation/PoseSearch/Source/Editor/Private/PoseSearchDatabaseAssetListItem.cpp b/Engine/Plugins/Animation/PoseSearch/Source/Editor/Private/PoseSearchDatabaseAssetListItem.cpp
--- a/Engine/Plugins/Animation/PoseSearch/Source/Editor/Private/PoseSearchDatabaseAssetListItem.cpp
+++ b/Engine/Plugins/Animation/PoseSearch/Source/Editor/Private/PoseSearchDatabaseAssetListItem.cpp
@@ -1,7 +1,6 @@
 // Copyright Epic Games, Inc. All Rights Reserved.
 
 #include "PoseSearchDatabaseAssetListItem.h"
-
 #include "Animation/AnimComposite.h"
 #include "Animation/AnimSequence.h"
 #include "Animation/BlendSpace.h"
@@ -16,6 +15,7 @@
 #include "Framework/Commands/GenericCommands.h"
 #include "Framework/MultiBox/MultiBoxBuilder.h"
 #include "IAnimationEditor.h"
+#include "IMultiAnimAssetEditor.h"
 #include "IPersonaToolkit.h"
 #include "Misc/FeedbackContext.h"
 #include "Misc/TransactionObjectEvent.h"
@@ -252,20 +252,27 @@ namespace UE::PoseSearch
 
 							if (IAssetEditorInstance* Editor = AssetEditorSS->FindEditorForAsset(AnimationAsset, true))
 							{
+								float AnimationAssetTime = 0.f;
+								FVector AnimationAssetBlendParameters = FVector::ZeroVector;
+								ViewModel->GetAnimationTime(AssetTreeNode->SourceAssetIdx, AnimationAssetTime, AnimationAssetBlendParameters);
+
 								if (Editor->GetEditorName() == "AnimationEditor")
 								{
-									float AnimationAssetTime = 0.f;
-									FVector AnimationAssetBlendParameters = FVector::ZeroVector;
-									ViewModel->GetAnimationTime(AssetTreeNode->SourceAssetIdx, AnimationAssetTime, AnimationAssetBlendParameters);
-
-									const IAnimationEditor* AnimationEditor = static_cast<IAnimationEditor*>(Editor);
-									const UDebugSkelMeshComponent* PreviewComponent = AnimationEditor->GetPersonaToolkit()->GetPreviewMeshComponent();
+									IAnimationEditor* AnimationEditor = static_cast<IAnimationEditor*>(Editor);
+									UDebugSkelMeshComponent* PreviewComponent = AnimationEditor->GetPersonaToolkit()->GetPreviewMeshComponent();
 
 									// Open asset paused and at specific time as seen on the pose search debugger.
 									PreviewComponent->PreviewInstance->SetPosition(AnimationAssetTime);
 									PreviewComponent->PreviewInstance->SetPlaying(false);
 									PreviewComponent->PreviewInstance->SetBlendSpacePosition(AnimationAssetBlendParameters);
 								}
+								else if (Editor->GetEditorName() == "ChimeraAssetEditor")
+								{
+									IMultiAnimAssetEditor* MultiAnimAssetEditor = static_cast<IMultiAnimAssetEditor*>(Editor);
+
+									// Open asset paused and at specific time as seen on the pose search debugger.
+									MultiAnimAssetEditor->SetPreviewProperties(AnimationAssetTime, AnimationAssetBlendParameters, false);
+								}
 							}
 						}
 					}
diff --git a/Engine/Plugins/Animation/PoseSearch/Source/Editor/Public/IMultiAnimAssetEditor.h b/Engine/Plugins/Animation/PoseSearch/Source/Editor/Public/IMultiAnimAssetEditor.h
new file mode 100644
--- /dev/null
+++ b/Engine/Plugins/Animation/PoseSearch/Source/Editor/Public/IMultiAnimAssetEditor.h
@@ -0,0 +1,13 @@
+// Copyright Epic Games, Inc. All Rights Reserved.
+
+#pragma once
+
+// @todo: move UMultiAnimAsset as well as IMultiAnimAssetEditor to Engine or a base plugin for multi character animation assets
+
+#include "Toolkits/AssetEditorToolkit.h"
+
+class IMultiAnimAssetEditor : public FAssetEditorToolkit
+{
+public:
+	virtual void SetPreviewProperties(float AnimAssetTime, const FVector& AnimAssetBlendParameters, bool bAnimAssetPlaying) = 0;
+};
\ No newline at end of file
diff --git a/Engine/Plugins/Animation/PoseSearch/Source/Runtime/Public/PoseSearch/MultiAnimAsset.h b/Engine/Plugins/Animation/PoseSearch/Source/Runtime/Public/PoseSearch/MultiAnimAsset.h
--- a/Engine/Plugins/Animation/PoseSearch/Source/Runtime/Public/PoseSearch/MultiAnimAsset.h
+++ b/Engine/Plugins/Animation/PoseSearch/Source/Runtime/Public/PoseSearch/MultiAnimAsset.h
@@ -2,11 +2,13 @@
 
 #pragma once
 
+// @todo: move UMultiAnimAsset as well as IMultiAnimAssetEditor to Engine or a base plugin for multi character animation assets
+
 #include "MultiAnimAsset.generated.h"
 
 class UAnimationAsset;
 
-// @todo: move to Engine
+// UObject defining tuples of UAnimationAsset(s) with associated Role(s) and relative transforms from a shared reference system via GetOrigin
 UCLASS(Abstract, Experimental, BlueprintType, Category = "Animation")
 class POSESEARCH_API UMultiAnimAsset : public UObject
 {
