motion matching - multi character warp calculation

diff --git a/Engine/Plugins/Animation/PoseSearch/Source/Editor/Private/PoseSearchDatabaseViewModel.cpp b/Engine/Plugins/Animation/PoseSearch/Source/Editor/Private/PoseSearchDatabaseViewModel.cpp
--- a/Engine/Plugins/Animation/PoseSearch/Source/Editor/Private/PoseSearchDatabaseViewModel.cpp
+++ b/Engine/Plugins/Animation/PoseSearch/Source/Editor/Private/PoseSearchDatabaseViewModel.cpp
@@ -93,17 +93,23 @@ bool FDatabasePreviewActor::SpawnPreviewActor(UWorld* World, const UPoseSearchDa
 	const FMirrorDataCache MirrorDataCache(AnimInstance->GetMirrorDataTable(), AnimInstance->GetRequiredBonesOnAnyThread());
 	
 	Sampler.Init(PreviewAsset, SamplerRootTransformOrigin, IndexAsset.GetBlendParameters());
-	Sampler.Process();
 
 	PlayTimeOffset = 0.f;
 	if (PoseIdxForTimeOffset >= 0)
 	{
 		PlayTimeOffset = PoseSearchDatabase->GetRealAssetTime(PoseIdxForTimeOffset) - IndexAsset.GetFirstSampleTime(PoseSearchDatabase->Schema->SampleRate);
-		
-		// centering the Sampler RootTransformOrigin at PlayTimeOffset time, to be able to "align" multiple actors from different animation frames when selected by the pose search debugger
-		FTransform NewSamplerRootTransformOrigin = MirrorDataCache.MirrorTransform(Sampler.ExtractRootTransform(0.f));
-		NewSamplerRootTransformOrigin.SetToRelativeTransform(MirrorDataCache.MirrorTransform(Sampler.ExtractRootTransform(PlayTimeOffset)));
-		Sampler.SetRootTransformOrigin(NewSamplerRootTransformOrigin);
+
+		if (DatabaseAnimationAsset->GetNumRoles() > 1)
+		{
+			// @todo: implement support for UMultiAnimAsset. the transform should be centered to the origin of the multi character animation!
+		}
+		else
+		{
+			// centering the Sampler RootTransformOrigin at PlayTimeOffset time, to be able to "align" multiple actors from different animation frames when selected by the pose search debugger
+			FTransform NewSamplerRootTransformOrigin = MirrorDataCache.MirrorTransform(Sampler.ExtractRootTransform(0.f));
+			NewSamplerRootTransformOrigin.SetToRelativeTransform(MirrorDataCache.MirrorTransform(Sampler.ExtractRootTransform(PlayTimeOffset)));
+			Sampler.SetRootTransformOrigin(NewSamplerRootTransformOrigin);
+		}
 	}
 
 	AnimInstance->PlayAnim(IndexAsset.IsLooping(), 0.f);
diff --git a/Engine/Plugins/Animation/PoseSearch/Source/Runtime/Private/PoseSearchAssetSampler.cpp b/Engine/Plugins/Animation/PoseSearch/Source/Runtime/Private/PoseSearchAssetSampler.cpp
--- a/Engine/Plugins/Animation/PoseSearch/Source/Runtime/Private/PoseSearchAssetSampler.cpp
+++ b/Engine/Plugins/Animation/PoseSearch/Source/Runtime/Private/PoseSearchAssetSampler.cpp
@@ -217,18 +217,23 @@ static int32 GetHighestWeightSample(const TArray<struct FBlendSampleData>& Sampl
 
 //////////////////////////////////////////////////////////////////////////
 // FAssetSamplerBase
-FAnimationAssetSampler::FAnimationAssetSampler(TObjectPtr<const UAnimationAsset> InAnimationAsset, const FTransform& InRootTransformOrigin, const FVector& InBlendParameters, int32 InRootTransformSamplingRate)
+FAnimationAssetSampler::FAnimationAssetSampler(TObjectPtr<const UAnimationAsset> InAnimationAsset, const FTransform& InRootTransformOrigin, const FVector& InBlendParameters, int32 InRootTransformSamplingRate, bool bPreProcessRootTransform)
 {
-	Init(InAnimationAsset, InRootTransformOrigin, InBlendParameters, InRootTransformSamplingRate);
+	Init(InAnimationAsset, InRootTransformOrigin, InBlendParameters, InRootTransformSamplingRate, bPreProcessRootTransform);
 }
 
-void FAnimationAssetSampler::Init(TObjectPtr<const UAnimationAsset> InAnimationAsset, const FTransform& InRootTransformOrigin, const FVector& InBlendParameters, int32 InRootTransformSamplingRate)
+void FAnimationAssetSampler::Init(TObjectPtr<const UAnimationAsset> InAnimationAsset, const FTransform& InRootTransformOrigin, const FVector& InBlendParameters, int32 InRootTransformSamplingRate, bool bPreProcessRootTransform)
 {
 	AnimationAssetPtr = InAnimationAsset;
 	RootTransformOrigin = InRootTransformOrigin;
 	BlendParameters = InBlendParameters;
 	RootTransformSamplingRate = InRootTransformSamplingRate;
 	CachedPlayLength = GetPlayLength(AnimationAssetPtr.Get(), BlendParameters);
+
+	if (bPreProcessRootTransform)
+	{
+		Process();
+	}
 }
 
 bool FAnimationAssetSampler::IsInitialized() const
@@ -566,7 +571,8 @@ void FAnimationAssetSampler::ExtractPoseSearchNotifyStates(float Time, TFunction
 					}
 
 					// Get notifies for highest weighted
-					BlendSample.Animation->GetAnimNotifies((SampleTime - (ExtractionInterval * 0.5f)), ExtractionInterval, NotifyContext);
+					const float ExtractionStartTime = FMath::Min(SampleTime - (ExtractionInterval * 0.5f), BlendSample.Animation->GetPlayLength());
+					BlendSample.Animation->GetAnimNotifies(ExtractionStartTime, ExtractionInterval, NotifyContext);
 				}
 			}
 		}
@@ -578,7 +584,8 @@ void FAnimationAssetSampler::ExtractPoseSearchNotifyStates(float Time, TFunction
 	else if (const UAnimSequenceBase* SequenceBase = Cast<UAnimSequenceBase>(AnimationAssetPtr.Get()))
 	{
 		// getting pose search notifies in an interval of size ExtractionInterval, centered on Time
-		SequenceBase->GetAnimNotifies(Time - (ExtractionInterval * 0.5f), ExtractionInterval, NotifyContext);
+		const float ExtractionStartTime = FMath::Min(Time - (ExtractionInterval * 0.5f), SequenceBase->GetPlayLength());
+		SequenceBase->GetAnimNotifies(ExtractionStartTime, ExtractionInterval, NotifyContext);
 	}
 	else
 	{
diff --git a/Engine/Plugins/Animation/PoseSearch/Source/Runtime/Private/PoseSearchDerivedData.cpp b/Engine/Plugins/Animation/PoseSearch/Source/Runtime/Private/PoseSearchDerivedData.cpp
--- a/Engine/Plugins/Animation/PoseSearch/Source/Runtime/Private/PoseSearchDerivedData.cpp
+++ b/Engine/Plugins/Animation/PoseSearch/Source/Runtime/Private/PoseSearchDerivedData.cpp
@@ -911,7 +911,7 @@ static bool IndexDatabase(FSearchIndexBase& SearchIndexBase, const TArray<FInsta
 						if (!SamplerMap.Contains(SamplerMapKey))
 						{
 							SamplerMap.Add(SamplerMapKey, Samplers.Num());
-							Samplers.Emplace(DatabaseBlendSpace->BlendSpace, RootTransformOrigin, BlendParameters);
+							Samplers.Emplace(DatabaseBlendSpace->BlendSpace, RootTransformOrigin, BlendParameters, FAnimationAssetSampler::DefaultRootTransformSamplingRate, false);
 						}
 					}
 				}
@@ -931,7 +931,7 @@ static bool IndexDatabase(FSearchIndexBase& SearchIndexBase, const TArray<FInsta
 					if (!SamplerMap.Contains(SamplerMapKey))
 					{
 						SamplerMap.Add(SamplerMapKey, Samplers.Num());
-						Samplers.Emplace(AnimationAsset, RootTransformOrigin);
+						Samplers.Emplace(AnimationAsset, RootTransformOrigin, FVector::ZeroVector, FAnimationAssetSampler::DefaultRootTransformSamplingRate, false);
 					}
 				}
 			}
diff --git a/Engine/Plugins/Animation/PoseSearch/Source/Runtime/Public/PoseSearch/PoseSearchAssetSampler.h b/Engine/Plugins/Animation/PoseSearch/Source/Runtime/Public/PoseSearch/PoseSearchAssetSampler.h
--- a/Engine/Plugins/Animation/PoseSearch/Source/Runtime/Public/PoseSearch/PoseSearchAssetSampler.h
+++ b/Engine/Plugins/Animation/PoseSearch/Source/Runtime/Public/PoseSearch/PoseSearchAssetSampler.h
@@ -17,8 +17,9 @@ namespace UE::PoseSearch
  */
 struct POSESEARCH_API FAnimationAssetSampler
 {
-	FAnimationAssetSampler(TObjectPtr<const UAnimationAsset> InAnimationAsset = nullptr, const FTransform& InRootTransformOrigin = FTransform::Identity, const FVector& InBlendParameters = FVector::ZeroVector, int32 InRootTransformSamplingRate = 30);
-	void Init(TObjectPtr<const UAnimationAsset> InAnimationAsset, const FTransform& InRootTransformOrigin = FTransform::Identity, const FVector& InBlendParameters = FVector::ZeroVector, int32 InRootTransformSamplingRate = 30);
+	enum { DefaultRootTransformSamplingRate = 30 };
+	FAnimationAssetSampler(TObjectPtr<const UAnimationAsset> InAnimationAsset = nullptr, const FTransform& InRootTransformOrigin = FTransform::Identity, const FVector& InBlendParameters = FVector::ZeroVector, int32 InRootTransformSamplingRate = DefaultRootTransformSamplingRate, bool bPreProcessRootTransform = true);
+	void Init(TObjectPtr<const UAnimationAsset> InAnimationAsset, const FTransform& InRootTransformOrigin = FTransform::Identity, const FVector& InBlendParameters = FVector::ZeroVector, int32 InRootTransformSamplingRate = DefaultRootTransformSamplingRate, bool bPreProcessRootTransform = true);
 
 	bool IsInitialized() const;
 	float GetPlayLength() const;
@@ -60,7 +61,7 @@ protected:
 
 	// members used to sample blend spaces only!
 	FVector BlendParameters = FVector::ZeroVector;
-	int32 RootTransformSamplingRate = 30;
+	int32 RootTransformSamplingRate = DefaultRootTransformSamplingRate;
 	float CachedPlayLength = 0.f;
 	TArray<FTransform> AccumulatedRootTransform;
 
