Update PoseSearchGenerateTrajectory to face controller facing direction when in strafing mode.

diff --git a/Engine/Plugins/Animation/PoseSearch/Source/Runtime/Private/PoseSearchTrajectoryLibrary.cpp b/Engine/Plugins/Animation/PoseSearch/Source/Runtime/Private/PoseSearchTrajectoryLibrary.cpp
--- a/Engine/Plugins/Animation/PoseSearch/Source/Runtime/Private/PoseSearchTrajectoryLibrary.cpp
+++ b/Engine/Plugins/Animation/PoseSearch/Source/Runtime/Private/PoseSearchTrajectoryLibrary.cpp
@@ -76,8 +76,16 @@ void FPoseSearchTrajectoryData::UpdateData(
 	if (const USkeletalMeshComponent* MeshComp = Character->GetMesh())
 	{
 		TrajectoryDataDerived.Position = MeshComp->GetComponentLocation();
-		TrajectoryDataDerived.Facing = MeshComp->GetComponentRotation().Quaternion();
 		TrajectoryDataDerived.MeshCompRelativeRotation = MeshComp->GetRelativeRotation().Quaternion();
+		
+		if (TrajectoryDataDerived.bOrientRotationToMovement)
+		{
+			TrajectoryDataDerived.Facing = MeshComp->GetComponentRotation().Quaternion();
+		}
+		else
+		{
+			TrajectoryDataDerived.Facing = FQuat::MakeFromRotator(FRotator(0,TrajectoryDataState.DesiredControllerYawLastUpdate,0)) * TrajectoryDataDerived.MeshCompRelativeRotation;
+		}
 	}
 }
 
@@ -424,7 +432,7 @@ void UPoseSearchTrajectoryLibrary::HandleTrajectoryWorldCollisions(const UObject
 				}
 
 				// applying gravity
-				const FVector FreeFallOffset = Gravity * (0.5f * FreeFallAccumulatedSeconds * FreeFallAccumulatedSeconds);
+				const FVector FreeFallOffset =  Gravity * (0.5f * FreeFallAccumulatedSeconds * FreeFallAccumulatedSeconds);
 				Sample.Position += FreeFallOffset;
 
 				FHitResult HitResult;
