Added playrate multiplier for the motion matching node. Allows users to affect the playrate after motion matching state has been updated.

#rnx

diff --git a/Engine/Plugins/Animation/PoseSearch/Source/Runtime/Private/AnimNode_MotionMatching.cpp b/Engine/Plugins/Animation/PoseSearch/Source/Runtime/Private/AnimNode_MotionMatching.cpp
--- a/Engine/Plugins/Animation/PoseSearch/Source/Runtime/Private/AnimNode_MotionMatching.cpp
+++ b/Engine/Plugins/Animation/PoseSearch/Source/Runtime/Private/AnimNode_MotionMatching.cpp
@@ -201,7 +201,7 @@ void FAnimNode_MotionMatching::UpdateAssetPlayer(const FAnimationUpdateContext&
 			{
 				FAnimNode_BlendStack_Standalone::BlendTo(Context, AnimationAsset, MotionMatchingState.CurrentSearchResult.AssetTime,
 					SearchIndexAsset->IsLooping(), SearchIndexAsset->IsMirrored(), CurrentResultDatabase->Schema->GetMirrorDataTable(DefaultRole), BlendTime,
-					BlendProfile, BlendOption, bUseInertialBlend, SearchIndexAsset->GetBlendParameters(), MotionMatchingState.WantedPlayRate);
+					BlendProfile, BlendOption, bUseInertialBlend, SearchIndexAsset->GetBlendParameters(), MotionMatchingState.WantedPlayRate * PlayRateMultiplier);
 			}
 			else
 			{
@@ -213,7 +213,7 @@ void FAnimNode_MotionMatching::UpdateAssetPlayer(const FAnimationUpdateContext&
 	const bool bDidBlendToRequestAnInertialBlend = MotionMatchingState.bJumpedToPose && bUseInertialBlend;
 	UE::Anim::TOptionalScopedGraphMessage<UE::Anim::FAnimInertializationSyncScope> InertializationSync(bDidBlendToRequestAnInertialBlend, Context);
 	
-	FAnimNode_BlendStack_Standalone::UpdatePlayRate(MotionMatchingState.WantedPlayRate);
+	FAnimNode_BlendStack_Standalone::UpdatePlayRate(MotionMatchingState.WantedPlayRate * PlayRateMultiplier);
 	FAnimNode_BlendStack_Standalone::UpdateAssetPlayer(Context);
 
 	NextUpdateInterruptMode = EPoseSearchInterruptMode::DoNotInterrupt;
diff --git a/Engine/Plugins/Animation/PoseSearch/Source/Runtime/Public/PoseSearch/AnimNode_MotionMatching.h b/Engine/Plugins/Animation/PoseSearch/Source/Runtime/Public/PoseSearch/AnimNode_MotionMatching.h
--- a/Engine/Plugins/Animation/PoseSearch/Source/Runtime/Public/PoseSearch/AnimNode_MotionMatching.h
+++ b/Engine/Plugins/Animation/PoseSearch/Source/Runtime/Public/PoseSearch/AnimNode_MotionMatching.h
@@ -78,7 +78,11 @@ private:
 	// Effective range of play rate that can be applied to the animations to account for discrepancies in estimated velocity between the movement model and the animation.
 	UPROPERTY(EditAnywhere, Category = Settings, meta = (PinHiddenByDefault, ClampMin = "0.2", ClampMax = "3.0", UIMin = "0.2", UIMax = "3.0"))
 	FFloatInterval PlayRate = FFloatInterval(1.f, 1.f);
-
+	
+	// Experimental: Multiplier applied to the play rate of the selected animation after Motion Matching State has been updated.
+	UPROPERTY(EditAnywhere, Category = "Settings|Experimental", meta = (PinHiddenByDefault))
+	float PlayRateMultiplier = 1.0f;
+	
 	UPROPERTY(EditAnywhere, Category = Settings, Category = Settings)
 	bool bUseInertialBlend = false;
 
