motion matching - preventing numerical error in UPoseSearchDatabase::GetNormalizedAssetTime

diff --git a/Engine/Plugins/Animation/PoseSearch/Source/Runtime/Private/PoseSearchDatabase.cpp b/Engine/Plugins/Animation/PoseSearch/Source/Runtime/Private/PoseSearchDatabase.cpp
--- a/Engine/Plugins/Animation/PoseSearch/Source/Runtime/Private/PoseSearchDatabase.cpp
+++ b/Engine/Plugins/Animation/PoseSearch/Source/Runtime/Private/PoseSearchDatabase.cpp
@@ -1052,9 +1052,9 @@ float UPoseSearchDatabase::GetNormalizedAssetTime(int32 PoseIdx) const
 	{
 		// For BlendSpaces the AssetTime is in the range [0, 1] while the Sampling Range
 		// is in real time (seconds). We should be using but FAnimationAssetSampler::GetPlayLength(...) to normalize precisely,
-		// but Asset.GetNumPoses() - 1 is a good enough estimator
-		AssetTime = AssetTime * Schema->SampleRate / float(Asset.GetNumPoses() - 1);
-		check(AssetTime >= 0.f && AssetTime <= 1.f);
+		// but Asset.GetNumPoses() - 1 is a good enough estimator. FMath::Min(1, ...) is there to clamp numerical errors
+		AssetTime = FMath::Min(1.f, AssetTime * Schema->SampleRate / float(Asset.GetNumPoses() - 1));
+		check(AssetTime >= 0.f);
 	}
 
 	return AssetTime;
