chimera - improvements

diff --git a/Engine/Plugins/Animation/PoseSearch/Source/Editor/Private/AnimGraphNode_MotionMatching.cpp b/Engine/Plugins/Animation/PoseSearch/Source/Editor/Private/AnimGraphNode_MotionMatching.cpp
--- a/Engine/Plugins/Animation/PoseSearch/Source/Editor/Private/AnimGraphNode_MotionMatching.cpp
+++ b/Engine/Plugins/Animation/PoseSearch/Source/Editor/Private/AnimGraphNode_MotionMatching.cpp
@@ -27,12 +27,6 @@ FText UAnimGraphNode_MotionMatching::GetMenuCategory() const
 	return LOCTEXT("NodeCategory", "Pose Search");
 }
 
-void UAnimGraphNode_MotionMatching::BakeDataDuringCompilation(class FCompilerResultsLog& MessageLog)
-{
-	UAnimBlueprint* AnimBlueprint = GetAnimBlueprint();
-	AnimBlueprint->FindOrAddGroup(Node.GetGroupName());
-}
-
 UScriptStruct* UAnimGraphNode_MotionMatching::GetTimePropertyStruct() const
 {
 	return FAnimNode_MotionMatching::StaticStruct();
diff --git a/Engine/Plugins/Animation/PoseSearch/Source/Editor/Public/AnimGraphNode_MotionMatching.h b/Engine/Plugins/Animation/PoseSearch/Source/Editor/Public/AnimGraphNode_MotionMatching.h
--- a/Engine/Plugins/Animation/PoseSearch/Source/Editor/Public/AnimGraphNode_MotionMatching.h
+++ b/Engine/Plugins/Animation/PoseSearch/Source/Editor/Public/AnimGraphNode_MotionMatching.h
@@ -23,7 +23,6 @@ class UAnimGraphNode_MotionMatching : public UAnimGraphNode_BlendStack_Base
 	virtual FText GetNodeTitle(ENodeTitleType::Type TitleType) const override;
 	virtual FText GetMenuCategory() const override;
 	virtual UScriptStruct* GetTimePropertyStruct() const override;
-	virtual void BakeDataDuringCompilation(class FCompilerResultsLog& MessageLog) override;
 
 	virtual FAnimNode_BlendStack_Standalone* GetBlendStackNode() const override { return (FAnimNode_BlendStack_Standalone*)(&Node); }
 
diff --git a/Engine/Plugins/Animation/PoseSearch/Source/Runtime/Private/AnimNode_PoseSearchHistoryCollector.cpp b/Engine/Plugins/Animation/PoseSearch/Source/Runtime/Private/AnimNode_PoseSearchHistoryCollector.cpp
--- a/Engine/Plugins/Animation/PoseSearch/Source/Runtime/Private/AnimNode_PoseSearchHistoryCollector.cpp
+++ b/Engine/Plugins/Animation/PoseSearch/Source/Runtime/Private/AnimNode_PoseSearchHistoryCollector.cpp
@@ -166,7 +166,7 @@ void FAnimNode_PoseSearchHistoryCollector::Update_AnyThread(const FAnimationUpda
 {
 	DECLARE_SCOPE_HIERARCHICAL_COUNTER_ANIMNODE(Update_AnyThread);
 	Super::Update_AnyThread(Context);
-	UE::Anim::TScopedGraphMessage<UE::PoseSearch::FPoseHistoryProvider> ScopedMessage(Context, PoseHistory);
+	UE::Anim::TScopedGraphMessage<UE::PoseSearch::FPoseHistoryProvider> ScopedMessage(Context, this);
 	Source.Update(Context);
 }
 
@@ -231,7 +231,7 @@ void FAnimNode_PoseSearchComponentSpaceHistoryCollector::Update_AnyThread(const
 {
 	DECLARE_SCOPE_HIERARCHICAL_COUNTER_ANIMNODE(Update_AnyThread);
 	Super::Update_AnyThread(Context);
-	UE::Anim::TScopedGraphMessage<UE::PoseSearch::FPoseHistoryProvider> ScopedMessage(Context, PoseHistory);
+	UE::Anim::TScopedGraphMessage<UE::PoseSearch::FPoseHistoryProvider> ScopedMessage(Context, this);
 	Source.Update(Context);
 }
 
diff --git a/Engine/Plugins/Animation/PoseSearch/Source/Runtime/Private/PoseHistoryProvider.cpp b/Engine/Plugins/Animation/PoseSearch/Source/Runtime/Private/PoseHistoryProvider.cpp
new file mode 100644
--- /dev/null
+++ b/Engine/Plugins/Animation/PoseSearch/Source/Runtime/Private/PoseHistoryProvider.cpp
@@ -0,0 +1,20 @@
+// Copyright Epic Games, Inc. All Rights Reserved.
+
+#include "PoseSearch/PoseHistoryProvider.h"
+#include "PoseSearch/AnimNode_PoseSearchHistoryCollector.h"
+
+namespace UE::PoseSearch
+{
+	
+const IPoseHistory& FPoseHistoryProvider::GetPoseHistory() const
+{
+	if (HistoryCollector)
+	{
+		return HistoryCollector->GetPoseHistory();
+	}
+	
+	check(PoseHistory);
+	return *PoseHistory;
+}
+
+} // namespace UE::PoseSearch
\ No newline at end of file
diff --git a/Engine/Plugins/Animation/PoseSearch/Source/Runtime/Private/PoseSearchLibrary.cpp b/Engine/Plugins/Animation/PoseSearch/Source/Runtime/Private/PoseSearchLibrary.cpp
--- a/Engine/Plugins/Animation/PoseSearch/Source/Runtime/Private/PoseSearchLibrary.cpp
+++ b/Engine/Plugins/Animation/PoseSearch/Source/Runtime/Private/PoseSearchLibrary.cpp
@@ -407,11 +407,14 @@ void UPoseSearchLibrary::TraceMotionMatching(
 #endif // WITH_EDITOR && ENABLE_ANIM_DEBUG
 
 	TraceState.Cycle = FPlatformTime::Cycles64();
-	TraceState.AnimInstanceId = FObjectTrace::GetObjectId(SearchContext.GetAnimInstances()[0]);
 
-	TraceState.NodeId = SearchId;
-
-	TraceState.Output();
+	// @todo: avoid publishing duplicated TraceState in ALL the AnimInstances! -currently necessary for multi character-
+	for (const UAnimInstance* AnimInstance : SearchContext.GetAnimInstances())
+	{
+		TraceState.AnimInstanceId = FObjectTrace::GetObjectId(AnimInstance);
+		TraceState.NodeId = SearchId;
+		TraceState.Output();
+	}
 }
 #endif // UE_POSE_SEARCH_TRACE_ENABLED
 
diff --git a/Engine/Plugins/Animation/PoseSearch/Source/Runtime/Public/PoseSearch/PoseHistoryProvider.h b/Engine/Plugins/Animation/PoseSearch/Source/Runtime/Public/PoseSearch/PoseHistoryProvider.h
new file mode 100644
--- /dev/null
+++ b/Engine/Plugins/Animation/PoseSearch/Source/Runtime/Public/PoseSearch/PoseHistoryProvider.h
@@ -0,0 +1,31 @@
+// Copyright Epic Games, Inc. All Rights Reserved.
+
+#pragma once
+
+#include "Animation/AnimNodeMessages.h"
+
+struct FAnimNode_PoseSearchHistoryCollector_Base;
+
+namespace UE::PoseSearch
+{
+
+struct IPoseHistory;
+
+class POSESEARCH_API FPoseHistoryProvider : public UE::Anim::IGraphMessage
+{
+	DECLARE_ANIMGRAPH_MESSAGE(FPoseHistoryProvider);
+
+public:
+	UE_DEPRECATED(5.4, "Use FPoseHistoryProvider(const FAnimNode_PoseSearchHistoryCollector_Base* HistoryCollector) instead")
+	FPoseHistoryProvider(const IPoseHistory& InPoseHistory) : PoseHistory(&InPoseHistory), HistoryCollector(nullptr) { }
+	FPoseHistoryProvider(const FAnimNode_PoseSearchHistoryCollector_Base* InHistoryCollector) : PoseHistory(nullptr), HistoryCollector(InHistoryCollector) { }
+	const IPoseHistory& GetPoseHistory() const;
+	const FAnimNode_PoseSearchHistoryCollector_Base* GetHistoryCollector() const { return HistoryCollector; }
+
+private:
+	// delete after deleting FPoseHistoryProvider(const IPoseHistory& InPoseHistory)
+	const IPoseHistory* PoseHistory = nullptr;
+	const FAnimNode_PoseSearchHistoryCollector_Base* HistoryCollector = nullptr;
+};
+
+} // namespace UE::PoseSearch
\ No newline at end of file
diff --git a/Engine/Plugins/Animation/PoseSearch/Source/Runtime/Public/PoseSearch/PoseSearchHistory.h b/Engine/Plugins/Animation/PoseSearch/Source/Runtime/Public/PoseSearch/PoseSearchHistory.h
--- a/Engine/Plugins/Animation/PoseSearch/Source/Runtime/Public/PoseSearch/PoseSearchHistory.h
+++ b/Engine/Plugins/Animation/PoseSearch/Source/Runtime/Public/PoseSearch/PoseSearchHistory.h
@@ -2,10 +2,10 @@
 
 #pragma once
 
-#include "Animation/AnimNodeMessages.h"
 #include "BonePose.h"
 #include "Containers/RingBuffer.h"
 #include "DrawDebugHelpers.h"
+#include "PoseSearch/PoseHistoryProvider.h"
 #include "PoseSearch/PoseSearchDefines.h"
 #include "PoseSearch/PoseSearchTrajectoryLibrary.h"
 #include "UObject/ObjectKey.h"
@@ -206,18 +206,6 @@ private:
 	TArray<FPoseHistoryEntry, TInlineAllocator<4, TMemStackAllocator<>>> FutureEntries;
 };
 
-class POSESEARCH_API FPoseHistoryProvider : public UE::Anim::IGraphMessage
-{
-	DECLARE_ANIMGRAPH_MESSAGE(FPoseHistoryProvider);
-
-public:
-	FPoseHistoryProvider(const IPoseHistory& InPoseHistory) : PoseHistory(InPoseHistory) { }
-	const IPoseHistory& GetPoseHistory() const { return PoseHistory; }
-
-private:
-	const IPoseHistory& PoseHistory;
-};
-
 struct FHistoricalPoseIndex
 {
 	bool operator==(const FHistoricalPoseIndex& Index) const
