motion matching - FAnimationAssetSampler::ExtractPoseSearchNotifyStates numerical error fix

diff --git a/Engine/Plugins/Animation/PoseSearch/Source/Runtime/Private/KDTree.cpp b/Engine/Plugins/Animation/PoseSearch/Source/Runtime/Private/KDTree.cpp
--- a/Engine/Plugins/Animation/PoseSearch/Source/Runtime/Private/KDTree.cpp
+++ b/Engine/Plugins/Animation/PoseSearch/Source/Runtime/Private/KDTree.cpp
@@ -301,14 +301,18 @@ inline int32 FindNeighborsInternal(FKDTreeImplementation* KDTreeImplementation,
 
 	QUICK_SCOPE_CYCLE_COUNTER(STAT_FKDTree_FindNeighbors);
 
-	check(Query.GetData() && Query.Num() == KDTreeImplementation->dim && KDTreeImplementation->root_node);
-
-	const nanoflann::SearchParams SearchParams(
-		32,			// Ignored parameter (Kept for compatibility with the FLANN interface).
-		0.f,		// search for eps-approximate neighbours (default: 0)
-		false);		// only for radius search, require neighbours sorted by
-	KDTreeImplementation->findNeighbors(Result, Query.GetData(), SearchParams);
-	return Result.Num();
+	if (KDTreeImplementation)
+	{
+		check(Query.GetData() && Query.Num() == KDTreeImplementation->dim && KDTreeImplementation->root_node);
+
+		const nanoflann::SearchParams SearchParams(
+			32,			// Ignored parameter (Kept for compatibility with the FLANN interface).
+			0.f,		// search for eps-approximate neighbours (default: 0)
+			false);		// only for radius search, require neighbours sorted by
+		KDTreeImplementation->findNeighbors(Result, Query.GetData(), SearchParams);
+		return Result.Num();
+	}
+	return 0;
 
 #else // UE_POSE_SEARCH_USE_NANOFLANN
 
diff --git a/Engine/Plugins/Animation/PoseSearch/Source/Runtime/Private/PoseSearchAssetIndexer.cpp b/Engine/Plugins/Animation/PoseSearch/Source/Runtime/Private/PoseSearchAssetIndexer.cpp
--- a/Engine/Plugins/Animation/PoseSearch/Source/Runtime/Private/PoseSearchAssetIndexer.cpp
+++ b/Engine/Plugins/Animation/PoseSearch/Source/Runtime/Private/PoseSearchAssetIndexer.cpp
@@ -133,7 +133,7 @@ bool FAnimationAssetSamplers::IsLoopable() const
 	return true;
 }
 
-void FAnimationAssetSamplers::ExtractPoseSearchNotifyStates(float Time, TFunction<bool(UAnimNotifyState_PoseSearchBase*)> ProcessPoseSearchBase) const
+void FAnimationAssetSamplers::ExtractPoseSearchNotifyStates(float Time, const TFunction<bool(UAnimNotifyState_PoseSearchBase*)>& ProcessPoseSearchBase) const
 {
 	for (const FAnimationAssetSampler* Sampler : AnimationAssetSamplers)
 	{
@@ -141,7 +141,7 @@ void FAnimationAssetSamplers::ExtractPoseSearchNotifyStates(float Time, TFunctio
 	}
 }
 
-bool FAnimationAssetSamplers::ProcessAllAnimNotifyEvents(TFunction<bool(TConstArrayView<FAnimNotifyEvent>)> ProcessAnimNotifyEvents) const
+bool FAnimationAssetSamplers::ProcessAllAnimNotifyEvents(const TFunction<bool(TConstArrayView<FAnimNotifyEvent>)>& ProcessAnimNotifyEvents) const
 {
 	for (const FAnimationAssetSampler* Sampler : AnimationAssetSamplers)
 	{
@@ -764,7 +764,7 @@ bool FAssetIndexer::GetSampleVelocity(FVector& OutSampleVelocity, float SampleTi
 	return false;
 }
 
-bool FAssetIndexer::ProcessAllAnimNotifyEvents(TFunction<bool(TConstArrayView<FAnimNotifyEvent>)> ProcessAnimNotifyEvents) const
+bool FAssetIndexer::ProcessAllAnimNotifyEvents(const TFunction<bool(TConstArrayView<FAnimNotifyEvent>)>& ProcessAnimNotifyEvents) const
 {
 	return AssetSamplers.ProcessAllAnimNotifyEvents(ProcessAnimNotifyEvents);
 }
diff --git a/Engine/Plugins/Animation/PoseSearch/Source/Runtime/Private/PoseSearchAssetSampler.cpp b/Engine/Plugins/Animation/PoseSearch/Source/Runtime/Private/PoseSearchAssetSampler.cpp
--- a/Engine/Plugins/Animation/PoseSearch/Source/Runtime/Private/PoseSearchAssetSampler.cpp
+++ b/Engine/Plugins/Animation/PoseSearch/Source/Runtime/Private/PoseSearchAssetSampler.cpp
@@ -546,7 +546,7 @@ void FAnimationAssetSampler::Process()
 	}
 }
 
-void FAnimationAssetSampler::ExtractPoseSearchNotifyStates(float Time, TFunction<bool(UAnimNotifyState_PoseSearchBase*)> ProcessPoseSearchBase) const
+void FAnimationAssetSampler::ExtractPoseSearchNotifyStates(float Time, const TFunction<bool(UAnimNotifyState_PoseSearchBase*)>& ProcessPoseSearchBase) const
 {
 	float SampleTime = Time;
 	FAnimNotifyContext NotifyContext;
@@ -571,7 +571,7 @@ void FAnimationAssetSampler::ExtractPoseSearchNotifyStates(float Time, TFunction
 					}
 
 					// Get notifies for highest weighted
-					const float ExtractionStartTime = FMath::Min(SampleTime - (ExtractionInterval * 0.5f), BlendSample.Animation->GetPlayLength());
+					const float ExtractionStartTime = FMath::Min(SampleTime, BlendSample.Animation->GetPlayLength()) - (ExtractionInterval * 0.5f);
 					BlendSample.Animation->GetAnimNotifies(ExtractionStartTime, ExtractionInterval, NotifyContext);
 				}
 			}
@@ -584,7 +584,7 @@ void FAnimationAssetSampler::ExtractPoseSearchNotifyStates(float Time, TFunction
 	else if (const UAnimSequenceBase* SequenceBase = Cast<UAnimSequenceBase>(AnimationAssetPtr.Get()))
 	{
 		// getting pose search notifies in an interval of size ExtractionInterval, centered on Time
-		const float ExtractionStartTime = FMath::Min(Time - (ExtractionInterval * 0.5f), SequenceBase->GetPlayLength());
+		const float ExtractionStartTime = FMath::Min(Time, SequenceBase->GetPlayLength()) - (ExtractionInterval * 0.5f);
 		SequenceBase->GetAnimNotifies(ExtractionStartTime, ExtractionInterval, NotifyContext);
 	}
 	else
@@ -595,23 +595,14 @@ void FAnimationAssetSampler::ExtractPoseSearchNotifyStates(float Time, TFunction
 	// check which notifies actually overlap Time and are of the right base type
 	for (const FAnimNotifyEventReference& EventReference : NotifyContext.ActiveNotifies)
 	{
-		const FAnimNotifyEvent* NotifyEvent = EventReference.GetNotify();
-		if (!NotifyEvent)
+		if (const FAnimNotifyEvent* NotifyEvent = EventReference.GetNotify())
 		{
-			continue;
-		}
-
-		// @todo: is this condition necessary? can we just rely on the ExtractionInterval?
-		if (NotifyEvent->GetTime() > SampleTime || (NotifyEvent->GetTime() + NotifyEvent->GetDuration()) < SampleTime)
-		{
-			continue;
-		}
-
-		if (UAnimNotifyState_PoseSearchBase* PoseSearchAnimNotify = Cast<UAnimNotifyState_PoseSearchBase>(NotifyEvent->NotifyStateClass))
-		{
-			if (!ProcessPoseSearchBase(PoseSearchAnimNotify))
+			if (UAnimNotifyState_PoseSearchBase* PoseSearchAnimNotify = Cast<UAnimNotifyState_PoseSearchBase>(NotifyEvent->NotifyStateClass))
 			{
-				break;
+				if (!ProcessPoseSearchBase(PoseSearchAnimNotify))
+				{
+					break;
+				}
 			}
 		}
 	}
diff --git a/Engine/Plugins/Animation/PoseSearch/Source/Runtime/Private/PoseSearchDerivedDataKey.cpp b/Engine/Plugins/Animation/PoseSearch/Source/Runtime/Private/PoseSearchDerivedDataKey.cpp
--- a/Engine/Plugins/Animation/PoseSearch/Source/Runtime/Private/PoseSearchDerivedDataKey.cpp
+++ b/Engine/Plugins/Animation/PoseSearch/Source/Runtime/Private/PoseSearchDerivedDataKey.cpp
@@ -42,7 +42,7 @@ FKeyBuilder::FKeyBuilder(const UObject* Object, bool bUseDataVer, bool bPerformC
 	if (bUseDataVer)
 	{
 		// used to invalidate the key without having to change POSESEARCHDB_DERIVEDDATA_VER all the times
-		int32 POSESEARCHDB_DERIVEDDATA_VER_SMALL = 257;
+		int32 POSESEARCHDB_DERIVEDDATA_VER_SMALL = 259;
 		FGuid VersionGuid = FDevSystemGuids::GetSystemGuid(FDevSystemGuids::Get().POSESEARCHDB_DERIVEDDATA_VER);
 		FString AnimationCompressionVersionString = UE::Anim::Compression::AnimationCompressionVersionString;
 
diff --git a/Engine/Plugins/Animation/PoseSearch/Source/Runtime/Public/PoseSearch/PoseSearchAssetIndexer.h b/Engine/Plugins/Animation/PoseSearch/Source/Runtime/Public/PoseSearch/PoseSearchAssetIndexer.h
--- a/Engine/Plugins/Animation/PoseSearch/Source/Runtime/Public/PoseSearch/PoseSearchAssetIndexer.h
+++ b/Engine/Plugins/Animation/PoseSearch/Source/Runtime/Public/PoseSearch/PoseSearchAssetIndexer.h
@@ -35,8 +35,8 @@ struct FAnimationAssetSamplers
 
 	float GetPlayLength() const;
 	bool IsLoopable() const;
-	void ExtractPoseSearchNotifyStates(float Time, TFunction<bool(UAnimNotifyState_PoseSearchBase*)> ProcessPoseSearchBase) const;
-	bool ProcessAllAnimNotifyEvents(TFunction<bool(TConstArrayView<FAnimNotifyEvent>)> ProcessAnimNotifyEvents) const;
+	void ExtractPoseSearchNotifyStates(float Time, const TFunction<bool(UAnimNotifyState_PoseSearchBase*)>& ProcessPoseSearchBase) const;
+	bool ProcessAllAnimNotifyEvents(const TFunction<bool(TConstArrayView<FAnimNotifyEvent>)>& ProcessAnimNotifyEvents) const;
 	const FString GetAssetName() const;
 
 	FTransform ExtractRootTransform(float Time, int32 RoleIndex) const;
@@ -83,7 +83,7 @@ public:
 	// if bUseCharacterSpaceVelocities is true, velocities will be computed in root bone space, rather than animation (world) space
 	bool GetSampleVelocity(FVector& OutSampleVelocity, float SampleTimeOffset, float OriginTimeOffset, int32 SampleIdx, int8 SchemaSampleBoneIdx, int8 SchemaOriginBoneIdx, const FRole& SampleRole, const FRole& OriginRole, bool bUseCharacterSpaceVelocities = true, EPermutationTimeType PermutationTimeType = EPermutationTimeType::UseSampleTime, int32 SamplingAttributeId = INDEX_NONE);
 
-	bool ProcessAllAnimNotifyEvents(TFunction<bool(TConstArrayView<FAnimNotifyEvent>)> ProcessAnimNotifyEvents) const;
+	bool ProcessAllAnimNotifyEvents(const TFunction<bool(TConstArrayView<FAnimNotifyEvent>)>& ProcessAnimNotifyEvents) const;
 	const FString GetAssetName() const;
 	float GetPlayLength() const;
 
diff --git a/Engine/Plugins/Animation/PoseSearch/Source/Runtime/Public/PoseSearch/PoseSearchAssetSampler.h b/Engine/Plugins/Animation/PoseSearch/Source/Runtime/Public/PoseSearch/PoseSearchAssetSampler.h
--- a/Engine/Plugins/Animation/PoseSearch/Source/Runtime/Public/PoseSearch/PoseSearchAssetSampler.h
+++ b/Engine/Plugins/Animation/PoseSearch/Source/Runtime/Public/PoseSearch/PoseSearchAssetSampler.h
@@ -42,7 +42,7 @@ struct POSESEARCH_API FAnimationAssetSampler
 
 	// Extracts notify states inheriting from UAnimNotifyState_PoseSearchBase present in the sequence at Time.
 	// The function does not empty NotifyStates before adding new notifies!
-	void ExtractPoseSearchNotifyStates(float Time, TFunction<bool(UAnimNotifyState_PoseSearchBase*)> ProcessPoseSearchBase) const;
+	void ExtractPoseSearchNotifyStates(float Time, const TFunction<bool(UAnimNotifyState_PoseSearchBase*)>& ProcessPoseSearchBase) const;
 	
 	TConstArrayView<FAnimNotifyEvent> GetAllAnimNotifyEvents() const;
 
