motion matching - removing unneeded UPoseSearchSchema::PreSave - fix for PoseSearch and Cook Determinism

diff --git a/Engine/Plugins/Animation/PoseSearch/Source/Runtime/Private/PoseSearchSchema.cpp b/Engine/Plugins/Animation/PoseSearch/Source/Runtime/Private/PoseSearchSchema.cpp
--- a/Engine/Plugins/Animation/PoseSearch/Source/Runtime/Private/PoseSearchSchema.cpp
+++ b/Engine/Plugins/Animation/PoseSearch/Source/Runtime/Private/PoseSearchSchema.cpp
@@ -328,12 +328,6 @@ void UPoseSearchSchema::Finalize()
 	}
 }
 
-void UPoseSearchSchema::PreSave(FObjectPreSaveContext ObjectSaveContext)
-{
-	Finalize();
-	Super::PreSave(ObjectSaveContext);
-}
-
 void UPoseSearchSchema::PostLoad()
 {
 	Super::PostLoad();
diff --git a/Engine/Plugins/Animation/PoseSearch/Source/Runtime/Public/PoseSearch/PoseSearchSchema.h b/Engine/Plugins/Animation/PoseSearch/Source/Runtime/Public/PoseSearch/PoseSearchSchema.h
--- a/Engine/Plugins/Animation/PoseSearch/Source/Runtime/Public/PoseSearch/PoseSearchSchema.h
+++ b/Engine/Plugins/Animation/PoseSearch/Source/Runtime/Public/PoseSearch/PoseSearchSchema.h
@@ -141,7 +141,6 @@ public:
 	}
 
 	// UObject
-	virtual void PreSave(FObjectPreSaveContext ObjectSaveContext) override;
 	virtual void PostLoad() override;
 
 	int8 AddBoneReference(const FBoneReference& BoneReference, const UE::PoseSearch::FRole& Role);
