motion matching - normalizing rotations in UE::PoseSearch::ProcessRootTransform when sampling blend spaces to keep numerical errors in check

diff --git a/Engine/Plugins/Animation/PoseSearch/Source/Runtime/Private/PoseSearchAssetSampler.cpp b/Engine/Plugins/Animation/PoseSearch/Source/Runtime/Private/PoseSearchAssetSampler.cpp
--- a/Engine/Plugins/Animation/PoseSearch/Source/Runtime/Private/PoseSearchAssetSampler.cpp
+++ b/Engine/Plugins/Animation/PoseSearch/Source/Runtime/Private/PoseSearchAssetSampler.cpp
@@ -195,6 +195,9 @@ static void ProcessRootTransform(const UBlendSpace* BlendSpace, const FVector& B
 			}
 
 			AccumulatedRootTransform[SampleIdx] = RootMotionMovementParams.GetRootMotionTransform() * AccumulatedRootTransform[SampleIdx - 1];
+			
+			// keep numerical errors in check
+			AccumulatedRootTransform[SampleIdx].NormalizeRotation();
 		}
 	}
 }
