motion matching
- support for blend spaces sampling range
- few tool side memory/performance optimizations

diff --git a/Engine/Plugins/Animation/PoseSearch/Source/Editor/Private/PoseSearchDebuggerDatabaseView.cpp b/Engine/Plugins/Animation/PoseSearch/Source/Editor/Private/PoseSearchDebuggerDatabaseView.cpp
--- a/Engine/Plugins/Animation/PoseSearch/Source/Editor/Private/PoseSearchDebuggerDatabaseView.cpp
+++ b/Engine/Plugins/Animation/PoseSearch/Source/Editor/Private/PoseSearchDebuggerDatabaseView.cpp
@@ -616,7 +616,7 @@ void SDebuggerDatabaseView::PopulateViewRows()
 {
 	ActiveView.Rows.Reset();
 	ContinuingPoseView.Rows.Reset();
-	FilteredDatabaseView.Rows.Empty();
+	FilteredDatabaseView.Rows.Reset();
 
 	const int32 UnfilteredDatabaseRowsNum = UnfilteredDatabaseRows.Num();
 	if (UnfilteredDatabaseRowsNum > 0)
diff --git a/Engine/Plugins/Animation/PoseSearch/Source/Editor/Private/PoseSearchDebuggerViewModel.cpp b/Engine/Plugins/Animation/PoseSearch/Source/Editor/Private/PoseSearchDebuggerViewModel.cpp
--- a/Engine/Plugins/Animation/PoseSearch/Source/Editor/Private/PoseSearchDebuggerViewModel.cpp
+++ b/Engine/Plugins/Animation/PoseSearch/Source/Editor/Private/PoseSearchDebuggerViewModel.cpp
@@ -58,7 +58,7 @@ int32 FDebuggerViewModel::GetNodesNum() const
 
 void FDebuggerViewModel::OnUpdate()
 {
-	MotionMatchingStates.Empty();
+	MotionMatchingStates.Reset();
 
 	// Get provider and validate
 	const TraceServices::IAnalysisSession* Session = RewindDebugger.Get()->GetAnalysisSession();
diff --git a/Engine/Plugins/Animation/PoseSearch/Source/Runtime/Private/PoseSearchDatabase.cpp b/Engine/Plugins/Animation/PoseSearch/Source/Runtime/Private/PoseSearchDatabase.cpp
--- a/Engine/Plugins/Animation/PoseSearch/Source/Runtime/Private/PoseSearchDatabase.cpp
+++ b/Engine/Plugins/Animation/PoseSearch/Source/Runtime/Private/PoseSearchDatabase.cpp
@@ -391,10 +391,12 @@ int64 FPoseSearchDatabaseAnimationAssetBase::GetEditorMemSize() const
 
 FFloatInterval FPoseSearchDatabaseAnimationAssetBase::GetEffectiveSamplingRange() const
 {
-	const FFloatInterval SamplingRange = GetSamplingRange();
+	return GetEffectiveSamplingRange(GetPlayLength(), GetSamplingRange());
+}
 
+FFloatInterval FPoseSearchDatabaseAnimationAssetBase::GetEffectiveSamplingRange(float PlayLength, const FFloatInterval& SamplingRange)
+{
 	const bool bSampleAll = (SamplingRange.Min == 0.0f) && (SamplingRange.Max == 0.0f);
-	const float PlayLength = GetPlayLength();
 	FFloatInterval Range;
 	Range.Min = bSampleAll ? 0.0f : SamplingRange.Min;
 	Range.Max = bSampleAll ? PlayLength : FMath::Min(PlayLength, SamplingRange.Max);
diff --git a/Engine/Plugins/Animation/PoseSearch/Source/Runtime/Private/PoseSearchDerivedData.cpp b/Engine/Plugins/Animation/PoseSearch/Source/Runtime/Private/PoseSearchDerivedData.cpp
--- a/Engine/Plugins/Animation/PoseSearch/Source/Runtime/Private/PoseSearchDerivedData.cpp
+++ b/Engine/Plugins/Animation/PoseSearch/Source/Runtime/Private/PoseSearchDerivedData.cpp
@@ -279,11 +279,39 @@ public:
 	}
 };
 
+static void FindValidBlendSpaceIntervals(const FPoseSearchDatabaseBlendSpace* DatabaseBlendSpace, const TArray<FBlendSampleData>& BlendSamples, const FFloatInterval& ExcludeFromDatabaseParameters, TArray<FFloatRange>& ValidRanges)
+{
+	check(DatabaseBlendSpace);
+
+	const float PlayLength = DatabaseBlendSpace->BlendSpace->GetAnimationLengthFromSampleData(BlendSamples);
+
+	const bool bIsLooping = DatabaseBlendSpace->IsLooping();
+
+	// scaling blend space SamplingRange from the space [0, 1] to [0, PlayLength] with PlayLength calculated from the BlendSamples
+	FFloatInterval ScaledSamplingRange = DatabaseBlendSpace->GetSamplingRange();
+	ScaledSamplingRange.Min *= PlayLength;
+	ScaledSamplingRange.Max *= PlayLength;
+
+	const FFloatInterval EffectiveSamplingInterval = FPoseSearchDatabaseAnimationAssetBase::GetEffectiveSamplingRange(PlayLength, ScaledSamplingRange);
+	FFloatRange EffectiveSamplingRange = FFloatRange::Inclusive(EffectiveSamplingInterval.Min, EffectiveSamplingInterval.Max);
+	if (!bIsLooping)
+	{
+		const FFloatRange ExcludeFromDatabaseRange(ExcludeFromDatabaseParameters.Min, PlayLength + ExcludeFromDatabaseParameters.Max);
+		EffectiveSamplingRange = FFloatRange::Intersection(EffectiveSamplingRange, ExcludeFromDatabaseRange);
+	}
+
+	// start from a single interval defined by the database sequence sampling range
+	ValidRanges.Reset();
+	ValidRanges.Add(EffectiveSamplingRange);
+
+	// @todo: evaluate eventual UAnimNotifyState_PoseSearchExcludeFromDatabase(s) like in FindValidSequenceIntervals if necessary
+}
+
 static void FindValidSequenceIntervals(const FPoseSearchDatabaseAnimationAssetBase* DatabaseAsset, const FFloatInterval& ExcludeFromDatabaseParameters, TArray<FFloatRange>& ValidRanges)
 {
 	check(DatabaseAsset);
 
-	bool bIsLooping = DatabaseAsset->IsLooping();
+	const bool bIsLooping = DatabaseAsset->IsLooping();
 	const float PlayLength = DatabaseAsset->GetPlayLength();
 
 	const FFloatInterval EffectiveSamplingInterval = DatabaseAsset->GetEffectiveSamplingRange();
@@ -295,7 +323,7 @@ static void FindValidSequenceIntervals(const FPoseSearchDatabaseAnimationAssetBa
 	}
 
 	// start from a single interval defined by the database sequence sampling range
-	ValidRanges.Empty();
+	ValidRanges.Reset();
 	ValidRanges.Add(EffectiveSamplingRange);
 
 	for (int32 RoleIndex = 0; RoleIndex < DatabaseAsset->GetNumRoles(); ++RoleIndex)
@@ -338,7 +366,7 @@ static bool InitSearchIndexAssets(FSearchIndexBase& SearchIndex, const TArray<FI
 
 	check(Schema);
 
-	SearchIndex.Assets.Empty();
+	SearchIndex.Assets.Reset();
 	TArray<FFloatRange> ValidRanges;
 	TArray<FBlendSampleData> BlendSamples;
 
@@ -399,29 +427,34 @@ static bool InitSearchIndexAssets(FSearchIndexBase& SearchIndex, const TArray<FI
 						int32 TriangulationIndex = 0;
 						DatabaseBlendSpace->BlendSpace->GetSamplesFromBlendInput(BlendParameters, BlendSamples, TriangulationIndex, true);
 
-						const float PlayLength = DatabaseBlendSpace->BlendSpace->GetAnimationLengthFromSampleData(BlendSamples);
+						FindValidBlendSpaceIntervals(DatabaseBlendSpace, BlendSamples, ExcludeFromDatabaseParameters, ValidRanges);
 
-						for (int32 PermutationIdx = 0; PermutationIdx < Schema->NumberOfPermutations; ++PermutationIdx)
+						for (const FFloatRange& Range : ValidRanges)
 						{
-							if (bAddUnmirrored)
+							for (int32 PermutationIdx = 0; PermutationIdx < Schema->NumberOfPermutations; ++PermutationIdx)
 							{
-								const FSearchIndexAsset PoseSearchIndexAsset(AnimationAssetIndex, TotalPoses, false, bIsLooping, 
-									bDisableReselection, FFloatInterval(0.f, PlayLength), Schema->SampleRate, PermutationIdx, BlendParameters);
-								if (PoseSearchIndexAsset.GetNumPoses() > 0)
+								const FFloatInterval RangeInterval(Range.GetLowerBoundValue(), Range.GetUpperBoundValue());
+
+								if (bAddUnmirrored)
 								{
-									SearchIndex.Assets.Add(PoseSearchIndexAsset);
-									TotalPoses += PoseSearchIndexAsset.GetNumPoses();
+									const FSearchIndexAsset PoseSearchIndexAsset(AnimationAssetIndex, TotalPoses, false, bIsLooping,
+										bDisableReselection, RangeInterval, Schema->SampleRate, PermutationIdx, BlendParameters);
+									if (PoseSearchIndexAsset.GetNumPoses() > 0)
+									{
+										SearchIndex.Assets.Add(PoseSearchIndexAsset);
+										TotalPoses += PoseSearchIndexAsset.GetNumPoses();
+									}
 								}
-							}
 
-							if (bAddMirrored)
-							{
-								const FSearchIndexAsset PoseSearchIndexAsset(AnimationAssetIndex, TotalPoses, true, bIsLooping,
-									bDisableReselection, FFloatInterval(0.f, PlayLength), Schema->SampleRate, PermutationIdx, BlendParameters);
-								if (PoseSearchIndexAsset.GetNumPoses() > 0)
+								if (bAddMirrored)
 								{
-									SearchIndex.Assets.Add(PoseSearchIndexAsset);
-									TotalPoses += PoseSearchIndexAsset.GetNumPoses();
+									const FSearchIndexAsset PoseSearchIndexAsset(AnimationAssetIndex, TotalPoses, true, bIsLooping,
+										bDisableReselection, RangeInterval, Schema->SampleRate, PermutationIdx, BlendParameters);
+									if (PoseSearchIndexAsset.GetNumPoses() > 0)
+									{
+										SearchIndex.Assets.Add(PoseSearchIndexAsset);
+										TotalPoses += PoseSearchIndexAsset.GetNumPoses();
+									}
 								}
 							}
 						}
@@ -431,17 +464,18 @@ static bool InitSearchIndexAssets(FSearchIndexBase& SearchIndex, const TArray<FI
 			// support for FPoseSearchDatabaseSequence, FPoseSearchDatabaseAnimComposite, FPoseSearchDatabaseAnimMontage, FPoseSearchDatabaseMultiAnimAsset
 			else
 			{
-				ValidRanges.Reset();
-
 				FindValidSequenceIntervals(DatabaseAsset, ExcludeFromDatabaseParameters, ValidRanges);
+
 				for (const FFloatRange& Range : ValidRanges)
 				{
 					for (int32 PermutationIdx = 0; PermutationIdx < Schema->NumberOfPermutations; ++PermutationIdx)
 					{
+						const FFloatInterval RangeInterval(Range.GetLowerBoundValue(), Range.GetUpperBoundValue());
+
 						if (bAddUnmirrored)
 						{
 							const FSearchIndexAsset PoseSearchIndexAsset(AnimationAssetIndex, TotalPoses, false, bIsLooping,
-								bDisableReselection, FFloatInterval(Range.GetLowerBoundValue(), Range.GetUpperBoundValue()), Schema->SampleRate, PermutationIdx);
+								bDisableReselection, RangeInterval, Schema->SampleRate, PermutationIdx);
 							if (PoseSearchIndexAsset.GetNumPoses() > 0)
 							{
 								SearchIndex.Assets.Add(PoseSearchIndexAsset);
@@ -452,7 +486,7 @@ static bool InitSearchIndexAssets(FSearchIndexBase& SearchIndex, const TArray<FI
 						if (bAddMirrored)
 						{
 							const FSearchIndexAsset PoseSearchIndexAsset(AnimationAssetIndex, TotalPoses, true, bIsLooping,
-								bDisableReselection, FFloatInterval(Range.GetLowerBoundValue(), Range.GetUpperBoundValue()), Schema->SampleRate, PermutationIdx);
+								bDisableReselection, RangeInterval, Schema->SampleRate, PermutationIdx);
 							if (PoseSearchIndexAsset.GetNumPoses() > 0)
 							{
 								SearchIndex.Assets.Add(PoseSearchIndexAsset);
diff --git a/Engine/Plugins/Animation/PoseSearch/Source/Runtime/Public/PoseSearch/PoseSearchDatabase.h b/Engine/Plugins/Animation/PoseSearch/Source/Runtime/Public/PoseSearch/PoseSearchDatabase.h
--- a/Engine/Plugins/Animation/PoseSearch/Source/Runtime/Public/PoseSearch/PoseSearchDatabase.h
+++ b/Engine/Plugins/Animation/PoseSearch/Source/Runtime/Public/PoseSearch/PoseSearchDatabase.h
@@ -78,6 +78,7 @@ struct POSESEARCH_API FPoseSearchDatabaseAnimationAssetBase
 	// [0, 0] represents the entire frame range of the original animation.
 	virtual FFloatInterval GetSamplingRange() const { return FFloatInterval(0.f, 0.f); }
 	FFloatInterval GetEffectiveSamplingRange() const;
+	static FFloatInterval GetEffectiveSamplingRange(float PlayLength, const FFloatInterval& SamplingRange);
 
 	virtual int64 GetEditorMemSize() const;
 	virtual int64 GetApproxCookedSize() const { return GetEditorMemSize(); }
@@ -204,9 +205,16 @@ struct POSESEARCH_API FPoseSearchDatabaseBlendSpace : public FPoseSearchDatabase
 	UPROPERTY(EditAnywhere, Category = "Settings", meta = (EditCondition = "bUseSingleSample", EditConditionHides, DisplayPriority = 9))
 	float BlendParamY = 0.f;
 
+	// It allows users to set a time range to an individual blend space in the database.
+	// This is effectively trimming the beginning and end of the animation in the database (not in the original blend space).
+	// If set to [0, 0] it will be the entire frame range of the original blend space.
+	UPROPERTY(EditAnywhere, Category = "Settings", meta = (ClampToMinMaxLimits, DisplayPriority = 2))
+	FFloatInterval SamplingRange = FFloatInterval(0.f, 0.f);
+
 	virtual UClass* GetAnimationAssetStaticClass() const override;
 	virtual bool IsLooping() const override;
 	virtual bool IsRootMotionEnabled() const override;
+	virtual FFloatInterval GetSamplingRange() const override { return SamplingRange; }
 
 	void GetBlendSpaceParameterSampleRanges(int32& HorizontalBlendNum, int32& VerticalBlendNum) const;
 	FVector BlendParameterForSampleRanges(int32 HorizontalBlendIndex, int32 VerticalBlendIndex) const;
@@ -229,6 +237,7 @@ struct POSESEARCH_API FPoseSearchDatabaseBlendSpace : public FPoseSearchDatabase
 			&& A.NumberOfVerticalSamples == B.NumberOfVerticalSamples
 			&& A.BlendParamX == B.BlendParamX
 			&& A.BlendParamY == B.BlendParamY
+			&& A.SamplingRange == B.SamplingRange
 #endif // WITH_EDITORONLY_DATA
 			;
 	}
@@ -250,9 +259,9 @@ struct POSESEARCH_API FPoseSearchDatabaseAnimComposite : public FPoseSearchDatab
 	TObjectPtr<UAnimComposite> AnimComposite;
 
 #if WITH_EDITORONLY_DATA
-	// It allows users to set a time range to an individual animation sequence in the database. 
-	// This is effectively trimming the beginning and end of the animation in the database (not in the original sequence).
-	// If set to [0, 0] it will be the entire frame range of the original sequence.
+	// It allows users to set a time range to an individual animation composite in the database. 
+	// This is effectively trimming the beginning and end of the animation in the database (not in the original composite).
+	// If set to [0, 0] it will be the entire frame range of the original composite.
 	UPROPERTY(EditAnywhere, Category = "Settings", meta = (ClampToMinMaxLimits, DisplayPriority = 3))
 	FFloatInterval SamplingRange = FFloatInterval(0.f, 0.f);
 
@@ -291,9 +300,9 @@ struct POSESEARCH_API FPoseSearchDatabaseAnimMontage : public FPoseSearchDatabas
 	TObjectPtr<UAnimMontage> AnimMontage;
 
 #if WITH_EDITORONLY_DATA
-	// It allows users to set a time range to an individual animation sequence in the database. 
-	// This is effectively trimming the beginning and end of the animation in the database (not in the original sequence).
-	// If set to [0, 0] it will be the entire frame range of the original sequence.
+	// It allows users to set a time range to an individual animation montage in the database. 
+	// This is effectively trimming the beginning and end of the animation in the database (not in the original montage).
+	// If set to [0, 0] it will be the entire frame range of the original montage.
 	UPROPERTY(EditAnywhere, Category = "Settings", meta = (ClampToMinMaxLimits, DisplayPriority = 2))
 	FFloatInterval SamplingRange = FFloatInterval(0.f, 0.f);
 
@@ -331,9 +340,9 @@ struct POSESEARCH_API FPoseSearchDatabaseMultiAnimAsset : public FPoseSearchData
 	TObjectPtr<UMultiAnimAsset> MultiAnimAsset;
 
 #if WITH_EDITORONLY_DATA
-	// It allows users to set a time range to an individual animation sequence in the database. 
-	// This is effectively trimming the beginning and end of the animation in the database (not in the original sequence).
-	// If set to [0, 0] it will be the entire frame range of the original sequence.
+	// It allows users to set a time range to an individual UMultiAnimAsset in the database. 
+	// This is effectively trimming the beginning and end of the animation in the database (not in the original UMultiAnimAsset).
+	// If set to [0, 0] it will be the entire frame range of the original UMultiAnimAsset.
 	UPROPERTY(EditAnywhere, Category = "Settings", meta = (ClampToMinMaxLimits, DisplayPriority = 2))
 	FFloatInterval SamplingRange = FFloatInterval(0.f, 0.f);
 
