motion matching - adding NormalizationGroup to UPoseSearchFeatureChannel_Trajectory FPoseSearchTrajectorySample(s)

diff --git a/Engine/Plugins/Animation/PoseSearch/Source/Runtime/Private/PoseSearchFeatureChannel_Heading.h b/Engine/Plugins/Animation/PoseSearch/Source/Runtime/Private/PoseSearchFeatureChannel_Heading.h
--- a/Engine/Plugins/Animation/PoseSearch/Source/Runtime/Private/PoseSearchFeatureChannel_Heading.h
+++ b/Engine/Plugins/Animation/PoseSearch/Source/Runtime/Private/PoseSearchFeatureChannel_Heading.h
@@ -82,7 +82,7 @@ public:
 
 #if WITH_EDITORONLY_DATA
 	// if set, all the channels of the same class with the same cardinality, and the same NormalizationGroup, will be normalized together.
-	// for example in a locomotion database of a character holding a weapon, containing non mirrorable animations, you'd still want to normalize togeter 
+	// for example in a locomotion database of a character holding a weapon, containing non mirrorable animations, you'd still want to normalize together 
 	// left foot and right foot position and velocity
 	UPROPERTY(EditAnywhere, Category = "Settings")
 	FName NormalizationGroup;
diff --git a/Engine/Plugins/Animation/PoseSearch/Source/Runtime/Private/PoseSearchFeatureChannel_Phase.h b/Engine/Plugins/Animation/PoseSearch/Source/Runtime/Private/PoseSearchFeatureChannel_Phase.h
--- a/Engine/Plugins/Animation/PoseSearch/Source/Runtime/Private/PoseSearchFeatureChannel_Phase.h
+++ b/Engine/Plugins/Animation/PoseSearch/Source/Runtime/Private/PoseSearchFeatureChannel_Phase.h
@@ -37,7 +37,7 @@ public:
 
 #if WITH_EDITORONLY_DATA
 	// if set, all the channels of the same class with the same cardinality, and the same NormalizationGroup, will be normalized together.
-	// for example in a locomotion database of a character holding a weapon, containing non mirrorable animations, you'd still want to normalize togeter 
+	// for example in a locomotion database of a character holding a weapon, containing non mirrorable animations, you'd still want to normalize together 
 	// left foot and right foot position and velocity
 	UPROPERTY(EditAnywhere, Category = "Settings")
 	FName NormalizationGroup;
diff --git a/Engine/Plugins/Animation/PoseSearch/Source/Runtime/Private/PoseSearchFeatureChannel_Pose.h b/Engine/Plugins/Animation/PoseSearch/Source/Runtime/Private/PoseSearchFeatureChannel_Pose.h
--- a/Engine/Plugins/Animation/PoseSearch/Source/Runtime/Private/PoseSearchFeatureChannel_Pose.h
+++ b/Engine/Plugins/Animation/PoseSearch/Source/Runtime/Private/PoseSearchFeatureChannel_Pose.h
@@ -31,7 +31,7 @@ struct POSESEARCH_API FPoseSearchBone
 
 #if WITH_EDITORONLY_DATA
 	// if set, all the channels of the same class with the same cardinality, and the same NormalizationGroup, will be normalized together.
-	// for example in a locomotion database of a character holding a weapon, containing non mirrorable animations, you'd still want to normalize togeter 
+	// for example in a locomotion database of a character holding a weapon, containing non mirrorable animations, you'd still want to normalize together 
 	// left foot and right foot position and velocity
 	UPROPERTY(EditAnywhere, Category = "Settings")
 	FName NormalizationGroup;
diff --git a/Engine/Plugins/Animation/PoseSearch/Source/Runtime/Private/PoseSearchFeatureChannel_Position.h b/Engine/Plugins/Animation/PoseSearch/Source/Runtime/Private/PoseSearchFeatureChannel_Position.h
--- a/Engine/Plugins/Animation/PoseSearch/Source/Runtime/Private/PoseSearchFeatureChannel_Position.h
+++ b/Engine/Plugins/Animation/PoseSearch/Source/Runtime/Private/PoseSearchFeatureChannel_Position.h
@@ -69,7 +69,7 @@ public:
 
 #if WITH_EDITORONLY_DATA
 	// if set, all the channels of the same class with the same cardinality, and the same NormalizationGroup, will be normalized together.
-	// for example in a locomotion database of a character holding a weapon, containing non mirrorable animations, you'd still want to normalize togeter 
+	// for example in a locomotion database of a character holding a weapon, containing non mirrorable animations, you'd still want to normalize together 
 	// left foot and right foot position and velocity
 	UPROPERTY(EditAnywhere, Category = "Settings")
 	FName NormalizationGroup;
diff --git a/Engine/Plugins/Animation/PoseSearch/Source/Runtime/Private/PoseSearchFeatureChannel_Trajectory.cpp b/Engine/Plugins/Animation/PoseSearch/Source/Runtime/Private/PoseSearchFeatureChannel_Trajectory.cpp
--- a/Engine/Plugins/Animation/PoseSearch/Source/Runtime/Private/PoseSearchFeatureChannel_Trajectory.cpp
+++ b/Engine/Plugins/Animation/PoseSearch/Source/Runtime/Private/PoseSearchFeatureChannel_Trajectory.cpp
@@ -21,25 +21,25 @@ UPoseSearchFeatureChannel_Trajectory::UPoseSearchFeatureChannel_Trajectory()
 
 	Samples.Add(FPoseSearchTrajectorySample({ -0.4f, int32(EPoseSearchTrajectoryFlags::PositionXY)
 #if WITH_EDITORONLY_DATA
-		, 0.4f, FLinearColor::Red
+		, 0.4f, FName(), FLinearColor::Red
 #endif // WITH_EDITORONLY_DATA
 		}));
 
 	Samples.Add(FPoseSearchTrajectorySample({ 0.f, int32(EPoseSearchTrajectoryFlags::VelocityXY | EPoseSearchTrajectoryFlags::FacingDirectionXY)
 #if WITH_EDITORONLY_DATA
-		, 2.f, FLinearColor::Blue
+		, 2.f, FName(), FLinearColor::Blue
 #endif // WITH_EDITORONLY_DATA
 		}));
 
 	Samples.Add(FPoseSearchTrajectorySample({ 0.35f, int32(EPoseSearchTrajectoryFlags::PositionXY | EPoseSearchTrajectoryFlags::FacingDirectionXY)
 #if WITH_EDITORONLY_DATA
-		, 0.7f, FLinearColor::Blue
+		, 0.7f, FName(), FLinearColor::Blue
 #endif // WITH_EDITORONLY_DATA
 		}));
 
 	Samples.Add(FPoseSearchTrajectorySample({ 0.7f, int32(EPoseSearchTrajectoryFlags::VelocityXY | EPoseSearchTrajectoryFlags::PositionXY | EPoseSearchTrajectoryFlags::FacingDirectionXY)
 #if WITH_EDITORONLY_DATA
-		, 0.5f, FLinearColor::Blue
+		, 0.5f, FName(), FLinearColor::Blue
 #endif // WITH_EDITORONLY_DATA
 		}));
 }
@@ -57,6 +57,7 @@ bool UPoseSearchFeatureChannel_Trajectory::Finalize(UPoseSearchSchema* Schema)
 			Position->OriginRole = SampleRole;
 #if WITH_EDITORONLY_DATA
 			Position->Weight = Sample.Weight * Weight;
+			Position->NormalizationGroup = Sample.NormalizationGroup;
 			Position->DebugColor = Sample.DebugColor;
 #endif // WITH_EDITORONLY_DATA
 			Position->SampleTimeOffset = Sample.Offset;
@@ -76,6 +77,7 @@ bool UPoseSearchFeatureChannel_Trajectory::Finalize(UPoseSearchSchema* Schema)
 			Velocity->OriginRole = SampleRole;
 #if WITH_EDITORONLY_DATA
 			Velocity->Weight = Sample.Weight * Weight;
+			Velocity->NormalizationGroup = Sample.NormalizationGroup;
 			Velocity->DebugColor = Sample.DebugColor;
 #endif // WITH_EDITORONLY_DATA
 			Velocity->SampleTimeOffset = Sample.Offset;
@@ -95,6 +97,7 @@ bool UPoseSearchFeatureChannel_Trajectory::Finalize(UPoseSearchSchema* Schema)
 			Velocity->OriginRole = SampleRole;
 #if WITH_EDITORONLY_DATA
 			Velocity->Weight = Sample.Weight * Weight;
+			Velocity->NormalizationGroup = Sample.NormalizationGroup;
 			Velocity->DebugColor = Sample.DebugColor;
 #endif // WITH_EDITORONLY_DATA
 			Velocity->SampleTimeOffset = Sample.Offset;
@@ -115,6 +118,7 @@ bool UPoseSearchFeatureChannel_Trajectory::Finalize(UPoseSearchSchema* Schema)
 			Heading->OriginRole = SampleRole;
 #if WITH_EDITORONLY_DATA
 			Heading->Weight = Sample.Weight * Weight;
+			Heading->NormalizationGroup = Sample.NormalizationGroup;
 			Heading->DebugColor = Sample.DebugColor;
 #endif // WITH_EDITORONLY_DATA
 			Heading->SampleTimeOffset = Sample.Offset;
diff --git a/Engine/Plugins/Animation/PoseSearch/Source/Runtime/Private/PoseSearchFeatureChannel_Trajectory.h b/Engine/Plugins/Animation/PoseSearch/Source/Runtime/Private/PoseSearchFeatureChannel_Trajectory.h
--- a/Engine/Plugins/Animation/PoseSearch/Source/Runtime/Private/PoseSearchFeatureChannel_Trajectory.h
+++ b/Engine/Plugins/Animation/PoseSearch/Source/Runtime/Private/PoseSearchFeatureChannel_Trajectory.h
@@ -40,6 +40,12 @@ struct POSESEARCH_API FPoseSearchTrajectorySample
 	UPROPERTY(EditAnywhere, Category = Config)
 	float Weight = 1.f;
 
+	// if set, all the channels of the same class with the same cardinality, and the same NormalizationGroup, will be normalized together.
+	// for example in a locomotion database of a character holding a weapon, containing non mirrorable animations, you'd still want to normalize together 
+	// left foot and right foot position and velocity
+	UPROPERTY(EditAnywhere, Category = Config)
+	FName NormalizationGroup;
+
 	UPROPERTY(EditAnywhere, Category = Config, meta = (ExcludeFromHash, DisplayPriority = 0))
 	FLinearColor DebugColor = FLinearColor::Blue;
 #endif // WITH_EDITORONLY_DATA
diff --git a/Engine/Plugins/Animation/PoseSearch/Source/Runtime/Private/PoseSearchFeatureChannel_Velocity.h b/Engine/Plugins/Animation/PoseSearch/Source/Runtime/Private/PoseSearchFeatureChannel_Velocity.h
--- a/Engine/Plugins/Animation/PoseSearch/Source/Runtime/Private/PoseSearchFeatureChannel_Velocity.h
+++ b/Engine/Plugins/Animation/PoseSearch/Source/Runtime/Private/PoseSearchFeatureChannel_Velocity.h
@@ -77,7 +77,7 @@ public:
 
 #if WITH_EDITORONLY_DATA
 	// if set, all the channels of the same class with the same cardinality, and the same NormalizationGroup, will be normalized together.
-	// for example in a locomotion database of a character holding a weapon, containing non mirrorable animations, you'd still want to normalize togeter 
+	// for example in a locomotion database of a character holding a weapon, containing non mirrorable animations, you'd still want to normalize together 
 	// left foot and right foot position and velocity
 	UPROPERTY(EditAnywhere, Category = "Settings")
 	FName NormalizationGroup;
diff --git a/Engine/Plugins/Animation/PoseSearch/Source/Runtime/Public/PoseSearch/PoseSearchFeatureChannel.h b/Engine/Plugins/Animation/PoseSearch/Source/Runtime/Public/PoseSearch/PoseSearchFeatureChannel.h
--- a/Engine/Plugins/Animation/PoseSearch/Source/Runtime/Public/PoseSearch/PoseSearchFeatureChannel.h
+++ b/Engine/Plugins/Animation/PoseSearch/Source/Runtime/Public/PoseSearch/PoseSearchFeatureChannel.h
@@ -157,7 +157,7 @@ public:
 
 	// if this channel GetNormalizationGroup returns a valid FName, all the channels of the same class with the same cardinality, and the same NormalizationGroup, 
 	// will make CanBeNormalizedWith return true and will be normalized together.
-	// for example in a locomotion database of a character holding a weapon, containing non mirrorable animations, you'd still want to normalize togeter 
+	// for example in a locomotion database of a character holding a weapon, containing non mirrorable animations, you'd still want to normalize together 
 	// left foot and right foot position and velocity, so you'd want those channels returning the same GetNormalizationGroup value
 	virtual FName GetNormalizationGroup() const { return FName(); }
 
