First pass at having a Mover-based character using motion matching in the animation sample project.  Added concept of an external trajectory predictor for PoseSearch, and implemented one for MoverComponent actors. Added a Mover version of the sandbox character actor BP and anim BP, selectable in the GameAnim widget.

#virtualized

diff --git a/Engine/Plugins/Animation/PoseSearch/Source/Runtime/Private/PoseSearchHistory.cpp b/Engine/Plugins/Animation/PoseSearch/Source/Runtime/Private/PoseSearchHistory.cpp
--- a/Engine/Plugins/Animation/PoseSearch/Source/Runtime/Private/PoseSearchHistory.cpp
+++ b/Engine/Plugins/Animation/PoseSearch/Source/Runtime/Private/PoseSearchHistory.cpp
@@ -657,8 +657,8 @@ void FPoseHistory::GenerateTrajectory(const UAnimInstance* AnimInstance, float D
 	// @todo: Synchronize the FPoseSearchQueryTrajectorySample::AccumulatedSeconds of the generated trajectory with the FPoseHistoryEntry::AccumulatedSeconds of the captured poses
 	FPoseSearchTrajectoryData::FDerived TrajectoryDataDerived;
 	TrajectoryData.UpdateData(DeltaTime, AnimInstance, TrajectoryDataDerived, TrajectoryDataState);
-	UPoseSearchTrajectoryLibrary::InitTrajectorySamples(Trajectory, TrajectoryData, TrajectoryDataDerived, TrajectoryDataSampling, DeltaTime);
-	UPoseSearchTrajectoryLibrary::UpdateHistory_TransformHistory(Trajectory, TrajectoryData, TrajectoryDataDerived, TrajectoryDataSampling, DeltaTime);
+	UPoseSearchTrajectoryLibrary::InitTrajectorySamples(Trajectory, TrajectoryData, TrajectoryDataDerived.Position, TrajectoryDataDerived.Facing, TrajectoryDataSampling, DeltaTime);
+	UPoseSearchTrajectoryLibrary::UpdateHistory_TransformHistory(Trajectory, TrajectoryData, TrajectoryDataDerived.Position, TrajectoryDataDerived.Velocity, TrajectoryDataSampling, DeltaTime);
 	UPoseSearchTrajectoryLibrary::UpdatePrediction_SimulateCharacterMovement(Trajectory, TrajectoryData, TrajectoryDataDerived, TrajectoryDataSampling, DeltaTime);
 
 	// @todo: support TrajectorySpeedMultiplier
diff --git a/Engine/Plugins/Animation/PoseSearch/Source/Runtime/Private/PoseSearchTrajectoryLibrary.cpp b/Engine/Plugins/Animation/PoseSearch/Source/Runtime/Private/PoseSearchTrajectoryLibrary.cpp
--- a/Engine/Plugins/Animation/PoseSearch/Source/Runtime/Private/PoseSearchTrajectoryLibrary.cpp
+++ b/Engine/Plugins/Animation/PoseSearch/Source/Runtime/Private/PoseSearchTrajectoryLibrary.cpp
@@ -1,6 +1,7 @@
 // Copyright Epic Games, Inc. All Rights Reserved.
 
 #include "PoseSearch/PoseSearchTrajectoryLibrary.h"
+#include "PoseSearch/PoseSearchTrajectoryPredictor.h"
 #include "Animation/AnimInstanceProxy.h"
 #include "Components/SkeletalMeshComponent.h"
 #include "GameFramework/Character.h"
@@ -165,6 +166,21 @@ void UPoseSearchTrajectoryLibrary::InitTrajectorySamples(
 	const FPoseSearchTrajectoryData::FDerived& TrajectoryDataDerived,
 	const FPoseSearchTrajectoryData::FSampling& TrajectoryDataSampling,
 	float DeltaTime)
+{
+	InitTrajectorySamples(
+		Trajectory,
+		TrajectoryData,
+		TrajectoryDataDerived.Position, TrajectoryDataDerived.Facing,
+		TrajectoryDataSampling, 
+		DeltaTime );
+}
+
+void UPoseSearchTrajectoryLibrary::InitTrajectorySamples(
+	FPoseSearchQueryTrajectory& Trajectory,
+	const FPoseSearchTrajectoryData& TrajectoryData,
+	FVector DefaultPosition, FQuat DefaultFacing,
+	const FPoseSearchTrajectoryData::FSampling& TrajectoryDataSampling,
+	float DeltaTime)
 {
 	const int32 NumHistorySamples = TrajectoryDataSampling.NumHistorySamples;
 	const int32 NumPredictionSamples = TrajectoryDataSampling.NumPredictionSamples;
@@ -180,8 +196,8 @@ void UPoseSearchTrajectoryLibrary::InitTrajectorySamples(
 		const float SecondsPerHistorySample = FMath::Max(TrajectoryDataSampling.SecondsPerHistorySample, 0.f);
 		for (int32 i = 0; i < NumHistorySamples; ++i)
 		{
-			Trajectory.Samples[i].Position = TrajectoryDataDerived.Position;
-			Trajectory.Samples[i].Facing = TrajectoryDataDerived.Facing;
+			Trajectory.Samples[i].Position = DefaultPosition;
+			Trajectory.Samples[i].Facing = DefaultFacing;
 			Trajectory.Samples[i].AccumulatedSeconds = SecondsPerHistorySample * (i - NumHistorySamples - 1);
 		}
 
@@ -189,8 +205,8 @@ void UPoseSearchTrajectoryLibrary::InitTrajectorySamples(
 		const float SecondsPerPredictionSample = FMath::Max(TrajectoryDataSampling.SecondsPerPredictionSample, 0.f);
 		for (int32 i = NumHistorySamples; i < Trajectory.Samples.Num(); ++i)
 		{
-			Trajectory.Samples[i].Position = TrajectoryDataDerived.Position;
-			Trajectory.Samples[i].Facing = TrajectoryDataDerived.Facing;
+			Trajectory.Samples[i].Position = DefaultPosition;
+			Trajectory.Samples[i].Facing = DefaultFacing;
 			Trajectory.Samples[i].AccumulatedSeconds = SecondsPerPredictionSample * (i - NumHistorySamples) + DeltaTime;
 		}
 	}
@@ -202,6 +218,22 @@ void UPoseSearchTrajectoryLibrary::UpdateHistory_TransformHistory(
 	const FPoseSearchTrajectoryData::FDerived& TrajectoryDataDerived,
 	const FPoseSearchTrajectoryData::FSampling& TrajectoryDataSampling,
 	float DeltaTime)
+{
+	UpdateHistory_TransformHistory(
+		Trajectory,
+		TrajectoryData,
+		TrajectoryDataDerived.Position, TrajectoryDataDerived.Velocity,
+		TrajectoryDataSampling,
+		DeltaTime);
+}
+
+void UPoseSearchTrajectoryLibrary::UpdateHistory_TransformHistory(
+	FPoseSearchQueryTrajectory& Trajectory,
+	const FPoseSearchTrajectoryData& TrajectoryData,
+	FVector CurrentPosition,
+	FVector CurrentVelocity, 
+	const FPoseSearchTrajectoryData::FSampling& TrajectoryDataSampling,
+	float DeltaTime)
 {
 	const int32 NumHistorySamples = TrajectoryDataSampling.NumHistorySamples;
 	if (NumHistorySamples > 0)
@@ -216,7 +248,7 @@ void UPoseSearchTrajectoryLibrary::UpdateHistory_TransformHistory(
 			Trajectory.Samples[Index].Position = Trajectory.Samples[NumHistorySamples].Position - Trajectory.Samples[Index].Position;
 		}
 
-		FVector CurrentTranslation = TrajectoryDataDerived.Velocity * DeltaTime;
+		FVector CurrentTranslation = CurrentVelocity * DeltaTime;
 
 		// Shift history Samples when it's time to record a new one.
 		if (SecondsPerHistorySample <= 0.f || FMath::Abs(Trajectory.Samples[NumHistorySamples - 1].AccumulatedSeconds) >= SecondsPerHistorySample)
@@ -244,7 +276,7 @@ void UPoseSearchTrajectoryLibrary::UpdateHistory_TransformHistory(
 		// converting the history sample positions in world space by applying the current world position.
 		for (int32 Index = 0; Index < NumHistorySamples; ++Index)
 		{
-			Trajectory.Samples[Index].Position = TrajectoryDataDerived.Position - Trajectory.Samples[Index].Position;
+			Trajectory.Samples[Index].Position = CurrentPosition - Trajectory.Samples[Index].Position;
 		}
 	}
 }
@@ -349,8 +381,8 @@ void UPoseSearchTrajectoryLibrary::UpdatePrediction_SimulateCharacterMovement(
 }
 
 void UPoseSearchTrajectoryLibrary::PoseSearchGenerateTrajectory(
-	const UAnimInstance* InAnimInstance, UPARAM(ref)
-	const FPoseSearchTrajectoryData& InTrajectoryData,
+	const UAnimInstance* InAnimInstance, 
+	UPARAM(ref)	const FPoseSearchTrajectoryData& InTrajectoryData,
 	float InDeltaTime,
 	UPARAM(ref) FPoseSearchQueryTrajectory& InOutTrajectory,
 	UPARAM(ref) float& InOutDesiredControllerYawLastUpdate,
@@ -371,8 +403,8 @@ void UPoseSearchTrajectoryLibrary::PoseSearchGenerateTrajectory(
 
 	FPoseSearchTrajectoryData::FDerived TrajectoryDataDerived;
 	InTrajectoryData.UpdateData(InDeltaTime, InAnimInstance, TrajectoryDataDerived, TrajectoryDataState);
-	InitTrajectorySamples(InOutTrajectory, InTrajectoryData, TrajectoryDataDerived, TrajectoryDataSampling, InDeltaTime);
-	UpdateHistory_TransformHistory(InOutTrajectory, InTrajectoryData, TrajectoryDataDerived, TrajectoryDataSampling, InDeltaTime);
+	InitTrajectorySamples(InOutTrajectory, InTrajectoryData, TrajectoryDataDerived.Position, TrajectoryDataDerived.Facing, TrajectoryDataSampling, InDeltaTime);
+	UpdateHistory_TransformHistory(InOutTrajectory, InTrajectoryData, TrajectoryDataDerived.Position, TrajectoryDataDerived.Velocity, TrajectoryDataSampling, InDeltaTime);
 	UpdatePrediction_SimulateCharacterMovement(InOutTrajectory, InTrajectoryData, TrajectoryDataDerived, TrajectoryDataSampling, InDeltaTime);
 
 	InOutDesiredControllerYawLastUpdate = TrajectoryDataState.DesiredControllerYawLastUpdate;
@@ -380,29 +412,97 @@ void UPoseSearchTrajectoryLibrary::PoseSearchGenerateTrajectory(
 	OutTrajectory = InOutTrajectory;
 }
 
-void UPoseSearchTrajectoryLibrary::HandleTrajectoryWorldCollisions(const UObject* WorldContextObject, const UAnimInstance* AnimInstance, UPARAM(ref) const FPoseSearchQueryTrajectory& InTrajectory, bool bApplyGravity, float FloorCollisionsOffset, FPoseSearchQueryTrajectory& OutTrajectory, FPoseSearchTrajectory_WorldCollisionResults& CollisionResult,
-	ETraceTypeQuery TraceChannel, bool bTraceComplex, const TArray<AActor*>& ActorsToIgnore, EDrawDebugTrace::Type DrawDebugType, bool bIgnoreSelf, float MaxObstacleHeight, FLinearColor TraceColor, FLinearColor TraceHitColor, float DrawTime)
+
+void UPoseSearchTrajectoryLibrary::PoseSearchGeneratePredictorTrajectory(
+	UObject* InPredictor,	// must implement IPoseSearchTrajectoryPredictorInterface
+	UPARAM(ref)	const FPoseSearchTrajectoryData& InTrajectoryData,
+	float InDeltaTime,
+	UPARAM(ref) FPoseSearchQueryTrajectory& InOutTrajectory,
+	UPARAM(ref) float& InOutDesiredControllerYawLastUpdate,
+	FPoseSearchQueryTrajectory& OutTrajectory,
+	float InHistorySamplingInterval,
+	int32 InTrajectoryHistoryCount,
+	float InPredictionSamplingInterval,
+	int32 InTrajectoryPredictionCount)
 {
-	OutTrajectory = InTrajectory;
+	FPoseSearchTrajectoryData::FSampling TrajectoryDataSampling;
+	TrajectoryDataSampling.NumHistorySamples = InTrajectoryHistoryCount;
+	TrajectoryDataSampling.SecondsPerHistorySample = InHistorySamplingInterval;
+	TrajectoryDataSampling.NumPredictionSamples = InTrajectoryPredictionCount;
+	TrajectoryDataSampling.SecondsPerPredictionSample = InPredictionSamplingInterval;
 
-	TArray<FPoseSearchQueryTrajectorySample>& Samples = OutTrajectory.Samples;
-	const int32 NumSamples = Samples.Num();
+	// TODO: handle controller yaw
+	//TrajectoryDataState.DesiredControllerYawLastUpdate = InOutDesiredControllerYawLastUpdate;
 
-	FVector GravityDirection = FVector::ZeroVector;
-	float GravityZ = 0.f;
-	float InitialVelocityZ = 0.0f;
+	FVector CurrentPosition = FVector::ZeroVector;
+	FVector CurrentVelocity = FVector::ZeroVector;
+	FQuat CurrentFacing = FQuat::Identity;
+
+	if (InPredictor)
+	{
+		IPoseSearchTrajectoryPredictorInterface::Execute_GetCurrentState(InPredictor, CurrentPosition, CurrentFacing, CurrentVelocity);
+	}
+
+	InitTrajectorySamples(InOutTrajectory, InTrajectoryData, CurrentPosition, CurrentFacing, TrajectoryDataSampling, InDeltaTime);
+	UpdateHistory_TransformHistory(InOutTrajectory, InTrajectoryData, CurrentPosition, CurrentVelocity, TrajectoryDataSampling, InDeltaTime);
+
+	if (InPredictor)
+	{
+		IPoseSearchTrajectoryPredictorInterface::Execute_Predict(InPredictor, InOutTrajectory,
+			InTrajectoryPredictionCount + 1, InPredictionSamplingInterval, InTrajectoryHistoryCount);
+	}
+
+	//InOutDesiredControllerYawLastUpdate = TrajectoryDataState.DesiredControllerYawLastUpdate;
+
+	OutTrajectory = InOutTrajectory;
+}
+
+
+void UPoseSearchTrajectoryLibrary::HandleTrajectoryWorldCollisions(const UObject* WorldContextObject, const UAnimInstance* AnimInstance, UPARAM(ref) const FPoseSearchQueryTrajectory& InTrajectory, bool bApplyGravity, float FloorCollisionsOffset, FPoseSearchQueryTrajectory& OutTrajectory, FPoseSearchTrajectory_WorldCollisionResults& CollisionResult,
+	ETraceTypeQuery TraceChannel, bool bTraceComplex, const TArray<AActor*>& ActorsToIgnore, EDrawDebugTrace::Type DrawDebugType, bool bIgnoreSelf, float MaxObstacleHeight, FLinearColor TraceColor, FLinearColor TraceHitColor, float DrawTime)
+{
+	FVector StartingVelocity = FVector::ZeroVector;
+	FVector GravityAccel = FVector::ZeroVector;
 	if (bApplyGravity && AnimInstance)
 	{
 		if (const ACharacter* Character = Cast<ACharacter>(AnimInstance->GetOwningActor()))
 		{
 			if (const UCharacterMovementComponent* MoveComp = Character->GetCharacterMovement())
 			{
-				GravityZ = MoveComp->GetGravityZ();
-				GravityDirection = MoveComp->GetGravityDirection();
-				InitialVelocityZ = Character->GetVelocity().Z;
+				GravityAccel = MoveComp->GetGravityDirection() * -MoveComp->GetGravityZ();
+				StartingVelocity = Character->GetVelocity();
 			}
 		}
 	}
+
+	HandleTrajectoryWorldCollisionsWithGravity(WorldContextObject, InTrajectory, StartingVelocity, bApplyGravity, GravityAccel, 
+		FloorCollisionsOffset, OutTrajectory, CollisionResult, TraceChannel, bTraceComplex, ActorsToIgnore, 
+		DrawDebugType, bIgnoreSelf, MaxObstacleHeight, TraceColor, TraceHitColor, DrawTime);
+}
+
+
+void UPoseSearchTrajectoryLibrary::HandleTrajectoryWorldCollisionsWithGravity(const UObject* WorldContextObject,
+	UPARAM(ref) const FPoseSearchQueryTrajectory& InTrajectory, FVector StartingVelocity, bool bApplyGravity, FVector GravityAccel, float FloorCollisionsOffset, FPoseSearchQueryTrajectory& OutTrajectory, FPoseSearchTrajectory_WorldCollisionResults& CollisionResult,
+	ETraceTypeQuery TraceChannel, bool bTraceComplex, const TArray<AActor*>& ActorsToIgnore, EDrawDebugTrace::Type DrawDebugType, bool bIgnoreSelf, float MaxObstacleHeight, FLinearColor TraceColor, FLinearColor TraceHitColor, float DrawTime)
+{
+	OutTrajectory = InTrajectory;
+
+	TArray<FPoseSearchQueryTrajectorySample>& Samples = OutTrajectory.Samples;
+	const int32 NumSamples = Samples.Num();
+
+	FVector GravityDirection = FVector::ZeroVector;
+	float GravityZ = 0.f;
+	float InitialVelocityZ = StartingVelocity.Z;
+
+	if (bApplyGravity && !GravityAccel.IsNearlyZero())
+	{
+		GravityAccel.ToDirectionAndLength(GravityDirection, GravityZ);
+		GravityZ = -GravityZ;
+		const FVector VelocityOnGravityAxis = StartingVelocity.ProjectOnTo(GravityDirection);
+		
+		InitialVelocityZ = VelocityOnGravityAxis.Length() * FMath::Sign(GravityDirection.Dot(VelocityOnGravityAxis));
+	}
+
 	CollisionResult.TimeToLand = OutTrajectory.Samples.Last().AccumulatedSeconds;
 
 	if (!FMath::IsNearlyZero(GravityZ))
diff --git a/Engine/Plugins/Animation/PoseSearch/Source/Runtime/Public/PoseSearch/PoseSearchTrajectoryLibrary.h b/Engine/Plugins/Animation/PoseSearch/Source/Runtime/Public/PoseSearch/PoseSearchTrajectoryLibrary.h
--- a/Engine/Plugins/Animation/PoseSearch/Source/Runtime/Public/PoseSearch/PoseSearchTrajectoryLibrary.h
+++ b/Engine/Plugins/Animation/PoseSearch/Source/Runtime/Public/PoseSearch/PoseSearchTrajectoryLibrary.h
@@ -8,6 +8,7 @@
 #include "PoseSearchTrajectoryLibrary.generated.h"
 
 class UAnimInstance;
+class IPoseSearchTrajectoryPredictorInterface;
 
 USTRUCT(BlueprintType)
 struct POSESEARCH_API FPoseSearchTrajectoryData
@@ -103,19 +104,35 @@ class POSESEARCH_API UPoseSearchTrajectoryLibrary : public UBlueprintFunctionLib
 	GENERATED_BODY()
 
 public:
+
+	UE_DEPRECATED(5.5, "Use the InitTrajectorySamples function that takes DefaultPosition and DefaultFacing arguments instead of TrajectoryDataDerived")
 	static void InitTrajectorySamples(FPoseSearchQueryTrajectory& Trajectory, const FPoseSearchTrajectoryData& TrajectoryData, const FPoseSearchTrajectoryData::FDerived& TrajectoryDataDerived, const FPoseSearchTrajectoryData::FSampling& TrajectoryDataSampling, float DeltaTime);
 
+	// Initialize history and predicted samples based on sampling settings and a default state
+	static void InitTrajectorySamples(FPoseSearchQueryTrajectory& Trajectory, const FPoseSearchTrajectoryData& TrajectoryData, FVector DefaultPosition, FQuat DefaultFacing, const FPoseSearchTrajectoryData::FSampling& TrajectoryDataSampling, float DeltaTime);
+	
+	UE_DEPRECATED(5.5, "Use the UpdateHistory_TransformHistory function that takes CurrentPosition and CurrentVelocity arguments instead of TrajectoryDataDerived")
+	static void UpdateHistory_TransformHistory(FPoseSearchQueryTrajectory& Trajectory, const FPoseSearchTrajectoryData& TrajectoryData, const FPoseSearchTrajectoryData::FDerived& TrajectoryDataDerived, const FPoseSearchTrajectoryData::FSampling& TrajectoryDataSampling, float DeltaTime);
+
 	// Update history by tracking offsets that result from character intent (e.g. movement component velocity) and applying
 	// that to the current world transform. This works well on moving platforms as it only stores a history of movement
 	// that results from character intent, not movement from platforms.
-	static void UpdateHistory_TransformHistory(FPoseSearchQueryTrajectory& Trajectory, const FPoseSearchTrajectoryData& TrajectoryData, const FPoseSearchTrajectoryData::FDerived& TrajectoryDataDerived, const FPoseSearchTrajectoryData::FSampling& TrajectoryDataSampling, float DeltaTime);
+	static void UpdateHistory_TransformHistory(FPoseSearchQueryTrajectory& Trajectory, const FPoseSearchTrajectoryData& TrajectoryData, FVector CurrentPosition, FVector CurrentVelocity, const FPoseSearchTrajectoryData::FSampling& TrajectoryDataSampling, float DeltaTime);
 
 	// Update prediction by simulating the movement math for ground locomotion from UCharacterMovementComponent.
 	static void UpdatePrediction_SimulateCharacterMovement(FPoseSearchQueryTrajectory& Trajectory, const FPoseSearchTrajectoryData& TrajectoryData, const FPoseSearchTrajectoryData::FDerived& TrajectoryDataDerived, const FPoseSearchTrajectoryData::FSampling& TrajectoryDataSampling, float DeltaTime);
 
-	// Generates a prediction trajectory based of the current character intent.
-	UFUNCTION(BlueprintCallable, Category = "Animation|PoseSearch", meta = (BlueprintThreadSafe, DisplayName = "Pose Search Generate Trajectory"))
-	static void PoseSearchGenerateTrajectory(const UAnimInstance* InAnimInstance, UPARAM(ref) const FPoseSearchTrajectoryData& InTrajectoryData, float InDeltaTime,
+	// Generates a prediction trajectory based of the current character intent. For use with Character actors.
+	UFUNCTION(BlueprintCallable, Category = "Animation|PoseSearch", meta = (BlueprintThreadSafe, DisplayName = "Pose Search Generate Trajectory (for Character)"))
+	static void PoseSearchGenerateTrajectory(const UAnimInstance* InAnimInstance, 
+		UPARAM(ref) const FPoseSearchTrajectoryData& InTrajectoryData, float InDeltaTime,
+		UPARAM(ref) FPoseSearchQueryTrajectory& InOutTrajectory, UPARAM(ref) float& InOutDesiredControllerYawLastUpdate, FPoseSearchQueryTrajectory& OutTrajectory,
+		float InHistorySamplingInterval = 0.04f, int32 InTrajectoryHistoryCount = 10, float InPredictionSamplingInterval = 0.2f, int32 InTrajectoryPredictionCount = 8);
+
+	// Generates a prediction trajectory based of the current movement intent. For use with predictors. InPredictor must implement IPoseSearchTrajectoryPredictorInterface
+	UFUNCTION(BlueprintCallable, Category = "Animation|PoseSearch|Experimental", meta = (BlueprintThreadSafe, DisplayName = "Pose Search Generate Trajectory (using Predictor)"))
+	static void PoseSearchGeneratePredictorTrajectory(UObject* InPredictor,
+		UPARAM(ref) const FPoseSearchTrajectoryData& InTrajectoryData, float InDeltaTime,
 		UPARAM(ref) FPoseSearchQueryTrajectory& InOutTrajectory, UPARAM(ref) float& InOutDesiredControllerYawLastUpdate, FPoseSearchQueryTrajectory& OutTrajectory,
 		float InHistorySamplingInterval = 0.04f, int32 InTrajectoryHistoryCount = 10, float InPredictionSamplingInterval = 0.2f, int32 InTrajectoryPredictionCount = 8);
 
@@ -126,6 +143,14 @@ public:
 	static void HandleTrajectoryWorldCollisions(const UObject* WorldContextObject, const UAnimInstance* AnimInstance, UPARAM(ref) const FPoseSearchQueryTrajectory& InTrajectory, bool bApplyGravity, float FloorCollisionsOffset, FPoseSearchQueryTrajectory& OutTrajectory, FPoseSearchTrajectory_WorldCollisionResults& CollisionResult,
 		ETraceTypeQuery TraceChannel, bool bTraceComplex, const TArray<AActor*>& ActorsToIgnore, EDrawDebugTrace::Type DrawDebugType, bool bIgnoreSelf = true, float MaxObstacleHeight = 10000.f, FLinearColor TraceColor = FLinearColor::Red, FLinearColor TraceHitColor = FLinearColor::Green, float DrawTime = 5.0f);
 
+	// Experimental: Process InTrajectory to apply gravity and handle collisions. Eventually returns the modified OutTrajectory.
+	// If bApplyGravity is true, GravityAccel will be applied.
+	// If FloorCollisionsOffset > 0, vertical collision will be performed to every sample of the trajectory to have the samples float over the geometry (by FloorCollisionsOffset).
+	UFUNCTION(BlueprintCallable, Category="Animation|PoseSearch|Experimental", meta=(WorldContext="WorldContextObject", AutoCreateRefTerm="ActorsToIgnore", AdvancedDisplay="TraceChannel,bTraceComplex,ActorsToIgnore,DrawDebugType,bIgnoreSelf,MaxObstacleHeight,TraceColor,TraceHitColor,DrawTime"))
+	static void HandleTrajectoryWorldCollisionsWithGravity(const UObject* WorldContextObject, 
+		UPARAM(ref) const FPoseSearchQueryTrajectory& InTrajectory, FVector StartingVelocity, bool bApplyGravity, FVector GravityAccel, float FloorCollisionsOffset, FPoseSearchQueryTrajectory& OutTrajectory, FPoseSearchTrajectory_WorldCollisionResults& CollisionResult,
+		ETraceTypeQuery TraceChannel, bool bTraceComplex, const TArray<AActor*>& ActorsToIgnore, EDrawDebugTrace::Type DrawDebugType, bool bIgnoreSelf = true, float MaxObstacleHeight = 10000.f, FLinearColor TraceColor = FLinearColor::Red, FLinearColor TraceHitColor = FLinearColor::Green, float DrawTime = 5.0f);
+
 	UFUNCTION(BlueprintPure, meta = (BlueprintThreadSafe), Category="Animation|PoseSearch")
 	static void GetTrajectorySampleAtTime(UPARAM(ref) const FPoseSearchQueryTrajectory& InTrajectory, float Time, FPoseSearchQueryTrajectorySample& OutTrajectorySample, bool bExtrapolate = false);
 
diff --git a/Engine/Plugins/Animation/PoseSearch/Source/Runtime/Public/PoseSearch/PoseSearchTrajectoryPredictor.h b/Engine/Plugins/Animation/PoseSearch/Source/Runtime/Public/PoseSearch/PoseSearchTrajectoryPredictor.h
new file mode 100644
--- /dev/null
+++ b/Engine/Plugins/Animation/PoseSearch/Source/Runtime/Public/PoseSearch/PoseSearchTrajectoryPredictor.h
@@ -0,0 +1,38 @@
+// Copyright Epic Games, Inc. All Rights Reserved.
+
+#pragma once
+
+#include "UObject/Interface.h"
+#include "PoseSearchTrajectoryPredictor.generated.h"
+
+struct FPoseSearchQueryTrajectory;
+
+
+UINTERFACE(BlueprintType)
+class POSESEARCH_API UPoseSearchTrajectoryPredictorInterface : public UInterface
+{
+	GENERATED_BODY()
+};
+
+/**
+ * PoseSearchTrajectoryPredictor: API for an object to implement to act as a predictor of the future trajectory for motion matching animation purposes
+ */
+class IPoseSearchTrajectoryPredictorInterface : public IInterface
+{
+	GENERATED_BODY()
+
+public:
+
+	UFUNCTION(BlueprintNativeEvent)
+	void Predict(FPoseSearchQueryTrajectory& InOutTrajectory, int32 NumPredictionSamples, float SecondsPerPredictionSample, int32 NumHistorySamples);
+
+	UFUNCTION(BlueprintCallable, BlueprintNativeEvent, Category="Animation|PoseSearch|Experimental")
+	void GetGravity(FVector& OutGravityAccel);
+
+	UFUNCTION(BlueprintCallable, BlueprintNativeEvent, Category = "Animation|PoseSearch|Experimental")
+	void GetCurrentState(FVector& OutPosition, FQuat& OutFacing, FVector& OutVelocity);
+
+	UFUNCTION(BlueprintCallable, BlueprintNativeEvent, Category = "Animation|PoseSearch|Experimental")
+	void GetVelocity(FVector& OutVelocity);
+
+};
diff --git a/Engine/Plugins/Experimental/Mover/Mover.uplugin b/Engine/Plugins/Experimental/Mover/Mover.uplugin
--- a/Engine/Plugins/Experimental/Mover/Mover.uplugin
+++ b/Engine/Plugins/Experimental/Mover/Mover.uplugin
@@ -36,8 +36,12 @@
 			"Enabled": true
 		},
 		{
-			"Name": "Water",
+			"Name": "PoseSearch",
 			"Enabled": true
 		},
+		{
+			"Name": "Water",
+			"Enabled": true
+		}
 	]
 }
\ No newline at end of file
diff --git a/Engine/Plugins/Experimental/Mover/Source/Mover/Mover.Build.cs b/Engine/Plugins/Experimental/Mover/Source/Mover/Mover.Build.cs
--- a/Engine/Plugins/Experimental/Mover/Source/Mover/Mover.Build.cs
+++ b/Engine/Plugins/Experimental/Mover/Source/Mover/Mover.Build.cs
@@ -5,7 +5,10 @@ using UnrealBuildTool;
 public class Mover : ModuleRules
 {
 	public Mover(ReadOnlyTargetRules Target) : base(Target)
-	{		
+	{
+
+		// TODO: find a better way to manage optional dependencies, such as Water and PoseSearch. This includes module dependencies here, as well as .uplugin dependencies.
+
 		PublicDependencyModuleNames.AddRange(
 			new string[]
 			{
@@ -29,7 +32,7 @@ public class Mover : ModuleRules
 				"Engine",
 				"PhysicsCore",
 				"DeveloperSettings",
-				// ... add private dependencies that you statically link with here ...	
+				"PoseSearch"
 			}
 			);
 
diff --git a/Engine/Plugins/Experimental/Mover/Source/Mover/Private/Debug/MoverDebugComponent.cpp b/Engine/Plugins/Experimental/Mover/Source/Mover/Private/Debug/MoverDebugComponent.cpp
--- a/Engine/Plugins/Experimental/Mover/Source/Mover/Private/Debug/MoverDebugComponent.cpp
+++ b/Engine/Plugins/Experimental/Mover/Source/Mover/Private/Debug/MoverDebugComponent.cpp
@@ -198,7 +198,11 @@ void UMoverDebugComponent::DrawTrajectory()
 
 	if (UMoverComponent* MoverComp = GetOwner()->FindComponentByClass<UMoverComponent>())
 	{
-		TArray<FTrajectorySampleInfo> TrajectorySamples = MoverComp->GetFutureTrajectory(LookaheadSeconds, LookaheadSamplesPerSecond);
+		FMoverPredictTrajectoryParams PredictionParams;
+		PredictionParams.NumPredictionSamples = FMath::Max(1, LookaheadSeconds * LookaheadSamplesPerSecond);
+		PredictionParams.SecondsPerSample = LookaheadSeconds / (float)PredictionParams.NumPredictionSamples;
+
+		TArray<FTrajectorySampleInfo> TrajectorySamples = MoverComp->GetPredictedTrajectory(PredictionParams);
 
 		for (int32 i = 0; i < TrajectorySamples.Num() - 1; ++i)
 		{
diff --git a/Engine/Plugins/Experimental/Mover/Source/Mover/Private/MoverComponent.cpp b/Engine/Plugins/Experimental/Mover/Source/Mover/Private/MoverComponent.cpp
--- a/Engine/Plugins/Experimental/Mover/Source/Mover/Private/MoverComponent.cpp
+++ b/Engine/Plugins/Experimental/Mover/Source/Mover/Private/MoverComponent.cpp
@@ -1164,7 +1164,7 @@ FTransform UMoverComponent::ConvertLocalRootMotionToWorld(const FTransform& Loca
 
 FTransform UMoverComponent::GetUpdatedComponentTransform() const
 {
-	if (ensure(UpdatedComponent))
+	if (UpdatedComponent)
 	{
 		return UpdatedComponent->GetComponentTransform();
 	}
@@ -1289,73 +1289,112 @@ void UMoverComponent::SetPlanarConstraint(const FPlanarConstraint& InConstraint)
 	PlanarConstraint = InConstraint;
 }
 
-TArray<FTrajectorySampleInfo> UMoverComponent::GetFutureTrajectory(float FutureSeconds, float SamplesPerSecond) const
+TArray<FTrajectorySampleInfo> UMoverComponent::GetFutureTrajectory(float FutureSeconds, float SamplesPerSecond)
 {
-	FMoverTickStartData StartingState;
+	FMoverPredictTrajectoryParams PredictionParams;
+	PredictionParams.NumPredictionSamples = FMath::Max(1, FutureSeconds * SamplesPerSecond);
+	PredictionParams.SecondsPerSample = FutureSeconds / (float)PredictionParams.NumPredictionSamples;
 
-	StartingState.InputCmd = GetLastInputCmd();
-	StartingState.SyncState = CachedLastSyncState;
-	StartingState.AuxState = CachedLastAuxState;
+	return GetPredictedTrajectory(PredictionParams);
+}
 
-	FMoverTickStartData StepState = StartingState;
+TArray<FTrajectorySampleInfo> UMoverComponent::GetPredictedTrajectory(FMoverPredictTrajectoryParams PredictionParams)
+{
+	if (ModeFSM)
+	{
+		FMoverTickStartData StartingState;
 
-	const int32 NumSamples = FMath::Max(1, FutureSeconds * SamplesPerSecond);
-	const float StepDeltaSeconds = FutureSeconds / (float)NumSamples;
+		StartingState.InputCmd = GetLastInputCmd();
+		StartingState.SyncState = CachedLastSyncState;
+		StartingState.AuxState = CachedLastAuxState;
 
-	FMoverTimeStep FutureTimeStep;
-	FutureTimeStep.StepMs = (FutureSeconds * 1000.f) / NumSamples;
-	FutureTimeStep.BaseSimTimeMs = CachedLastSimTickTimeStep.BaseSimTimeMs;
-	FutureTimeStep.ServerFrame = 0;
+		FMoverTickStartData StepState = StartingState;
 
-	if (const UBaseMovementMode* CurrentMovementMode = ModeFSM->GetCurrentMode())
-	{
-		TArray<FTrajectorySampleInfo> OutSamples;
-		OutSamples.AddDefaulted(NumSamples);
+		FMoverTimeStep FutureTimeStep;
+		FutureTimeStep.StepMs = (PredictionParams.SecondsPerSample * 1000.f);
+		FutureTimeStep.BaseSimTimeMs = CachedLastSimTickTimeStep.BaseSimTimeMs;
+		FutureTimeStep.ServerFrame = 0;
 
-		if (FMoverDefaultSyncState* StepSyncState = StepState.SyncState.SyncStateCollection.FindMutableDataByType<FMoverDefaultSyncState>())
+		if (const UBaseMovementMode* CurrentMovementMode = ModeFSM->GetCurrentMode())
 		{
-			FVector PriorLocation = StepSyncState->GetLocation_WorldSpace();
-			FRotator PriorOrientation = StepSyncState->GetOrientation_WorldSpace();
-			FVector PriorVelocity = StepSyncState->GetVelocity_WorldSpace();
+			const bool bOrigHasGravityOverride = bHasGravityOverride;
+			const FVector OrigGravityAccelOverride = GravityAccelOverride;
 
-			for (int32 i = 0; i < NumSamples; ++i)
+			if (PredictionParams.bDisableGravity)
 			{
-				// Capture sample from current step state
-				FTrajectorySampleInfo& Sample = OutSamples[i];
-
-				Sample.Transform.SetLocation(StepSyncState->GetLocation_WorldSpace());
-				Sample.Transform.SetRotation(StepSyncState->GetOrientation_WorldSpace().Quaternion());
-				Sample.LinearVelocity = StepSyncState->GetVelocity_WorldSpace();
-				Sample.InstantaneousAcceleration = (StepSyncState->GetVelocity_WorldSpace() - PriorVelocity) / StepDeltaSeconds;
-				Sample.AngularVelocity = (StepSyncState->GetOrientation_WorldSpace() - PriorOrientation) * (1.f / StepDeltaSeconds);
-
-				Sample.SimTimeMs = FutureTimeStep.BaseSimTimeMs;
-
-				// Cache prior values
-				PriorLocation = StepSyncState->GetLocation_WorldSpace();
-				PriorOrientation = StepSyncState->GetOrientation_WorldSpace();
-				PriorVelocity = StepSyncState->GetVelocity_WorldSpace();
-
-				// Generate next move from current step state
-				FProposedMove StepMove;
-				CurrentMovementMode->DoGenerateMove(StepState, FutureTimeStep, StepMove);
-
-				// Advance state based on move
-				StepSyncState->SetTransforms_WorldSpace(StepSyncState->GetLocation_WorldSpace() + (StepMove.LinearVelocity * StepDeltaSeconds),
-					StepSyncState->GetOrientation_WorldSpace() + (StepMove.AngularVelocity * StepDeltaSeconds),
-					StepMove.LinearVelocity,
-					StepSyncState->GetMovementBase(),
-					StepSyncState->GetMovementBaseBoneName());
-
-				FutureTimeStep.BaseSimTimeMs += FutureTimeStep.StepMs;
-				++FutureTimeStep.ServerFrame;
+				SetGravityOverride(true, FVector::ZeroVector);
 			}
-		}
 
-		return OutSamples;
+			TArray<FTrajectorySampleInfo> OutSamples;
+			OutSamples.AddDefaulted(PredictionParams.NumPredictionSamples);
+
+			if (FMoverDefaultSyncState* StepSyncState = StepState.SyncState.SyncStateCollection.FindMutableDataByType<FMoverDefaultSyncState>())
+			{
+				FVector PriorLocation = StepSyncState->GetLocation_WorldSpace();
+				FRotator PriorOrientation = StepSyncState->GetOrientation_WorldSpace();
+				FVector PriorVelocity = StepSyncState->GetVelocity_WorldSpace();
+
+				for (int32 i = 0; i < PredictionParams.NumPredictionSamples; ++i)
+				{
+					// Capture sample from current step state
+					FTrajectorySampleInfo& Sample = OutSamples[i];
+
+					Sample.Transform.SetLocation(StepSyncState->GetLocation_WorldSpace());
+					Sample.Transform.SetRotation(StepSyncState->GetOrientation_WorldSpace().Quaternion());
+					Sample.LinearVelocity = StepSyncState->GetVelocity_WorldSpace();
+					Sample.InstantaneousAcceleration = (StepSyncState->GetVelocity_WorldSpace() - PriorVelocity) / PredictionParams.SecondsPerSample;
+					Sample.AngularVelocity = (StepSyncState->GetOrientation_WorldSpace() - PriorOrientation) * (1.f / PredictionParams.SecondsPerSample);
+
+					Sample.SimTimeMs = FutureTimeStep.BaseSimTimeMs;
+
+					// Cache prior values
+					PriorLocation = StepSyncState->GetLocation_WorldSpace();
+					PriorOrientation = StepSyncState->GetOrientation_WorldSpace();
+					PriorVelocity = StepSyncState->GetVelocity_WorldSpace();
+
+					// Generate next move from current step state
+					FProposedMove StepMove;
+					CurrentMovementMode->DoGenerateMove(StepState, FutureTimeStep, StepMove);
+
+					// Advance state based on move
+					StepSyncState->SetTransforms_WorldSpace(StepSyncState->GetLocation_WorldSpace() + (StepMove.LinearVelocity * PredictionParams.SecondsPerSample),
+						StepSyncState->GetOrientation_WorldSpace() + (StepMove.AngularVelocity * PredictionParams.SecondsPerSample),
+						StepMove.LinearVelocity,
+						StepSyncState->GetMovementBase(),
+						StepSyncState->GetMovementBaseBoneName());
+
+					FutureTimeStep.BaseSimTimeMs += FutureTimeStep.StepMs;
+					++FutureTimeStep.ServerFrame;
+				}
+
+				// Put sample locations at visual root location if requested
+				if (PredictionParams.bUseVisualComponentRoot)
+				{
+					if (const USceneComponent* VisualComp = GetPrimaryVisualComponent())
+					{
+						const FVector VisualCompOffset = VisualComp->GetRelativeLocation();
+						const FTransform VisualCompRelativeTransform = VisualComp->GetRelativeTransform();
+
+						for (int32 i=0; i < PredictionParams.NumPredictionSamples; ++i)
+						{
+							OutSamples[i].Transform = VisualCompRelativeTransform * OutSamples[i].Transform;
+						}
+					}
+				}
+			}
+
+			if (PredictionParams.bDisableGravity)
+			{
+				SetGravityOverride(bOrigHasGravityOverride, OrigGravityAccelOverride);
+			}
+
+			return OutSamples;
+		}
 	}
 
-	return TArray<FTrajectorySampleInfo>();
+	TArray<FTrajectorySampleInfo> BlankDefaultSamples;
+	BlankDefaultSamples.AddDefaulted(PredictionParams.NumPredictionSamples);
+	return BlankDefaultSamples;
 }
 
 
diff --git a/Engine/Plugins/Experimental/Mover/Source/Mover/Private/MoverPoseSearchTrajectoryPredictor.cpp b/Engine/Plugins/Experimental/Mover/Source/Mover/Private/MoverPoseSearchTrajectoryPredictor.cpp
new file mode 100644
--- /dev/null
+++ b/Engine/Plugins/Experimental/Mover/Source/Mover/Private/MoverPoseSearchTrajectoryPredictor.cpp
@@ -0,0 +1,114 @@
+// Copyright Epic Games, Inc. All Rights Reserved.
+
+#include "MoverPoseSearchTrajectoryPredictor.h"
+ #include "PoseSearch/PoseSearchTrajectoryTypes.h"
+ #include "MoverComponent.h"
+ #include "MoveLibrary/MovementUtils.h"
+ 
+void UMoverTrajectoryPredictor::Predict_Implementation(FPoseSearchQueryTrajectory& InOutTrajectory,
+	int32 NumPredictionSamples, float SecondsPerPredictionSample, int NumHistorySamples)
+{
+	if (!MoverComponent)
+	{
+		UE_LOG(LogMover, Log, TEXT("Calling Predict without a Mover Component. This is invalid and the trajectory will not be modified."));
+		return;
+	}
+
+	FMoverPredictTrajectoryParams PredictParams;
+	PredictParams.NumPredictionSamples = NumPredictionSamples;
+	PredictParams.SecondsPerSample = SecondsPerPredictionSample;
+	PredictParams.bUseVisualComponentRoot = true;
+	PredictParams.bDisableGravity = true;
+
+	TArray<FTrajectorySampleInfo> MoverPredictionSamples = MoverComponent->GetPredictedTrajectory(PredictParams);
+
+	if (InOutTrajectory.Samples.Num() < (NumHistorySamples + MoverPredictionSamples.Num()))
+	{
+		UE_LOG(LogMover, Warning, TEXT("InOutTrajectory Samples array does not have enough space for %i predicted samples"), MoverPredictionSamples.Num());
+	}
+	else
+	{
+		float AccumulatedSeconds = 0.0;
+
+		for (int32 i = 0; i < NumPredictionSamples; ++i)
+		{
+			const int PoseSampleIdx = i + NumHistorySamples;
+			const int PredictedSampleIdx = i;
+
+			InOutTrajectory.Samples[PoseSampleIdx].Position = MoverPredictionSamples[PredictedSampleIdx].Transform.GetLocation();
+			InOutTrajectory.Samples[PoseSampleIdx].Facing = MoverPredictionSamples[PredictedSampleIdx].Transform.GetRotation();
+			InOutTrajectory.Samples[PoseSampleIdx].AccumulatedSeconds = AccumulatedSeconds;
+
+			AccumulatedSeconds += SecondsPerPredictionSample;
+		}
+	}
+}
+
+void UMoverTrajectoryPredictor::GetGravity_Implementation(FVector& OutGravityAccel)
+{
+	if (!MoverComponent)
+	{
+		UE_LOG(LogMover, Log, TEXT("Calling GetGravity without a Mover Component. Return value will be defaulted."));
+		OutGravityAccel = FVector::ZeroVector;
+		return;
+	}
+
+	OutGravityAccel = MoverComponent->GetGravityAcceleration();
+}
+
+
+void UMoverTrajectoryPredictor::GetCurrentState_Implementation(FVector& OutPosition, FQuat& OutFacing, FVector& OutVelocity)
+{
+	if (!MoverComponent)
+	{
+		UE_LOG(LogMover, Log, TEXT("Calling GetCurrentState without a Mover Component. Return values will be defaulted."));
+		OutPosition = OutVelocity = FVector::ZeroVector;
+		OutFacing = FQuat::Identity;
+		return;
+	}
+
+	const USceneComponent* VisualComp = MoverComponent->GetPrimaryVisualComponent();
+
+	if (VisualComp)
+	{
+		OutPosition = VisualComp->GetComponentLocation();
+	}
+	else
+	{
+		OutPosition = MoverComponent->GetUpdatedComponentTransform().GetLocation();
+	}
+
+	const bool bOrientRotationToMovement = true;
+
+	if (bOrientRotationToMovement)
+	{
+		if (VisualComp)
+		{
+			OutFacing = VisualComp->GetComponentRotation().Quaternion();
+		}
+		else
+		{
+			OutFacing = MoverComponent->GetUpdatedComponentTransform().GetRotation();
+		}
+	}
+	else
+	{
+		// JAH TODO: Needs a solve
+		//OutFacing = FQuat::MakeFromRotator(FRotator(0, TrajectoryDataState.DesiredControllerYawLastUpdate, 0)) * TrajectoryDataDerived.MeshCompRelativeRotation;
+	}
+
+	OutVelocity = MoverComponent->GetVelocity();
+}
+
+
+void UMoverTrajectoryPredictor::GetVelocity_Implementation(FVector& OutVelocity)
+{
+	if (!MoverComponent)
+	{
+		UE_LOG(LogMover, Log, TEXT("Calling GetVelocity without a Mover Component. Return value will be defaulted."));
+		OutVelocity = FVector::ZeroVector;
+		return;
+	}
+
+	OutVelocity = MoverComponent->GetVelocity();
+}
\ No newline at end of file
diff --git a/Engine/Plugins/Experimental/Mover/Source/Mover/Public/MoverComponent.h b/Engine/Plugins/Experimental/Mover/Source/Mover/Public/MoverComponent.h
--- a/Engine/Plugins/Experimental/Mover/Source/Mover/Public/MoverComponent.h
+++ b/Engine/Plugins/Experimental/Mover/Source/Mover/Public/MoverComponent.h
@@ -263,8 +263,13 @@ public:	// Queries
 	FRotator GetTargetOrientation() const;
 
 	/** Get a sampling of where the actor is projected to be in the future, based on a current state. Note that this is projecting ideal movement without doing full simulation and collision. */
+	UE_DEPRECATED(5.5, "Use GetPredictedTrajectory instead.")
 	UFUNCTION(BlueprintCallable, BlueprintPure = false, Category = Mover)
-	TArray<FTrajectorySampleInfo> GetFutureTrajectory(float FutureSeconds, float SamplesPerSecond) const;
+	TArray<FTrajectorySampleInfo> GetFutureTrajectory(float FutureSeconds, float SamplesPerSecond);
+
+	/** Get a sampling of where the actor is projected to be in the future, based on a current state. Note that this is projecting ideal movement without doing full simulation and collision. */
+	UFUNCTION(BlueprintCallable, BlueprintPure = false, Category = Mover)
+	TArray<FTrajectorySampleInfo> GetPredictedTrajectory(FMoverPredictTrajectoryParams PredictionParams);	
 
 	// Get the current movement mode name
 	UFUNCTION(BlueprintPure, Category = Mover)
diff --git a/Engine/Plugins/Experimental/Mover/Source/Mover/Public/MoverPoseSearchTrajectoryPredictor.h b/Engine/Plugins/Experimental/Mover/Source/Mover/Public/MoverPoseSearchTrajectoryPredictor.h
new file mode 100644
--- /dev/null
+++ b/Engine/Plugins/Experimental/Mover/Source/Mover/Public/MoverPoseSearchTrajectoryPredictor.h
@@ -0,0 +1,35 @@
+// Copyright Epic Games, Inc. All Rights Reserved.
+
+#pragma once
+
+#include "PoseSearch/PoseSearchTrajectoryPredictor.h"
+#include "MoverPoseSearchTrajectoryPredictor.generated.h"
+
+
+class UMoverComponent;
+
+/**
+ * Trajectory predictor that can query from a Mover-driven actor, for use with Pose Search
+ */
+UCLASS(BlueprintType, EditInlineNew)
+class MOVER_API UMoverTrajectoryPredictor : public UObject, public IPoseSearchTrajectoryPredictorInterface
+{
+	GENERATED_BODY()
+
+public:
+
+	UFUNCTION(BlueprintCallable, Category = "Animation|PoseSearch|Experimental")
+	void Setup(UMoverComponent* InMoverComponent) { MoverComponent = InMoverComponent; }
+
+	virtual void Predict_Implementation(FPoseSearchQueryTrajectory& InOutTrajectory, int32 NumPredictionSamples, float SecondsPerPredictionSample, int32 NumHistorySamples) override;
+
+	virtual void GetGravity_Implementation(FVector& OutGravityAccel) override;
+
+	virtual void GetCurrentState_Implementation(FVector& OutPosition, FQuat& OutFacing, FVector& OutVelocity) override;
+
+	virtual void GetVelocity_Implementation(FVector& OutVelocity) override;
+
+protected:
+	TObjectPtr<UMoverComponent> MoverComponent;
+
+};
diff --git a/Engine/Plugins/Experimental/Mover/Source/Mover/Public/MoverTypes.h b/Engine/Plugins/Experimental/Mover/Source/Mover/Public/MoverTypes.h
--- a/Engine/Plugins/Experimental/Mover/Source/Mover/Public/MoverTypes.h
+++ b/Engine/Plugins/Experimental/Mover/Source/Mover/Public/MoverTypes.h
@@ -82,6 +82,32 @@ struct MOVER_API FMoverTimeStep
 };
 
 
+USTRUCT(BlueprintType)
+struct MOVER_API FMoverPredictTrajectoryParams
+{
+	GENERATED_BODY()
+
+	/** How many samples to predict into the future, including the first sample, which is always a snapshot of the
+	 *  starting state with 0 accumulated time. */
+	UPROPERTY(EditAnywhere, BlueprintReadWrite, Category = Mover, meta=(ClampMin=1))
+	int32 NumPredictionSamples = 1;
+
+	/* How much time between predicted samples */
+	UPROPERTY(EditAnywhere, BlueprintReadWrite, Category = Mover, meta=(ClampMin=0.00001))
+	float SecondsPerSample = 0.333f;
+
+	/** If true, samples are based on the visual component transform, rather than the 'updated' movement root. 
+	 *  Typically, this is a mesh with its component location at the bottom of the collision primitive.
+	 *  If false, samples are from the movement root. */
+	UPROPERTY(EditAnywhere, BlueprintReadWrite, Category = Mover)
+	bool bUseVisualComponentRoot = false;
+
+	/** If true, gravity will not taken into account during prediction */
+	UPROPERTY(EditAnywhere, BlueprintReadWrite, Category = Mover)
+	bool bDisableGravity = false;
+
+};
+
 
 // Base type for all data structs used to compose Mover simulation model definition dynamically (input cmd, sync state, aux state)
 // NOTE: for simulation state data (sync/aux), derive from FMoverStateData instead 
