motion matching - excluding IgnoreForMemberInitializationTest properties from database ddc key generation in FKeyBuilder::ShouldSkipProperty

diff --git a/Engine/Plugins/Animation/PoseSearch/Source/Runtime/Private/PoseSearchDerivedDataKey.cpp b/Engine/Plugins/Animation/PoseSearch/Source/Runtime/Private/PoseSearchDerivedDataKey.cpp
--- a/Engine/Plugins/Animation/PoseSearch/Source/Runtime/Private/PoseSearchDerivedDataKey.cpp
+++ b/Engine/Plugins/Animation/PoseSearch/Source/Runtime/Private/PoseSearchDerivedDataKey.cpp
@@ -92,8 +92,15 @@ bool FKeyBuilder::ShouldSkipProperty(const FProperty* InProperty) const
 		#endif
 		return true;
 	}
-		
-	check(!InProperty->HasMetaData(IgnoreForMemberInitializationTestName));
+	
+	if (InProperty->HasMetaData(IgnoreForMemberInitializationTestName))
+	{
+		#if UE_POSE_SEARCH_DERIVED_DATA_LOGGING
+		UE_LOG(LogPoseSearch, Log, TEXT("%s x %s (IgnoreForMemberInitializationTest)"), *GetIndentation(), *InProperty->GetFullName());
+		#endif
+		return true;
+	}
+
 	check(!InProperty->HasMetaData(NeverInHashName));
 
 	#if UE_POSE_SEARCH_DERIVED_DATA_LOGGING
