Fix for trajectory yaw on remote clients

diff --git a/Engine/Plugins/Animation/PoseSearch/Source/Runtime/Private/PoseSearchTrajectoryLibrary.cpp b/Engine/Plugins/Animation/PoseSearch/Source/Runtime/Private/PoseSearchTrajectoryLibrary.cpp
--- a/Engine/Plugins/Animation/PoseSearch/Source/Runtime/Private/PoseSearchTrajectoryLibrary.cpp
+++ b/Engine/Plugins/Animation/PoseSearch/Source/Runtime/Private/PoseSearchTrajectoryLibrary.cpp
@@ -58,11 +58,8 @@ void FPoseSearchTrajectoryData::UpdateData(
 		}
 	}
 
-	// @todo: Simulated proxies don't have controllers, so they'll need some other mechanism to account for controller rotation rate.
-	const AController* Controller = Character->Controller;
-	if (Controller)
 	{
-		const float DesiredControllerYaw = Controller->GetDesiredRotation().Yaw;
+		const float DesiredControllerYaw = Character->GetViewRotation().Yaw;
 		
 		const float DesiredYawDelta = DesiredControllerYaw - TrajectoryDataState.DesiredControllerYawLastUpdate;
 		TrajectoryDataState.DesiredControllerYawLastUpdate = DesiredControllerYaw;
