motion matching - fix for Database->Modify() calls

diff --git a/Engine/Plugins/Animation/PoseSearch/Source/Editor/Private/PoseSearchDatabaseViewModel.cpp b/Engine/Plugins/Animation/PoseSearch/Source/Editor/Private/PoseSearchDatabaseViewModel.cpp
--- a/Engine/Plugins/Animation/PoseSearch/Source/Editor/Private/PoseSearchDatabaseViewModel.cpp
+++ b/Engine/Plugins/Animation/PoseSearch/Source/Editor/Private/PoseSearchDatabaseViewModel.cpp
@@ -564,6 +564,8 @@ void FDatabaseViewModel::AddAnimCompositeToDatabase(UAnimComposite* AnimComposit
 {
 	if (UPoseSearchDatabase* Database = GetPoseSearchDatabase())
 	{
+		Database->Modify();
+
 		FPoseSearchDatabaseAnimComposite NewAsset;
 		NewAsset.AnimComposite = AnimComposite;
 		Database->AddAnimationAsset(FInstancedStruct::Make(NewAsset));
@@ -705,39 +707,49 @@ bool FDatabaseViewModel::SetAnimationAsset(int32 AnimationAssetIndex, UObject* A
 		{
 			if (FPoseSearchDatabaseAnimationAssetBase* DatabaseAnimationAsset = Database->GetMutableDatabaseAnimationAsset<FPoseSearchDatabaseAnimationAssetBase>(AnimationAssetIndex))
 			{
-				Database->Modify();
-
 				// Ensure that our target database item matches the input object's class.
 				const UClass* AssetClass = AnimAsset->GetClass();
 				if (AssetClass == DatabaseAnimationAsset->GetAnimationAssetStaticClass())
 				{
 					if (AssetClass->IsChildOf( UAnimSequence::StaticClass()))
 					{
+						Database->Modify();
 						FPoseSearchDatabaseSequence* DatabaseSequenceAsset = static_cast<FPoseSearchDatabaseSequence*>(DatabaseAnimationAsset);
 						DatabaseSequenceAsset->Sequence = Cast<UAnimSequence>(AnimAsset);
+						return true;
 					}
-					else if (AssetClass->IsChildOf(UAnimComposite::StaticClass()))
+					
+					if (AssetClass->IsChildOf(UAnimComposite::StaticClass()))
 					{
+						Database->Modify();
 						FPoseSearchDatabaseAnimComposite* DatabaseCompositeAsset = static_cast<FPoseSearchDatabaseAnimComposite*>(DatabaseAnimationAsset);
 						DatabaseCompositeAsset->AnimComposite = Cast<UAnimComposite>(AnimAsset);
+						return true;
 					}
-					else if (AssetClass->IsChildOf(UAnimMontage::StaticClass()))
+					
+					if (AssetClass->IsChildOf(UAnimMontage::StaticClass()))
 					{
+						Database->Modify();
 						FPoseSearchDatabaseAnimMontage* DatabaseMontageAsset = static_cast<FPoseSearchDatabaseAnimMontage*>(DatabaseAnimationAsset);
 						DatabaseMontageAsset->AnimMontage = Cast<UAnimMontage>(AnimAsset);
+						return true;
 					}
-					else if (AssetClass->IsChildOf(UBlendSpace::StaticClass()))
+					
+					if (AssetClass->IsChildOf(UBlendSpace::StaticClass()))
 					{
+						Database->Modify();
 						FPoseSearchDatabaseBlendSpace* DatabaseBlendSpaceAsset = static_cast<FPoseSearchDatabaseBlendSpace*>(DatabaseAnimationAsset);
 						DatabaseBlendSpaceAsset->BlendSpace = Cast<UBlendSpace>(AnimAsset);
+						return true;
 					}
-					else if (AssetClass->IsChildOf(UMultiAnimAsset::StaticClass()))
+					
+					if (AssetClass->IsChildOf(UMultiAnimAsset::StaticClass()))
 					{
+						Database->Modify();
 						FPoseSearchDatabaseMultiAnimAsset* DatabaseMultiAnimAssetAsset = static_cast<FPoseSearchDatabaseMultiAnimAsset*>(DatabaseAnimationAsset);
 						DatabaseMultiAnimAssetAsset->MultiAnimAsset = Cast<UMultiAnimAsset>(AnimAsset);
+						return true;
 					}
-
-					return true;
 				}
 			}
 		}
