motion matching - fix for FAnimationAssetSampler::ExtractPose generating invalid poses on empty blend spaces

diff --git a/Engine/Plugins/Animation/PoseSearch/Source/Runtime/Private/PoseSearchAssetSampler.cpp b/Engine/Plugins/Animation/PoseSearch/Source/Runtime/Private/PoseSearchAssetSampler.cpp
--- a/Engine/Plugins/Animation/PoseSearch/Source/Runtime/Private/PoseSearchAssetSampler.cpp
+++ b/Engine/Plugins/Animation/PoseSearch/Source/Runtime/Private/PoseSearchAssetSampler.cpp
@@ -367,6 +367,11 @@ void FAnimationAssetSampler::ExtractPose(const FAnimExtractContext& ExtractionCt
 			BlendSpace->ResetBlendSamples(BlendSamples, ToNormalizedTime(ExtractionCtx.CurrentTime), ExtractionCtx.bLooping, false);
 			BlendSpace->GetAnimationPose(BlendSamples, ExtractionCtx, OutAnimPoseData);
 		}
+		else
+		{
+			UE_LOG(LogPoseSearch, Error, TEXT("FAnimMontageSampler::ExtractPose: UBlendSpace %s couldn't return a valid pose for BlendParameters (%.2f, %.2f, %.2f)"), *BlendSpace->GetName(), BlendParameters.X, BlendParameters.Y, BlendParameters.Z);
+			OutAnimPoseData.GetPose().ResetToRefPose();
+		}
 	}
 	else if (const UAnimMontage* AnimMontage = Cast<UAnimMontage>(AnimationAssetPtr.Get()))
 	{
