motion matching - fix for: After stopping PIE and using Rewind Debugger with Pose Search you can not see details from the debugger

diff --git a/Engine/Plugins/Animation/PoseSearch/Source/Editor/Private/PoseSearchDebuggerView.cpp b/Engine/Plugins/Animation/PoseSearch/Source/Editor/Private/PoseSearchDebuggerView.cpp
--- a/Engine/Plugins/Animation/PoseSearch/Source/Editor/Private/PoseSearchDebuggerView.cpp
+++ b/Engine/Plugins/Animation/PoseSearch/Source/Editor/Private/PoseSearchDebuggerView.cpp
@@ -182,13 +182,6 @@ void SDebuggerView::Tick(const FGeometry& AllottedGeometry, const double InCurre
 	const UWorld* DebuggerWorld = FDebugger::GetWorld();
 	check(DebuggerWorld);
 
-	// @todo: Handle editor world when those features are enabled for the Rewind Debugger
-	// Currently prevents debug draw remnants from stopped world
-	if (DebuggerWorld->WorldType != EWorldType::PIE)
-	{
-		return;
-	}
-
 	const bool bSameTime = FMath::Abs(TimeMarker - PreviousTimeMarker) < DOUBLE_SMALL_NUMBER;
 	PreviousTimeMarker = TimeMarker;
 
