motion matching - making pose search database assets SamplingRange read only if synchronized with a PoseSearchBranchIn notify state

diff --git a/Engine/Plugins/Animation/PoseSearch/Source/Runtime/Public/PoseSearch/PoseSearchDatabase.h b/Engine/Plugins/Animation/PoseSearch/Source/Runtime/Public/PoseSearch/PoseSearchDatabase.h
--- a/Engine/Plugins/Animation/PoseSearch/Source/Runtime/Public/PoseSearch/PoseSearchDatabase.h
+++ b/Engine/Plugins/Animation/PoseSearch/Source/Runtime/Public/PoseSearch/PoseSearchDatabase.h
@@ -142,7 +142,9 @@ struct POSESEARCH_API FPoseSearchDatabaseSequence : public FPoseSearchDatabaseAn
 	// It allows users to set a time range to an individual animation sequence in the database. 
 	// This is effectively trimming the beginning and end of the animation in the database (not in the original sequence).
 	// If set to [0, 0] it will be the entire frame range of the original sequence.
-	UPROPERTY(EditAnywhere, Category = "Settings", meta = (ClampToMinMaxLimits, DisplayPriority = 2))
+	// Set to readonly if this asset is synchronized via PoseSearchBranchIn notify state.
+	// To edit its value update the associated PoseSearchBranchIn in Sequence
+	UPROPERTY(EditAnywhere, Category = "Settings", meta = (EditCondition="BranchInId == 0", ClampToMinMaxLimits, DisplayPriority = 2))
 	FFloatInterval SamplingRange = FFloatInterval(0.f, 0.f);
 
 	virtual UClass* GetAnimationAssetStaticClass() const override;
@@ -262,7 +264,9 @@ struct POSESEARCH_API FPoseSearchDatabaseAnimComposite : public FPoseSearchDatab
 	// It allows users to set a time range to an individual animation composite in the database. 
 	// This is effectively trimming the beginning and end of the animation in the database (not in the original composite).
 	// If set to [0, 0] it will be the entire frame range of the original composite.
-	UPROPERTY(EditAnywhere, Category = "Settings", meta = (ClampToMinMaxLimits, DisplayPriority = 3))
+	// Set to readonly if this asset is synchronized via PoseSearchBranchIn notify state.
+	// To edit its value update the associated PoseSearchBranchIn in AnimComposite
+	UPROPERTY(EditAnywhere, Category = "Settings", meta = (EditCondition="BranchInId == 0", ClampToMinMaxLimits, DisplayPriority = 3))
 	FFloatInterval SamplingRange = FFloatInterval(0.f, 0.f);
 
 	virtual UClass* GetAnimationAssetStaticClass() const override;
@@ -303,7 +307,9 @@ struct POSESEARCH_API FPoseSearchDatabaseAnimMontage : public FPoseSearchDatabas
 	// It allows users to set a time range to an individual animation montage in the database. 
 	// This is effectively trimming the beginning and end of the animation in the database (not in the original montage).
 	// If set to [0, 0] it will be the entire frame range of the original montage.
-	UPROPERTY(EditAnywhere, Category = "Settings", meta = (ClampToMinMaxLimits, DisplayPriority = 2))
+	// Set to readonly if this asset is synchronized via PoseSearchBranchIn notify state.
+	// To edit its value update the associated PoseSearchBranchIn in AnimMontage
+	UPROPERTY(EditAnywhere, Category = "Settings", meta = (EditCondition="BranchInId == 0", ClampToMinMaxLimits, DisplayPriority = 2))
 	FFloatInterval SamplingRange = FFloatInterval(0.f, 0.f);
 
 	virtual UClass* GetAnimationAssetStaticClass() const override;
