Pose Search: Added basic support for searching on anim float curves

diff --git a/Engine/Plugins/Animation/PoseSearch/Source/Runtime/Private/PoseSearchAssetIndexer.cpp b/Engine/Plugins/Animation/PoseSearch/Source/Runtime/Private/PoseSearchAssetIndexer.cpp
--- a/Engine/Plugins/Animation/PoseSearch/Source/Runtime/Private/PoseSearchAssetIndexer.cpp
+++ b/Engine/Plugins/Animation/PoseSearch/Source/Runtime/Private/PoseSearchAssetIndexer.cpp
@@ -189,6 +189,11 @@ void FAnimationAssetSamplers::ExtractPose(float Time, FCompactPose& OutPose, int
 	AnimationAssetSamplers[RoleIndex]->ExtractPose(Time, OutPose);
 }
 
+void FAnimationAssetSamplers::ExtractPose(float Time, FCompactPose& OutPose, FBlendedCurve& OutCurve, int32 RoleIndex) const
+{
+	AnimationAssetSamplers[RoleIndex]->ExtractPose(Time, OutPose, OutCurve);
+}
+
 FTransform FAnimationAssetSamplers::MirrorTransform(const FTransform& InTransform, int32 RoleIndex) const
 {
 	return MirrorDataCaches[RoleIndex]->MirrorTransform(InTransform);
@@ -417,6 +422,7 @@ FAssetIndexer::FCachedEntry& FAssetIndexer::GetEntry(float SampleTime)
 
 		Entry->RootTransform.SetNum(AssetSamplersNum);
 		Entry->ComponentSpacePose.SetNum(AssetSamplersNum);
+		Entry->Curves.SetNum(AssetSamplersNum);
 		for (int32 RoleIndex = 0; RoleIndex < AssetSamplers.AnimationAssetSamplers.Num(); ++RoleIndex)
 		{
 			if (!BoneContainers[RoleIndex].IsValid())
@@ -443,8 +449,10 @@ FAssetIndexer::FCachedEntry& FAssetIndexer::GetEntry(float SampleTime)
 
 			FMemMark Mark(FMemStack::Get());
 			FCompactPose Pose;
+			FBlendedCurve Curve;
 			Pose.SetBoneContainer(&BoneContainers[RoleIndex]);
-			AssetSamplers.ExtractPose(CurrentTime, Pose, RoleIndex);
+			Curve.InitFrom(BoneContainers[RoleIndex]);
+			AssetSamplers.ExtractPose(CurrentTime, Pose, Curve, RoleIndex);
 
 #if ENABLE_ANIM_DEBUG
 			const int32 NumIterations = CVarMotionMatchTestExtractPoseDeterminismNumIterations.GetValueOnAnyThread();
@@ -481,6 +489,7 @@ FAssetIndexer::FCachedEntry& FAssetIndexer::GetEntry(float SampleTime)
 			FCSPose<FCompactPose> StackComponentSpacePose;
 			StackComponentSpacePose.InitPose(MoveTemp(Pose));
 			Entry->ComponentSpacePose[RoleIndex].CopyPose(StackComponentSpacePose);
+			Entry->Curves[RoleIndex].CopyFrom(Curve);
 
 			Entry->RootTransform[RoleIndex] = MirrorTransform(SampleRootTransform, RoleIndex);
 			Entry->bClamped |= bSampleClamped;
@@ -599,6 +608,24 @@ bool FAssetIndexer::GetSampleRotation(FQuat& OutSampleRotation, float SampleTime
 	return true;
 }
 
+bool FAssetIndexer::GetSampleCurveValue(float& OutCurveValue, float SampleTimeOffset, int32 SampleIdx, const FName& CurveName, const FRole& SampleRole)
+{
+	const float Time = CalculateSampleTime(SampleIdx);
+	const float SampleTime = Time + SampleTimeOffset;
+
+	OutCurveValue = GetSampleCurveValueInternal(SampleTime, CurveName, SampleRole);
+	return true;
+}
+
+
+float FAssetIndexer::GetSampleCurveValueInternal(float SampleTime, const FName& CurveName, const FRole& Role)
+{
+	const int32 SampleRoleIndex = RoleToIndex[Role];
+	FCachedEntry& Entry = GetEntry(SampleTime);
+
+	return Entry.Curves[SampleRoleIndex].Get(CurveName);
+}
+
 bool FAssetIndexer::GetSamplePosition(FVector& OutSamplePosition, float SampleTimeOffset, float OriginTimeOffset, int32 SampleIdx, int8 SchemaSampleBoneIdx, int8 SchemaOriginBoneIdx, const FRole& SampleRole, const FRole& OriginRole, EPermutationTimeType PermutationTimeType, int32 SamplingAttributeId)
 {
 	float PermutationSampleTimeOffset = 0.f;
diff --git a/Engine/Plugins/Animation/PoseSearch/Source/Runtime/Private/PoseSearchAssetSampler.cpp b/Engine/Plugins/Animation/PoseSearch/Source/Runtime/Private/PoseSearchAssetSampler.cpp
--- a/Engine/Plugins/Animation/PoseSearch/Source/Runtime/Private/PoseSearchAssetSampler.cpp
+++ b/Engine/Plugins/Animation/PoseSearch/Source/Runtime/Private/PoseSearchAssetSampler.cpp
@@ -410,6 +410,19 @@ void FAnimationAssetSampler::ExtractPose(float Time, FCompactPose& OutPose) cons
 	ExtractPose(ExtractionCtx, AnimPoseData);
 }
 
+void FAnimationAssetSampler::ExtractPose(float Time, FCompactPose& OutPose, FBlendedCurve& OutCurve) const
+{
+	UE::Anim::FStackAttributeContainer UnusedAtrribute;
+	OutCurve.InitFrom(OutPose.GetBoneContainer());
+	FAnimationPoseData AnimPoseData = { OutPose, OutCurve, UnusedAtrribute };
+
+	FDeltaTimeRecord DeltaTimeRecord;
+	DeltaTimeRecord.Set(Time, 0.f);
+	FAnimExtractContext ExtractionCtx(double(Time), false, DeltaTimeRecord, IsLoopable());
+
+	ExtractPose(ExtractionCtx, AnimPoseData);
+}
+
 FTransform FAnimationAssetSampler::ExtractRootTransform(float Time) const
 {
 	FTransform RootTransform = FTransform::Identity;
diff --git a/Engine/Plugins/Animation/PoseSearch/Source/Runtime/Private/PoseSearchContext.cpp b/Engine/Plugins/Animation/PoseSearch/Source/Runtime/Private/PoseSearchContext.cpp
--- a/Engine/Plugins/Animation/PoseSearch/Source/Runtime/Private/PoseSearchContext.cpp
+++ b/Engine/Plugins/Animation/PoseSearch/Source/Runtime/Private/PoseSearchContext.cpp
@@ -490,6 +490,31 @@ void FSearchContext::UpdateCurrentResultPoseVector()
 	}
 }
 
+float FSearchContext::GetSampleCurveValue(float SampleTimeOffset, const FName& CurveName, const FRole& SampleRole)
+{
+	const float SampleTime = SampleTimeOffset;
+	return GetSampleCurveValueInternal(SampleTime, CurveName, SampleRole);
+}
+
+float FSearchContext::GetSampleCurveValueInternal(float SampleTime, const FName& CurveName, const FRole& SampleRole)
+{
+	check(!CachedQueries.IsEmpty());
+	const UPoseSearchSchema* Schema = CachedQueries.Last().GetSchema();
+	check(Schema);
+
+	float OutCurveValue = 0.0f;
+	const IPoseHistory* PoseHistory = GetPoseHistory(SampleRole);
+	{
+		if (ensure(PoseHistory))
+		{
+			// @TODO: Add curve support to the pose history node.
+			return OutCurveValue;
+		}
+	}
+
+	return OutCurveValue;
+}
+
 FQuat FSearchContext::GetSampleRotation(float SampleTimeOffset, float OriginTimeOffset, int8 SchemaSampleBoneIdx, int8 SchemaOriginBoneIdx, const FRole& SampleRole, const FRole& OriginRole, EPermutationTimeType PermutationTimeType, const FQuat* SampleBoneRotationWorldOverride)
 {
 	float PermutationSampleTimeOffset = 0.f;
diff --git a/Engine/Plugins/Animation/PoseSearch/Source/Runtime/Private/PoseSearchFeatureChannel_Curve.cpp b/Engine/Plugins/Animation/PoseSearch/Source/Runtime/Private/PoseSearchFeatureChannel_Curve.cpp
new file mode 100644
--- /dev/null
+++ b/Engine/Plugins/Animation/PoseSearch/Source/Runtime/Private/PoseSearchFeatureChannel_Curve.cpp
@@ -0,0 +1,133 @@
+// Copyright Epic Games, Inc. All Rights Reserved.
+
+#include "PoseSearchFeatureChannel_Curve.h"
+#include "Engine/BlueprintGeneratedClass.h"
+#include "PoseSearch/PoseSearchAssetIndexer.h"
+#include "PoseSearch/PoseSearchAssetSampler.h"
+#include "PoseSearch/PoseSearchContext.h"
+#include "PoseSearch/PoseSearchHistory.h"
+#include "PoseSearch/PoseSearchSchema.h"
+
+#if WITH_EDITOR
+#include "PropertyHandle.h"
+#endif // WITH_EDITOR
+
+UPoseSearchFeatureChannel_Curve::UPoseSearchFeatureChannel_Curve()
+{
+	bUseBlueprintQueryOverride = Cast<UBlueprintGeneratedClass>(GetClass()) != nullptr;
+}
+
+bool UPoseSearchFeatureChannel_Curve::Finalize(UPoseSearchSchema* Schema)
+{
+	ChannelDataOffset = Schema->SchemaCardinality;
+	ChannelCardinality = 1;
+	Schema->SchemaCardinality += ChannelCardinality;
+
+	CurveIdx = Schema->AddCurveReference(CurveName, SampleRole);
+
+	return true;
+}
+
+void UPoseSearchFeatureChannel_Curve::BuildQuery(UE::PoseSearch::FSearchContext& SearchContext) const
+{
+	using namespace UE::PoseSearch;
+
+	float CurveValue = 0.0f;
+	if (bUseBlueprintQueryOverride)
+	{
+		CurveValue = BP_GetCurveValue(SearchContext.GetAnimInstance(SampleRole));
+		FFeatureVectorHelper::EncodeFloat(SearchContext.EditFeatureVector(), ChannelDataOffset, CurveValue);
+		return;
+	}
+	
+	// trying to get the BuildQuery data from another schema UPoseSearchFeatureChannel_Curve already cached in the SearchContext
+	if (SearchContext.IsUseCachedChannelData())
+	{
+		// composing a unique identifier to specify this channel with all the required properties to be able to share the query data with other channels of the same type
+		uint32 UniqueIdentifier = GetClass()->GetUniqueID();
+		UniqueIdentifier = HashCombineFast(UniqueIdentifier, GetTypeHash(CurveName));
+		UniqueIdentifier = HashCombineFast(UniqueIdentifier, GetTypeHash(SampleRole));
+		UniqueIdentifier = HashCombineFast(UniqueIdentifier, GetTypeHash(CurveIdx));
+		UniqueIdentifier = HashCombineFast(UniqueIdentifier, GetTypeHash(SampleTimeOffset));
+		UniqueIdentifier = HashCombineFast(UniqueIdentifier, GetTypeHash(InputQueryPose));
+
+		TConstArrayView<float> CachedChannelData;
+		if (const UPoseSearchFeatureChannel* CachedChannel = SearchContext.GetCachedChannelData(UniqueIdentifier, this, CachedChannelData))
+		{
+#if DO_CHECK
+			const UPoseSearchFeatureChannel_Curve* CachedCurveChannel = Cast<UPoseSearchFeatureChannel_Curve>(CachedChannel);
+			check(CachedCurveChannel);
+			check(CachedCurveChannel->GetChannelCardinality() == ChannelCardinality);
+			check(CachedChannelData.Num() == ChannelCardinality);
+
+			// making sure there were no hash collisions
+			check(CachedCurveChannel->CurveName == CurveName);
+			check(CachedCurveChannel->SampleRole == SampleRole);
+			check(CachedCurveChannel->CurveIdx == CurveIdx);
+			check(CachedCurveChannel->SampleTimeOffset == SampleTimeOffset);
+			check(CachedCurveChannel->InputQueryPose == InputQueryPose);
+#endif //DO_CHECK
+
+			// copying the CachedChannelData into this channel portion of the FeatureVectorBuilder
+			FFeatureVectorHelper::Copy(SearchContext.EditFeatureVector().Slice(ChannelDataOffset, ChannelCardinality), 0, ChannelCardinality, CachedChannelData);
+			return;
+		}
+	}
+
+	const bool bCanUseCurrentResult = SearchContext.CanUseCurrentResult();
+	const bool bSkip = InputQueryPose != EInputQueryPose::UseCharacterPose && bCanUseCurrentResult;
+	if (bSkip || (!SearchContext.ArePoseHistoriesValid()))
+	{
+		if (bCanUseCurrentResult)
+		{
+			FFeatureVectorHelper::Copy(SearchContext.EditFeatureVector(), ChannelDataOffset, ChannelCardinality, SearchContext.GetCurrentResultPoseVector());
+			return;
+		}
+
+		// we leave the SearchContext.EditFeatureVector() set to zero since the SearchContext.PoseHistory is invalid and it'll fail if we continue
+		UE_LOG(LogPoseSearch, Error, TEXT("UPoseSearchFeatureChannel_Curve::BuildQuery - Failed because Pose History Node is missing."));
+		return;
+	}
+	
+	CurveValue = SearchContext.GetSampleCurveValue(SampleTimeOffset, CurveName, SampleRole);
+	FFeatureVectorHelper::EncodeFloat(SearchContext.EditFeatureVector(), ChannelDataOffset, CurveValue);
+}
+
+#if WITH_EDITOR
+void UPoseSearchFeatureChannel_Curve::FillWeights(TArrayView<float> Weights) const
+{
+	for (int32 i = 0; i < ChannelCardinality; ++i)
+	{
+		Weights[ChannelDataOffset + i] = Weight;
+	}
+}
+
+bool UPoseSearchFeatureChannel_Curve::IndexAsset(UE::PoseSearch::FAssetIndexer& Indexer) const
+{
+	using namespace UE::PoseSearch;
+
+	for (int32 SampleIdx = Indexer.GetBeginSampleIdx(); SampleIdx != Indexer.GetEndSampleIdx(); ++SampleIdx)
+	{
+		float CurveValue;
+
+		Indexer.GetSampleCurveValue(CurveValue, SampleTimeOffset, SampleIdx, CurveName, SampleRole);
+		FFeatureVectorHelper::EncodeFloat(Indexer.GetPoseVector(SampleIdx), ChannelDataOffset, CurveValue);
+	}
+
+	return true;
+}
+
+UE::PoseSearch::TLabelBuilder& UPoseSearchFeatureChannel_Curve::GetLabel(UE::PoseSearch::TLabelBuilder& LabelBuilder, UE::PoseSearch::ELabelFormat LabelFormat) const
+{
+	using namespace UE::PoseSearch;
+
+	GetOuterLabel(LabelBuilder, LabelFormat);
+	AppendLabelSeparator(LabelBuilder, LabelFormat);
+
+	LabelBuilder.Append(TEXT("Curve_"));
+	LabelBuilder.Append(CurveName.ToString());
+	AppendLabelSeparator(LabelBuilder, LabelFormat, true);
+	LabelBuilder.Appendf(TEXT("%.2f"), SampleTimeOffset);
+	return LabelBuilder;
+}
+#endif
\ No newline at end of file
diff --git a/Engine/Plugins/Animation/PoseSearch/Source/Runtime/Private/PoseSearchFeatureChannel_Curve.h b/Engine/Plugins/Animation/PoseSearch/Source/Runtime/Private/PoseSearchFeatureChannel_Curve.h
new file mode 100644
--- /dev/null
+++ b/Engine/Plugins/Animation/PoseSearch/Source/Runtime/Private/PoseSearchFeatureChannel_Curve.h
@@ -0,0 +1,67 @@
+// Copyright Epic Games, Inc. All Rights Reserved.
+
+#pragma once
+
+#include "PoseSearch/PoseSearchFeatureChannel.h"
+#include "PoseSearchFeatureChannel_Curve.generated.h"
+
+UCLASS(Experimental, BlueprintType, EditInlineNew, Blueprintable, meta = (DisplayName = "Curve Channel"), CollapseCategories)
+class POSESEARCH_API UPoseSearchFeatureChannel_Curve : public UPoseSearchFeatureChannel
+{
+	GENERATED_BODY()
+
+public:
+
+	UPROPERTY(EditAnywhere, Category = "Settings")
+	FName CurveName = NAME_None;
+
+	UPROPERTY(EditAnywhere, Category = "Settings")
+	FName SampleRole = UE::PoseSearch::DefaultRole;
+
+	UPROPERTY()
+	int8 CurveIdx;
+
+#if WITH_EDITORONLY_DATA
+	UPROPERTY(EditAnywhere, Category = "Settings")
+	float Weight = 1.f;
+#endif // WITH_EDITORONLY_DATA
+	
+	// the data relative to the sampling time associated to this channel will be offset by SampleTimeOffset seconds.
+	// For example, if Curve is DistanceToWall, and SampleTimeOffset is 0.5, this channel will try to match the future DistanceToWall curve of the character 0.5 seconds ahead
+	UPROPERTY(EditAnywhere, Category = "Settings")
+	float SampleTimeOffset = 0.f;
+
+#if WITH_EDITORONLY_DATA
+	UPROPERTY(EditAnywhere, Category = "Settings", meta = (ExcludeFromHash, DisplayPriority = 0))
+	FLinearColor DebugColor = FLinearColor::Green;
+#endif // WITH_EDITORONLY_DATA
+
+	UPROPERTY(EditAnywhere, Category = "Settings")
+	EInputQueryPose InputQueryPose = EInputQueryPose::UseContinuingPose;
+
+#if WITH_EDITORONLY_DATA
+	// if set, all the channels of the same class with the same cardinality, and the same NormalizationGroup, will be normalized together.
+	// for example in a locomotion database of a character holding a weapon, containing non mirrorable animations, you'd still want to normalize together 
+	// left foot and right foot position and Curve
+	UPROPERTY(EditAnywhere, Category = "Settings")
+	FName NormalizationGroup;
+#endif //WITH_EDITORONLY_DATA
+
+	UFUNCTION(BlueprintPure, BlueprintImplementableEvent, meta=(BlueprintThreadSafe, DisplayName = "Get World Curve"), Category = "Settings")
+	float BP_GetCurveValue(const UAnimInstance* AnimInstance) const;
+
+	bool bUseBlueprintQueryOverride = false;
+
+	UPoseSearchFeatureChannel_Curve();
+
+	// UPoseSearchFeatureChannel interface
+	virtual bool Finalize(UPoseSearchSchema* Schema) override;
+	virtual void BuildQuery(UE::PoseSearch::FSearchContext& SearchContext) const override;
+
+#if WITH_EDITOR
+	virtual void FillWeights(TArrayView<float> Weights) const override;
+	virtual bool IndexAsset(UE::PoseSearch::FAssetIndexer& Indexer) const override;
+	virtual UE::PoseSearch::TLabelBuilder& GetLabel(UE::PoseSearch::TLabelBuilder& LabelBuilder, UE::PoseSearch::ELabelFormat LabelFormat = UE::PoseSearch::ELabelFormat::Full_Horizontal) const override;
+	virtual FName GetNormalizationGroup() const override { return NormalizationGroup; }
+#endif
+};
diff --git a/Engine/Plugins/Animation/PoseSearch/Source/Runtime/Private/PoseSearchSchema.cpp b/Engine/Plugins/Animation/PoseSearch/Source/Runtime/Private/PoseSearchSchema.cpp
--- a/Engine/Plugins/Animation/PoseSearch/Source/Runtime/Private/PoseSearchSchema.cpp
+++ b/Engine/Plugins/Animation/PoseSearch/Source/Runtime/Private/PoseSearchSchema.cpp
@@ -83,7 +83,11 @@ void UPoseSearchSchema::InitBoneContainersFromRoledSkeleton(TMap<FName, FBoneCon
 
 	for (const FPoseSearchRoledSkeleton& RoledSkeleton : Skeletons)
 	{
-		RoledBoneContainers.Add(RoledSkeleton.Role).InitializeTo(RoledSkeleton.BoneIndicesWithParents, UE::Anim::FCurveFilterSettings(UE::Anim::ECurveFilterMode::DisallowAll), *RoledSkeleton.Skeleton);
+		FBoneContainer& RoledBoneContainer = RoledBoneContainers.Add(RoledSkeleton.Role);
+		// Add a curve filter to our bone container to only eval curves actually used by the schema.
+		const UE::Anim::FCurveFilterSettings CurveFilterSettings(UE::Anim::ECurveFilterMode::AllowOnlyFiltered, &RoledSkeleton.RequiredCurves);
+
+		RoledBoneContainer.InitializeTo(RoledSkeleton.BoneIndicesWithParents, CurveFilterSettings, *RoledSkeleton.Skeleton);
 	}
 }
 
@@ -202,6 +206,31 @@ int8 UPoseSearchSchema::AddBoneReference(const FBoneReference& BoneReference, co
 	return int8(SchemaBoneIdx);
 }
 
+int8 UPoseSearchSchema::AddCurveReference(const FName& CurveReference, const UE::PoseSearch::FRole& Role)
+{
+	using namespace UE::PoseSearch;
+
+	FPoseSearchRoledSkeleton* RoledSkeleton = GetRoledSkeleton(Role);
+	if (!RoledSkeleton)
+	{
+		UE_LOG(LogPoseSearch, Error, TEXT("UPoseSearchSchema::AddCurveReference: couldn't find data for the requested Role '%s' in UPoseSearchSchema '%s'"), *Role.ToString(), *GetNameSafe(this));
+		return -1;
+	}
+
+	int32 SchemaBoneIdx = 0;
+	const USkeleton* Skeleton = RoledSkeleton->Skeleton;
+	if (!Skeleton)
+	{
+		UE_LOG(LogPoseSearch, Error, TEXT("UPoseSearchSchema::AddCurveReference: couldn't find Skeleton with Role '%s' in UPoseSearchSchema '%s'"), *Role.ToString(), *GetNameSafe(this));
+		return -1;
+	}
+
+	// Curves are loosely bound, so there's no guarantee this curve will ever exist in any of the assets indexed by the database.
+	const int32 CurveIdx = RoledSkeleton->RequiredCurves.AddUnique(CurveReference);
+	check(CurveIdx >= 0 && CurveIdx < 128);
+	return int8(CurveIdx);
+}
+
 void UPoseSearchSchema::ResetFinalize()
 {
 	for (FPoseSearchRoledSkeleton& RoledSkeleton : Skeletons)
diff --git a/Engine/Plugins/Animation/PoseSearch/Source/Runtime/Public/PoseSearch/PoseSearchAssetIndexer.h b/Engine/Plugins/Animation/PoseSearch/Source/Runtime/Public/PoseSearch/PoseSearchAssetIndexer.h
--- a/Engine/Plugins/Animation/PoseSearch/Source/Runtime/Public/PoseSearch/PoseSearchAssetIndexer.h
+++ b/Engine/Plugins/Animation/PoseSearch/Source/Runtime/Public/PoseSearch/PoseSearchAssetIndexer.h
@@ -42,6 +42,7 @@ struct FAnimationAssetSamplers
 	FTransform ExtractRootTransform(float Time, int32 RoleIndex) const;
 	FTransform GetTotalRootTransform(int32 RoleIndex) const;
 	void ExtractPose(float Time, FCompactPose& OutPose, int32 RoleIndex) const;
+	void ExtractPose(float Time, FCompactPose& OutPose, FBlendedCurve& OutCurve, int32 RoleIndex) const;
 	FTransform MirrorTransform(const FTransform& InTransform, int32 RoleIndex) const;
 	void MirrorPose(FCompactPose& Pose, int32 RoleIndex) const;
 
@@ -67,6 +68,9 @@ public:
 	void Process(int32 AssetIdx);
 	const FStats& GetStats() const { return Stats; }
 
+	// Returns the value of float curve CurveName at time CalculateSampleTime(SampleIdx) + SampleTimeOffset.
+	bool GetSampleCurveValue(float& OutCurveValue, float SampleTimeOffset, int32 SampleIdx, const FName& CurveName, const FRole& SampleRole);
+
 	// Returns OutSampleRotation as the rotation of the bone Schema.BoneReferences[SchemaSampleBoneIdx] at time CalculateSampleTime(SampleIdx) + SampleTimeOffset relative to the
 	// transform of the bone Schema.BoneReferences[SchemaOriginBoneIdx] at time CalculateSampleTime(SampleIdx) + OriginTimeOffset 
 	// Times will be processed by GetPermutationTimeOffsets(PermutationTimeType, ...)
@@ -111,6 +115,9 @@ private:
 	// bClamped will be true if SampleTime is outside the animation duration boundaries
 	FTransform GetTransform(float SampleTime, const FRole& Role, bool& bClamped, int8 SchemaBoneIdx = RootSchemaBoneIdx);
 	FTransform GetTransform(float SampleTime, int32 RoleIndex, bool& bClamped, const FBoneReference& BoneReference);
+	
+	// Returns the value of float curve CurveName at time SampleTime
+	float GetSampleCurveValueInternal(float SampleTime, const FName& CurveName, const FRole& Role);
 
 	// Returns the component space transform of the bone Schema.BoneReferences[SchemaBoneIdx] at time SampleTime
 	// bClamped will be true if SampleTime is outside the animation duration boundaries
@@ -129,6 +136,7 @@ private:
 		// RootTransform and ComponentSpacePose are stored mirrored in case SearchIndexAsset.IsMirrored
 		TArray<FTransform> RootTransform;
 		TArray<FCSPose<FCompactHeapPose>> ComponentSpacePose;
+		TArray<FBlendedHeapCurve> Curves;
 	};
 
 	void GetSampleInfo(float SampleTime, int32 RoleIndex, FTransform& OutRootTransform, float& OutClipTime, bool& bOutClamped) const;
diff --git a/Engine/Plugins/Animation/PoseSearch/Source/Runtime/Public/PoseSearch/PoseSearchAssetSampler.h b/Engine/Plugins/Animation/PoseSearch/Source/Runtime/Public/PoseSearch/PoseSearchAssetSampler.h
--- a/Engine/Plugins/Animation/PoseSearch/Source/Runtime/Public/PoseSearch/PoseSearchAssetSampler.h
+++ b/Engine/Plugins/Animation/PoseSearch/Source/Runtime/Public/PoseSearch/PoseSearchAssetSampler.h
@@ -35,6 +35,7 @@ struct POSESEARCH_API FAnimationAssetSampler
 
 	// Extracts pose for this asset at a given Time
 	void ExtractPose(float Time, FCompactPose& OutPose) const;
+	void ExtractPose(float Time, FCompactPose& OutPose, FBlendedCurve& OutCurve) const;
 
 	// Extracts root transform at the given time, using the extremities of the sequence to extrapolate beyond the 
 	// sequence limits when Time is less than zero or greater than the sequence length.
diff --git a/Engine/Plugins/Animation/PoseSearch/Source/Runtime/Public/PoseSearch/PoseSearchContext.h b/Engine/Plugins/Animation/PoseSearch/Source/Runtime/Public/PoseSearch/PoseSearchContext.h
--- a/Engine/Plugins/Animation/PoseSearch/Source/Runtime/Public/PoseSearch/PoseSearchContext.h
+++ b/Engine/Plugins/Animation/PoseSearch/Source/Runtime/Public/PoseSearch/PoseSearchContext.h
@@ -127,6 +127,10 @@ struct POSESEARCH_API FSearchContext
 
 	void AddRole(const FRole& Role, const UAnimInstance* AnimInstance, const IPoseHistory* PoseHistory);
 
+	// Returns the curve value of name CurveName at an offset time of SampleTimeOffset.
+	// If the curve is not found, assume value of 0, which is consistent with curve behavior in the animation update.
+	float GetSampleCurveValue(float SampleTimeOffset, const FName& CurveName, const FRole& SampleRole);
+
 	// Returns the rotation of the bone Schema.BoneReferences[SchemaSampleBoneIdx] at an offset time of SampleTimeOffset relative to the
 	// transform of the bone Schema.BoneReferences[SchemaOriginBoneIdx] at an offset time of time OriginTimeOffset 
 	// Times will be processed by GetPermutationTimeOffsets(PermutationTimeType, ...)
@@ -191,6 +195,7 @@ struct POSESEARCH_API FSearchContext
 	void SetUseCachedChannelData(bool bInUseCachedChannelData) { bUseCachedChannelData = bInUseCachedChannelData; }
 
 private:
+	float GetSampleCurveValueInternal(float SampleTime, const FName& CurveName, const FRole& SampleRole);
 	FVector GetSamplePositionInternal(float SampleTime, float OriginTime, int8 SchemaSampleBoneIdx, int8 SchemaOriginBoneIdx, const FRole& SampleRole, const FRole& OriginRole, const FVector* SampleBonePositionWorldOverride = nullptr);
 	FQuat GetSampleRotationInternal(float SampleTime, float OriginTime, int8 SchemaSampleBoneIdx, int8 SchemaOriginBoneIdx, const FRole& SampleRole, const FRole& OriginRole, const FQuat* SampleBoneRotationWorldOverride = nullptr);
 	FTransform GetWorldRootBoneTransformAtTime(float SampleTime, const FRole& SampleRole) const;
diff --git a/Engine/Plugins/Animation/PoseSearch/Source/Runtime/Public/PoseSearch/PoseSearchSchema.h b/Engine/Plugins/Animation/PoseSearch/Source/Runtime/Public/PoseSearch/PoseSearchSchema.h
--- a/Engine/Plugins/Animation/PoseSearch/Source/Runtime/Public/PoseSearch/PoseSearchSchema.h
+++ b/Engine/Plugins/Animation/PoseSearch/Source/Runtime/Public/PoseSearch/PoseSearchSchema.h
@@ -48,6 +48,9 @@ struct POSESEARCH_API FPoseSearchRoledSkeleton
 
 	UPROPERTY(Transient)
 	TArray<uint16> BoneIndicesWithParents;
+
+	UPROPERTY(Transient)
+	TArray<FName> RequiredCurves;
 };
 
 /**
@@ -144,6 +147,7 @@ public:
 	virtual void PostLoad() override;
 
 	int8 AddBoneReference(const FBoneReference& BoneReference, const UE::PoseSearch::FRole& Role);
+	int8 AddCurveReference(const FName& CurveReference, const UE::PoseSearch::FRole& Role);
 
 #if WITH_EDITOR
 	virtual void PostEditChangeProperty(struct FPropertyChangedEvent& PropertyChangedEvent) override;
