motion matching - deprecating UPoseSearchDatabase::GetAnimationAssets

diff --git a/Engine/Plugins/Animation/PoseSearch/Source/Editor/Private/PoseSearchCustomization.cpp b/Engine/Plugins/Animation/PoseSearch/Source/Editor/Private/PoseSearchCustomization.cpp
--- a/Engine/Plugins/Animation/PoseSearch/Source/Editor/Private/PoseSearchCustomization.cpp
+++ b/Engine/Plugins/Animation/PoseSearch/Source/Editor/Private/PoseSearchCustomization.cpp
@@ -80,25 +80,5 @@ void FPoseSearchDatabaseSequenceCustomization::CustomizeChildren(TSharedRef<IPro
 	}
 }
 
-
-///////////////////////////////////////////////////////////////////////////////////////////////////////////////////////
-/// FPoseSearchDatabaseDetails
-
-TSharedRef<IDetailCustomization> FPoseSearchDatabaseDetails::MakeInstance()
-{
-	return MakeShareable(new FPoseSearchDatabaseDetails);
-}
-
-void FPoseSearchDatabaseDetails::CustomizeDetails(IDetailLayoutBuilder& DetailBuilder)
-{
-	TArray<TSharedPtr<IPropertyHandle>> HiddenHandles;
-	HiddenHandles.Add(DetailBuilder.GetProperty(GET_MEMBER_NAME_CHECKED(UPoseSearchDatabase, AnimationAssets)));
-
-	for (TSharedPtr<IPropertyHandle> PropertyHandle : HiddenHandles)
-	{
-		DetailBuilder.HideProperty(PropertyHandle);
-	}
-}
-
 #undef LOCTEXT_NAMESPACE
 
diff --git a/Engine/Plugins/Animation/PoseSearch/Source/Editor/Private/PoseSearchDatabaseAssetListItem.cpp b/Engine/Plugins/Animation/PoseSearch/Source/Editor/Private/PoseSearchDatabaseAssetListItem.cpp
--- a/Engine/Plugins/Animation/PoseSearch/Source/Editor/Private/PoseSearchDatabaseAssetListItem.cpp
+++ b/Engine/Plugins/Animation/PoseSearch/Source/Editor/Private/PoseSearchDatabaseAssetListItem.cpp
@@ -707,7 +707,7 @@ namespace UE::PoseSearch
 		{
 			if (const UPoseSearchDatabase* Database = ViewModelPtr->GetPoseSearchDatabase())
 			{
-				if (Database->GetAnimationAssets().IsValidIndex(AssetTreeNode->SourceAssetIdx))
+				if (Database->GetDatabaseAnimationAsset<FPoseSearchDatabaseAnimationAssetBase>(AssetTreeNode->SourceAssetIdx))
 				{
 					return ViewModelPtr->IsEnabled(AssetTreeNode->SourceAssetIdx);
 				}
@@ -786,7 +786,7 @@ namespace UE::PoseSearch
 		{
 			if (TSharedPtr<FDatabaseAssetTreeNode> AssetTreeNode = WeakAssetTreeNode.Pin())
 			{
-				if (Database->GetAnimationAssets().IsValidIndex(AssetTreeNode->SourceAssetIdx))
+				if (Database->GetDatabaseAnimationAsset<FPoseSearchDatabaseAnimationAssetBase>(AssetTreeNode->SourceAssetIdx))
 				{
 					if (ViewModelPtr->IsDisableReselection(AssetTreeNode->SourceAssetIdx))
 					{
@@ -826,7 +826,7 @@ namespace UE::PoseSearch
 		{
 			if (TSharedPtr<FDatabaseAssetTreeNode> AssetTreeNode = WeakAssetTreeNode.Pin())
 			{
-				if (Database->GetAnimationAssets().IsValidIndex(AssetTreeNode->SourceAssetIdx))
+				if (Database->GetDatabaseAnimationAsset<FPoseSearchDatabaseAnimationAssetBase>(AssetTreeNode->SourceAssetIdx))
 				{
 					if (ViewModelPtr->IsEnabled(AssetTreeNode->SourceAssetIdx))
 					{
diff --git a/Engine/Plugins/Animation/PoseSearch/Source/Editor/Private/PoseSearchDatabaseAssetTree.cpp b/Engine/Plugins/Animation/PoseSearch/Source/Editor/Private/PoseSearchDatabaseAssetTree.cpp
--- a/Engine/Plugins/Animation/PoseSearch/Source/Editor/Private/PoseSearchDatabaseAssetTree.cpp
+++ b/Engine/Plugins/Animation/PoseSearch/Source/Editor/Private/PoseSearchDatabaseAssetTree.cpp
@@ -238,8 +238,8 @@ namespace UE::PoseSearch
 		{
 			// Build an index based off of alphabetical order than iterate the index instead
 			TArray<uint32> IndexArray;
-			IndexArray.SetNumUninitialized(Database->GetAnimationAssets().Num());
-			for (int32 AnimationAssetIdx = 0; AnimationAssetIdx < Database->GetAnimationAssets().Num(); ++AnimationAssetIdx)
+			IndexArray.SetNumUninitialized(Database->GetNumAnimationAssets());
+			for (int32 AnimationAssetIdx = 0; AnimationAssetIdx < Database->GetNumAnimationAssets(); ++AnimationAssetIdx)
 			{
 				IndexArray[AnimationAssetIdx] = AnimationAssetIdx;
 			}
@@ -274,7 +274,7 @@ namespace UE::PoseSearch
 			});
 
 			// create all nodes
-			for (int32 AnimationAssetIdx = 0; AnimationAssetIdx < Database->GetAnimationAssets().Num(); ++AnimationAssetIdx)
+			for (int32 AnimationAssetIdx = 0; AnimationAssetIdx < Database->GetNumAnimationAssets(); ++AnimationAssetIdx)
 			{
 				const int32 MappedId = IndexArray[AnimationAssetIdx];
 
diff --git a/Engine/Plugins/Animation/PoseSearch/Source/Editor/Private/PoseSearchDatabaseEditor.cpp b/Engine/Plugins/Animation/PoseSearch/Source/Editor/Private/PoseSearchDatabaseEditor.cpp
--- a/Engine/Plugins/Animation/PoseSearch/Source/Editor/Private/PoseSearchDatabaseEditor.cpp
+++ b/Engine/Plugins/Animation/PoseSearch/Source/Editor/Private/PoseSearchDatabaseEditor.cpp
@@ -558,16 +558,12 @@ namespace UE::PoseSearch
 		{
 			for (TSharedPtr<FDatabaseAssetTreeNode>& SelectedItem : SelectedItems)
 			{
-				if (!SelectedItem.IsValid() || !PoseSearchDatabase->GetAnimationAssets().IsValidIndex(SelectedItem->SourceAssetIdx))
+				if (!SelectedItem.IsValid() || !PoseSearchDatabase->GetDatabaseAnimationAsset<FPoseSearchDatabaseAnimationAssetBase>(SelectedItem->SourceAssetIdx))
 				{
 					continue;
 				}
 
-				const FInstancedStruct& DatabaseAsset = PoseSearchDatabase->GetAnimationAssets()[SelectedItem->SourceAssetIdx];
-				const UScriptStruct* ScriptStruct = DatabaseAsset.GetScriptStruct();
-				FSelectionWidget& SelectionWidget = FindOrAddSelectionWidget(ScriptStruct);
-
-				if (const FPoseSearchDatabaseSequence* DatabaseSequence = DatabaseAsset.GetPtr<FPoseSearchDatabaseSequence>())
+				if (const FPoseSearchDatabaseSequence* DatabaseSequence = PoseSearchDatabase->GetDatabaseAnimationAsset<FPoseSearchDatabaseSequence>(SelectedItem->SourceAssetIdx))
 				{
 					UPoseSearchDatabaseSequenceReflection* NewSelectionReflection = NewObject<UPoseSearchDatabaseSequenceReflection>();
 					NewSelectionReflection->AddToRoot();
@@ -576,10 +572,10 @@ namespace UE::PoseSearch
 					NewSelectionReflection->Sequence.bHasRootMotion = DatabaseSequence->IsRootMotionEnabled();
 					NewSelectionReflection->SetSourceLink(SelectedItem, AssetTreeWidget);
 					NewSelectionReflection->SetFlags(RF_Transactional);
-
-					SelectionWidget.SelectedReflections.Add(NewSelectionReflection);
+					
+					FindOrAddSelectionWidget(FPoseSearchDatabaseSequence::StaticStruct()).SelectedReflections.Add(NewSelectionReflection);
 				}
-				else if (const FPoseSearchDatabaseAnimComposite* DatabaseAnimComposite = DatabaseAsset.GetPtr<FPoseSearchDatabaseAnimComposite>())
+				else if (const FPoseSearchDatabaseAnimComposite* DatabaseAnimComposite = PoseSearchDatabase->GetDatabaseAnimationAsset<FPoseSearchDatabaseAnimComposite>(SelectedItem->SourceAssetIdx))
 				{
 					UPoseSearchDatabaseAnimCompositeReflection* NewSelectionReflection = NewObject<UPoseSearchDatabaseAnimCompositeReflection>();
 					NewSelectionReflection->AddToRoot();
@@ -589,9 +585,9 @@ namespace UE::PoseSearch
 					NewSelectionReflection->SetSourceLink(SelectedItem, AssetTreeWidget);
 					NewSelectionReflection->SetFlags(RF_Transactional);
 					
-					SelectionWidget.SelectedReflections.Add(NewSelectionReflection);
+					FindOrAddSelectionWidget(FPoseSearchDatabaseAnimComposite::StaticStruct()).SelectedReflections.Add(NewSelectionReflection);
 				}
-				else if (const FPoseSearchDatabaseBlendSpace* DatabaseBlendSpace = DatabaseAsset.GetPtr<FPoseSearchDatabaseBlendSpace>())
+				else if (const FPoseSearchDatabaseBlendSpace* DatabaseBlendSpace = PoseSearchDatabase->GetDatabaseAnimationAsset<FPoseSearchDatabaseBlendSpace>(SelectedItem->SourceAssetIdx))
 				{
 					UPoseSearchDatabaseBlendSpaceReflection* NewSelectionReflection = NewObject<UPoseSearchDatabaseBlendSpaceReflection>();
 					NewSelectionReflection->AddToRoot();
@@ -601,9 +597,9 @@ namespace UE::PoseSearch
 					NewSelectionReflection->SetSourceLink(SelectedItem, AssetTreeWidget);
 					NewSelectionReflection->SetFlags(RF_Transactional);
 					
-					SelectionWidget.SelectedReflections.Add(NewSelectionReflection);
+					FindOrAddSelectionWidget(FPoseSearchDatabaseBlendSpace::StaticStruct()).SelectedReflections.Add(NewSelectionReflection);
 				}
-				else if (const FPoseSearchDatabaseAnimMontage* DatabaseAnimMontage = DatabaseAsset.GetPtr<FPoseSearchDatabaseAnimMontage>())
+				else if (const FPoseSearchDatabaseAnimMontage* DatabaseAnimMontage = PoseSearchDatabase->GetDatabaseAnimationAsset<FPoseSearchDatabaseAnimMontage>(SelectedItem->SourceAssetIdx))
 				{
 					UPoseSearchDatabaseAnimMontageReflection* NewSelectionReflection = NewObject<UPoseSearchDatabaseAnimMontageReflection>();
 					NewSelectionReflection->AddToRoot();
@@ -613,9 +609,9 @@ namespace UE::PoseSearch
 					NewSelectionReflection->SetSourceLink(SelectedItem, AssetTreeWidget);
 					NewSelectionReflection->SetFlags(RF_Transactional);
 
-					SelectionWidget.SelectedReflections.Add(NewSelectionReflection);
+					FindOrAddSelectionWidget(FPoseSearchDatabaseAnimMontage::StaticStruct()).SelectedReflections.Add(NewSelectionReflection);
 				}
-				else if (const FPoseSearchDatabaseMultiAnimAsset* DatabaseMultiAnimAsset = DatabaseAsset.GetPtr<FPoseSearchDatabaseMultiAnimAsset>())
+				else if (const FPoseSearchDatabaseMultiAnimAsset* DatabaseMultiAnimAsset = PoseSearchDatabase->GetDatabaseAnimationAsset<FPoseSearchDatabaseMultiAnimAsset>(SelectedItem->SourceAssetIdx))
 				{
 					UPoseSearchDatabaseMultiAnimAssetReflection* NewSelectionReflection = NewObject<UPoseSearchDatabaseMultiAnimAssetReflection>();
 					NewSelectionReflection->AddToRoot();
@@ -625,7 +621,7 @@ namespace UE::PoseSearch
 					NewSelectionReflection->SetSourceLink(SelectedItem, AssetTreeWidget);
 					NewSelectionReflection->SetFlags(RF_Transactional);
 
-					SelectionWidget.SelectedReflections.Add(NewSelectionReflection);
+					FindOrAddSelectionWidget(FPoseSearchDatabaseMultiAnimAsset::StaticStruct()).SelectedReflections.Add(NewSelectionReflection);
 				}
 				else
 				{
diff --git a/Engine/Plugins/Animation/PoseSearch/Source/Editor/Private/PoseSearchEditor.cpp b/Engine/Plugins/Animation/PoseSearch/Source/Editor/Private/PoseSearchEditor.cpp
--- a/Engine/Plugins/Animation/PoseSearch/Source/Editor/Private/PoseSearchEditor.cpp
+++ b/Engine/Plugins/Animation/PoseSearch/Source/Editor/Private/PoseSearchEditor.cpp
@@ -43,7 +43,6 @@ private:
 	
 private:
 	void RegisterPropertyTypeCustomizations();
-	void RegisterObjectCustomizations();
 	void UnregisterCustomizations();
 	void RegisterCustomClassLayout(FName ClassName, FOnGetDetailCustomizationInstance DetailLayoutDelegate);
 	void RegisterCustomPropertyTypeLayout(FName PropertyTypeName, FOnGetPropertyTypeCustomizationInstance PropertyTypeLayoutDelegate);
@@ -73,7 +72,6 @@ void FEditorModule::StartupModule()
 	}
 
 	RegisterPropertyTypeCustomizations();
-	RegisterObjectCustomizations();
 
 	FPropertyEditorModule& PropertyModule = FModuleManager::GetModuleChecked<FPropertyEditorModule>("PropertyEditor");
 	PropertyModule.NotifyCustomizationModuleChanged();
@@ -126,11 +124,6 @@ void FEditorModule::RegisterPropertyTypeCustomizations()
 	RegisterCustomPropertyTypeLayout("PoseSearchDatabaseSequence", FOnGetPropertyTypeCustomizationInstance::CreateStatic(&FPoseSearchDatabaseSequenceCustomization::MakeInstance));
 }
 
-void FEditorModule::RegisterObjectCustomizations()
-{
-	RegisterCustomClassLayout("PoseSearchDatabase", FOnGetDetailCustomizationInstance::CreateStatic(&FPoseSearchDatabaseDetails::MakeInstance));
-}
-
 void FEditorModule::UnregisterCustomizations()
 {
 	if (FModuleManager::Get().IsModuleLoaded("PropertyEditor"))
diff --git a/Engine/Plugins/Animation/PoseSearch/Source/Editor/Public/PoseSearchCustomization.h b/Engine/Plugins/Animation/PoseSearch/Source/Editor/Public/PoseSearchCustomization.h
--- a/Engine/Plugins/Animation/PoseSearch/Source/Editor/Public/PoseSearchCustomization.h
+++ b/Engine/Plugins/Animation/PoseSearch/Source/Editor/Public/PoseSearchCustomization.h
@@ -29,18 +29,3 @@ public:
 	virtual void CustomizeChildren(TSharedRef<class IPropertyHandle> InStructPropertyHandle, class IDetailChildrenBuilder& StructBuilder, IPropertyTypeCustomizationUtils& StructCustomizationUtils) override;
 };
 
-
-
-/**
- * Customizes UPoseSearchDatabase to improve workflow
- */
-class FPoseSearchDatabaseDetails : public IDetailCustomization
-{
-public:
-	// Creates an instance of FObjectDetails
-	static TSharedRef<IDetailCustomization> MakeInstance();
-
-	// IDetailCustomization interface
-	virtual void CustomizeDetails(IDetailLayoutBuilder& DetailLayout) override;
-	// End of IDetailCustomization interface
-};
diff --git a/Engine/Plugins/Animation/PoseSearch/Source/Runtime/Private/PoseSearchDatabase.cpp b/Engine/Plugins/Animation/PoseSearch/Source/Runtime/Private/PoseSearchDatabase.cpp
--- a/Engine/Plugins/Animation/PoseSearch/Source/Runtime/Private/PoseSearchDatabase.cpp
+++ b/Engine/Plugins/Animation/PoseSearch/Source/Runtime/Private/PoseSearchDatabase.cpp
@@ -668,9 +668,11 @@ bool FPoseSearchDatabaseMultiAnimAsset::IsRootMotionEnabled() const
 
 //////////////////////////////////////////////////////////////////////////
 // UPoseSearchDatabase
+PRAGMA_DISABLE_DEPRECATION_WARNINGS
 UPoseSearchDatabase::~UPoseSearchDatabase()
 {
 }
+PRAGMA_ENABLE_DEPRECATION_WARNINGS
 
 void UPoseSearchDatabase::SetSearchIndex(const UE::PoseSearch::FSearchIndex& SearchIndex)
 {
@@ -692,36 +694,55 @@ int32 UPoseSearchDatabase::GetPoseIndexFromTime(float Time, const UE::PoseSearch
 
 void UPoseSearchDatabase::AddAnimationAsset(FInstancedStruct AnimationAsset)
 {
+	PRAGMA_DISABLE_DEPRECATION_WARNINGS
 	AnimationAssets.Add(AnimationAsset);
+	PRAGMA_ENABLE_DEPRECATION_WARNINGS
 }
 
 void UPoseSearchDatabase::RemoveAnimationAssetAt(int32 AnimationAssetIndex)
 {
+	PRAGMA_DISABLE_DEPRECATION_WARNINGS
 	AnimationAssets.RemoveAt(AnimationAssetIndex);
+	PRAGMA_ENABLE_DEPRECATION_WARNINGS
+}
+
+const TArray<FInstancedStruct>& UPoseSearchDatabase::GetAnimationAssets() const
+{
+	PRAGMA_DISABLE_DEPRECATION_WARNINGS
+	return AnimationAssets;
+	PRAGMA_ENABLE_DEPRECATION_WARNINGS
 }
 
 const FInstancedStruct& UPoseSearchDatabase::GetAnimationAssetStruct(int32 AnimationAssetIndex) const
 {
+	PRAGMA_DISABLE_DEPRECATION_WARNINGS
 	check(AnimationAssets.IsValidIndex(AnimationAssetIndex));
 	return AnimationAssets[AnimationAssetIndex];
+	PRAGMA_ENABLE_DEPRECATION_WARNINGS
 }
 
 const FInstancedStruct& UPoseSearchDatabase::GetAnimationAssetStruct(const UE::PoseSearch::FSearchIndexAsset& SearchIndexAsset) const
 {
+	PRAGMA_DISABLE_DEPRECATION_WARNINGS
 	check(AnimationAssets.IsValidIndex(SearchIndexAsset.GetSourceAssetIdx()));
 	return AnimationAssets[SearchIndexAsset.GetSourceAssetIdx()];
+	PRAGMA_ENABLE_DEPRECATION_WARNINGS
 }
 
 FInstancedStruct& UPoseSearchDatabase::GetMutableAnimationAssetStruct(int32 AnimationAssetIndex)
 {
+	PRAGMA_DISABLE_DEPRECATION_WARNINGS
 	check(AnimationAssets.IsValidIndex(AnimationAssetIndex));
 	return AnimationAssets[AnimationAssetIndex];
+	PRAGMA_ENABLE_DEPRECATION_WARNINGS
 }
 
 FInstancedStruct& UPoseSearchDatabase::GetMutableAnimationAssetStruct(const UE::PoseSearch::FSearchIndexAsset& SearchIndexAsset)
 {
+	PRAGMA_DISABLE_DEPRECATION_WARNINGS
 	check(AnimationAssets.IsValidIndex(SearchIndexAsset.GetSourceAssetIdx()));
 	return AnimationAssets[SearchIndexAsset.GetSourceAssetIdx()];
+	PRAGMA_ENABLE_DEPRECATION_WARNINGS
 }
 
 #if WITH_EDITOR
@@ -748,7 +769,9 @@ void UPoseSearchDatabase::PostLoad()
 	using namespace UE::PoseSearch;
 
 	bool bRequiresSynchronization = false;
+	PRAGMA_DISABLE_DEPRECATION_WARNINGS
 	for (FInstancedStruct& AnimationAsset : AnimationAssets)
+	PRAGMA_ENABLE_DEPRECATION_WARNINGS
 	{
 		if (FPoseSearchDatabaseAnimationAssetBase* AnimationAssetBase = AnimationAsset.GetMutablePtr<FPoseSearchDatabaseAnimationAssetBase>())
 		{
@@ -765,7 +788,9 @@ void UPoseSearchDatabase::PostLoad()
 	{
 		SynchronizeWithExternalDependencies();
 
+		PRAGMA_DISABLE_DEPRECATION_WARNINGS
 		for (FInstancedStruct& AnimationAsset : AnimationAssets)
+		PRAGMA_ENABLE_DEPRECATION_WARNINGS
 		{
 			if (FPoseSearchDatabaseAnimationAssetBase* AnimationAssetBase = AnimationAsset.GetMutablePtr<FPoseSearchDatabaseAnimationAssetBase>())
 			{
@@ -791,7 +816,9 @@ void UPoseSearchDatabase::PostLoad()
 
 bool UPoseSearchDatabase::Contains(const UObject* Object) const
 {
+	PRAGMA_DISABLE_DEPRECATION_WARNINGS
 	for (const FInstancedStruct& AnimationAsset : AnimationAssets)
+	PRAGMA_ENABLE_DEPRECATION_WARNINGS
 	{
 		if (const FPoseSearchDatabaseAnimationAssetBase* AnimationAssetBase = AnimationAsset.GetPtr<FPoseSearchDatabaseAnimationAssetBase>())
 		{
@@ -806,11 +833,14 @@ bool UPoseSearchDatabase::Contains(const UObject* Object) const
 
 int32 UPoseSearchDatabase::GetNumAnimationAssets() const
 {
+	PRAGMA_DISABLE_DEPRECATION_WARNINGS
 	return AnimationAssets.Num();
+	PRAGMA_ENABLE_DEPRECATION_WARNINGS
 }
 
 UObject* UPoseSearchDatabase::GetAnimationAsset(int32 Index) const
 {
+	PRAGMA_DISABLE_DEPRECATION_WARNINGS
 	if (AnimationAssets.IsValidIndex(Index))
 	{
 		if (const FPoseSearchDatabaseAnimationAssetBase* AnimationAssetBase = AnimationAssets[Index].GetPtr<FPoseSearchDatabaseAnimationAssetBase>())
@@ -818,6 +848,7 @@ UObject* UPoseSearchDatabase::GetAnimationAsset(int32 Index) const
 			return AnimationAssetBase->GetAnimationAsset();
 		}
 	}
+	PRAGMA_ENABLE_DEPRECATION_WARNINGS
 	return nullptr;
 }
 
@@ -872,6 +903,7 @@ void UPoseSearchDatabase::SynchronizeWithExternalDependencies(TConstArrayView<UA
 
 	// collecting all the database AnimationAsset(s) that don't require synchronization
 	TArray<bool> DisableReselection;
+	PRAGMA_DISABLE_DEPRECATION_WARNINGS
 	DisableReselection.Reserve(AnimationAssets.Num());
 	for (FInstancedStruct& AnimationAsset : AnimationAssets)
 	{
@@ -962,6 +994,8 @@ void UPoseSearchDatabase::SynchronizeWithExternalDependencies(TConstArrayView<UA
 		bModified = true;
 	}
 
+	PRAGMA_ENABLE_DEPRECATION_WARNINGS
+
 	if (bModified)
 	{
 		Modify();
@@ -987,6 +1021,7 @@ bool UPoseSearchDatabase::IsCachedCookedPlatformDataLoaded(const ITargetPlatform
 #if WITH_EDITOR && ENABLE_ANIM_DEBUG
 void UPoseSearchDatabase::TestSynchronizeWithExternalDependencies()
 {
+	PRAGMA_DISABLE_DEPRECATION_WARNINGS
 	TArray<FInstancedStruct> AnimationAssetsCopy = AnimationAssets;
 	SynchronizeWithExternalDependencies();
 
@@ -995,6 +1030,7 @@ void UPoseSearchDatabase::TestSynchronizeWithExternalDependencies()
 		UE_LOG(LogPoseSearch, Error, TEXT("TestSynchronizeWithExternalDependencies failed"));
 		AnimationAssets = AnimationAssetsCopy;
 	}
+	PRAGMA_ENABLE_DEPRECATION_WARNINGS
 }
 #endif // WITH_EDITOR && ENABLE_ANIM_DEBUG
 
@@ -1045,7 +1081,7 @@ float UPoseSearchDatabase::GetNormalizedAssetTime(int32 PoseIdx) const
 {
 	check(Schema);
 	const UE::PoseSearch::FSearchIndexAsset& Asset = GetSearchIndex().GetAssetForPose(PoseIdx);
-	const bool bIsBlendSpace = AnimationAssets[Asset.GetSourceAssetIdx()].GetPtr<FPoseSearchDatabaseBlendSpace>() != nullptr;
+	const bool bIsBlendSpace = GetDatabaseAnimationAsset<FPoseSearchDatabaseBlendSpace>(Asset.GetSourceAssetIdx()) != nullptr;
 
 	// sequences or anim composites
 	float AssetTime = Asset.GetTimeFromPoseIndex(PoseIdx, Schema->SampleRate);
diff --git a/Engine/Plugins/Animation/PoseSearch/Source/Runtime/Private/PoseSearchDerivedData.cpp b/Engine/Plugins/Animation/PoseSearch/Source/Runtime/Private/PoseSearchDerivedData.cpp
--- a/Engine/Plugins/Animation/PoseSearch/Source/Runtime/Private/PoseSearchDerivedData.cpp
+++ b/Engine/Plugins/Animation/PoseSearch/Source/Runtime/Private/PoseSearchDerivedData.cpp
@@ -360,7 +360,7 @@ static void FindValidSequenceIntervals(const FPoseSearchDatabaseAnimationAssetBa
 }
 
 // returns false in case of errors
-static bool InitSearchIndexAssets(FSearchIndexBase& SearchIndex, const TArray<FInstancedStruct>& DatabaseAnimationAssets, const UPoseSearchSchema* Schema, const FFloatInterval& ExcludeFromDatabaseParameters)
+static bool InitSearchIndexAssets(FSearchIndexBase& SearchIndex, const UPoseSearchDatabase* DatabaseToLookForAssets, const UPoseSearchSchema* Schema, const FFloatInterval& ExcludeFromDatabaseParameters)
 {
 	using namespace UE::PoseSearch;
 
@@ -373,10 +373,9 @@ static bool InitSearchIndexAssets(FSearchIndexBase& SearchIndex, const TArray<FI
 	bool bAnyErrors = false;
 
 	int32 TotalPoses = 0;
-	for (int32 AnimationAssetIndex = 0; AnimationAssetIndex < DatabaseAnimationAssets.Num(); ++AnimationAssetIndex)
+	for (int32 AnimationAssetIndex = 0; AnimationAssetIndex < DatabaseToLookForAssets->GetNumAnimationAssets(); ++AnimationAssetIndex)
 	{
-		const FInstancedStruct& DatabaseAssetStruct = DatabaseAnimationAssets[AnimationAssetIndex];
-		if (const FPoseSearchDatabaseAnimationAssetBase* DatabaseAsset = DatabaseAssetStruct.GetPtr<FPoseSearchDatabaseAnimationAssetBase>())
+		if (const FPoseSearchDatabaseAnimationAssetBase* DatabaseAsset = DatabaseToLookForAssets->GetDatabaseAnimationAsset<FPoseSearchDatabaseAnimationAssetBase>(AnimationAssetIndex))
 		{
 			if (!DatabaseAsset->IsEnabled() || !DatabaseAsset->GetAnimationAsset())
 			{
@@ -411,7 +410,7 @@ static bool InitSearchIndexAssets(FSearchIndexBase& SearchIndex, const TArray<FI
 			const bool bIsLooping = DatabaseAsset->IsLooping();
 			const bool bDisableReselection = DatabaseAsset->IsDisableReselection();
 
-			if (const FPoseSearchDatabaseBlendSpace* DatabaseBlendSpace = DatabaseAssetStruct.GetPtr<FPoseSearchDatabaseBlendSpace>())
+			if (const FPoseSearchDatabaseBlendSpace* DatabaseBlendSpace = DatabaseToLookForAssets->GetDatabaseAnimationAsset<FPoseSearchDatabaseBlendSpace>(AnimationAssetIndex))
 			{
 				int32 HorizontalBlendNum, VerticalBlendNum;
 				DatabaseBlendSpace->GetBlendSpaceParameterSampleRanges(HorizontalBlendNum, VerticalBlendNum);
@@ -909,7 +908,7 @@ struct FSamplerMapKey
 	FVector BlendParameters = FVector::ZeroVector;
 };
 
-static bool IndexDatabase(FSearchIndexBase& SearchIndexBase, const TArray<FInstancedStruct>& DatabaseAnimationAssets, const UPoseSearchSchema* Schema, const FAssetSamplingContext& SamplingContext, const FFloatInterval& AdditionalExtrapolationTime, UE::DerivedData::FRequestOwner& Owner)
+static bool IndexDatabase(FSearchIndexBase& SearchIndexBase, const UPoseSearchDatabase* DatabaseToLookForAssets, const UPoseSearchSchema* Schema, const FAssetSamplingContext& SamplingContext, const FFloatInterval& AdditionalExtrapolationTime, UE::DerivedData::FRequestOwner& Owner)
 {
 	if (!ensure(Schema))
 	{
@@ -923,10 +922,9 @@ static bool IndexDatabase(FSearchIndexBase& SearchIndexBase, const TArray<FInsta
 	TMap<FSamplerMapKey, int32> SamplerMap;
 	SamplerMap.Reserve(256);
 
-	for (int32 AnimationAssetIndex = 0; AnimationAssetIndex < DatabaseAnimationAssets.Num(); ++AnimationAssetIndex)
+	for (int32 AnimationAssetIndex = 0; AnimationAssetIndex < DatabaseToLookForAssets->GetNumAnimationAssets(); ++AnimationAssetIndex)
 	{
-		const FInstancedStruct& DatabaseAssetStruct = DatabaseAnimationAssets[AnimationAssetIndex];
-		if (const FPoseSearchDatabaseBlendSpace* DatabaseBlendSpace = DatabaseAssetStruct.GetPtr<FPoseSearchDatabaseBlendSpace>())
+		if (const FPoseSearchDatabaseBlendSpace* DatabaseBlendSpace = DatabaseToLookForAssets->GetDatabaseAnimationAsset<FPoseSearchDatabaseBlendSpace>(AnimationAssetIndex))
 		{
 			if (DatabaseBlendSpace->BlendSpace)
 			{
@@ -951,7 +949,7 @@ static bool IndexDatabase(FSearchIndexBase& SearchIndexBase, const TArray<FInsta
 				}
 			}
 		}
-		if (const FPoseSearchDatabaseAnimationAssetBase* DatabaseAnimationAssetBase = DatabaseAssetStruct.GetPtr<FPoseSearchDatabaseAnimationAssetBase>())
+		if (const FPoseSearchDatabaseAnimationAssetBase* DatabaseAnimationAssetBase = DatabaseToLookForAssets->GetDatabaseAnimationAsset<FPoseSearchDatabaseAnimationAssetBase>(AnimationAssetIndex))
 		{
 			if (const UObject* AnimationAssetObject = DatabaseAnimationAssetBase->GetAnimationAsset())
 			{
@@ -1003,8 +1001,7 @@ static bool IndexDatabase(FSearchIndexBase& SearchIndexBase, const TArray<FInsta
 		FSearchIndexAsset& SearchIndexAsset = SearchIndexBase.Assets[AssetIdx];
 		check(SearchIndexAsset.GetFirstPoseIdx() == TotalPoses);
 
-		const FInstancedStruct& DatabaseAssetStruct = DatabaseAnimationAssets[SearchIndexAsset.GetSourceAssetIdx()];
-		const FPoseSearchDatabaseAnimationAssetBase* DatabaseAnimationAssetBase = DatabaseAssetStruct.GetPtr<FPoseSearchDatabaseAnimationAssetBase>();
+		const FPoseSearchDatabaseAnimationAssetBase* DatabaseAnimationAssetBase = DatabaseToLookForAssets->GetDatabaseAnimationAsset<FPoseSearchDatabaseAnimationAssetBase>(SearchIndexAsset.GetSourceAssetIdx());
 		check(DatabaseAnimationAssetBase);
 
 		TempAssetSamplers.Reset();
@@ -1634,7 +1631,6 @@ void FPoseSearchDatabaseAsyncCacheTask::OnGetComplete(UE::DerivedData::FCacheGet
 					FSearchIndexBase& SearchIndexBase = SearchIndexBases[IndexBaseIdx];
 
 					Schemas[IndexBaseIdx] = DependentDatabaseSchema;
-					const TArray<FInstancedStruct>& DependentDatabaseAnimationAssets = DependentDatabase->GetAnimationAssets();
 					const FAssetSamplingContext DependentSamplingContext(*DependentDatabase);
 					const FFloatInterval& DependentExcludeFromDatabaseParameters = DependentDatabase->ExcludeFromDatabaseParameters;
 					const FFloatInterval& DependentAdditionalExtrapolationTime = DependentDatabase->AdditionalExtrapolationTime;
@@ -1657,9 +1653,9 @@ void FPoseSearchDatabaseAsyncCacheTask::OnGetComplete(UE::DerivedData::FCacheGet
 					// validating that the missing MirrorDataTable(s) are not necessary 
 					if (!DependentDatabaseSchema->AllRoledSkeletonHaveMirrorDataTable())
 					{
-						for (int32 AnimationAssetIndex = 0; AnimationAssetIndex < DependentDatabaseAnimationAssets.Num(); ++AnimationAssetIndex)
+						for (int32 AnimationAssetIndex = 0; AnimationAssetIndex < DependentDatabase->GetNumAnimationAssets(); ++AnimationAssetIndex)
 						{
-							if (const FPoseSearchDatabaseAnimationAssetBase* DatabaseAsset = DependentDatabaseAnimationAssets[AnimationAssetIndex].GetPtr<FPoseSearchDatabaseAnimationAssetBase>())
+							if (const FPoseSearchDatabaseAnimationAssetBase* DatabaseAsset = DependentDatabase->GetDatabaseAnimationAsset<FPoseSearchDatabaseAnimationAssetBase>(AnimationAssetIndex))
 							{
 								if (DatabaseAsset->GetMirrorOption() == EPoseSearchMirrorOption::MirroredOnly || DatabaseAsset->GetMirrorOption() == EPoseSearchMirrorOption::UnmirroredAndMirrored)
 								{
@@ -1680,7 +1676,7 @@ void FPoseSearchDatabaseAsyncCacheTask::OnGetComplete(UE::DerivedData::FCacheGet
 					}
 
 					// Building all the related FPoseSearchBaseIndex first
-					if (!InitSearchIndexAssets(SearchIndexBase, DependentDatabaseAnimationAssets, DependentDatabaseSchema, DependentExcludeFromDatabaseParameters))
+					if (!InitSearchIndexAssets(SearchIndexBase, DependentDatabase, DependentDatabaseSchema, DependentExcludeFromDatabaseParameters))
 					{
 						UE_LOG(LogPoseSearch, Error, TEXT("%s - %s BuildIndex Failed becasue of invalid assets"), *LexToString(FullIndexKey.Hash), *MainDatabaseName);
 						ResetSearchIndex();
@@ -1694,7 +1690,7 @@ void FPoseSearchDatabaseAsyncCacheTask::OnGetComplete(UE::DerivedData::FCacheGet
 						return;
 					}
 
-					if (!IndexDatabase(SearchIndexBase, DependentDatabaseAnimationAssets, DependentDatabaseSchema, DependentSamplingContext, DependentAdditionalExtrapolationTime, Owner))
+					if (!IndexDatabase(SearchIndexBase, DependentDatabase, DependentDatabaseSchema, DependentSamplingContext, DependentAdditionalExtrapolationTime, Owner))
 					{
 						UE_LOG(LogPoseSearch, Log, TEXT("%s - %s BuildIndex Cancelled"), *LexToString(FullIndexKey.Hash), *MainDatabaseName);
 						ResetSearchIndex();
@@ -1708,7 +1704,7 @@ void FPoseSearchDatabaseAsyncCacheTask::OnGetComplete(UE::DerivedData::FCacheGet
 						for (int32 Iteration = 0; Iteration < NumIterations; ++Iteration)
 						{
 							FSearchIndexBase TestSearchIndexBase = SearchIndexBase;
-							if (IndexDatabase(TestSearchIndexBase, DependentDatabaseAnimationAssets, DependentDatabaseSchema, DependentSamplingContext, DependentAdditionalExtrapolationTime, Owner))
+							if (IndexDatabase(TestSearchIndexBase, DependentDatabase, DependentDatabaseSchema, DependentSamplingContext, DependentAdditionalExtrapolationTime, Owner))
 							{
 								if (TestSearchIndexBase != SearchIndexBase)
 								{
@@ -2178,9 +2174,9 @@ void FAsyncPoseSearchDatabasesManagement::Tick(float DeltaTime)
 					TMap<FRole, FBoneContainer> RoledBoneContainers;
 					Database->Schema->InitBoneContainersFromRoledSkeleton(RoledBoneContainers);
 
-					for (int32 AnimationAssetIndex = 0; AnimationAssetIndex < Database->GetAnimationAssets().Num(); ++AnimationAssetIndex)
+					for (int32 AnimationAssetIndex = 0; AnimationAssetIndex < Database->GetNumAnimationAssets(); ++AnimationAssetIndex)
 					{
-						if (const FPoseSearchDatabaseAnimationAssetBase* DatabaseAsset = Database->GetAnimationAssets()[AnimationAssetIndex].GetPtr<FPoseSearchDatabaseAnimationAssetBase>())
+						if (const FPoseSearchDatabaseAnimationAssetBase* DatabaseAsset = Database->GetDatabaseAnimationAsset<FPoseSearchDatabaseAnimationAssetBase>(AnimationAssetIndex))
 						{
 							for (int32 RoleIndex = 0; RoleIndex < DatabaseAsset->GetNumRoles(); ++RoleIndex)
 							{
diff --git a/Engine/Plugins/Animation/PoseSearch/Source/Runtime/Public/PoseSearch/PoseSearchDatabase.h b/Engine/Plugins/Animation/PoseSearch/Source/Runtime/Public/PoseSearch/PoseSearchDatabase.h
--- a/Engine/Plugins/Animation/PoseSearch/Source/Runtime/Public/PoseSearch/PoseSearchDatabase.h
+++ b/Engine/Plugins/Animation/PoseSearch/Source/Runtime/Public/PoseSearch/PoseSearchDatabase.h
@@ -425,7 +425,9 @@ public:
 	FFloatInterval AdditionalExtrapolationTime = FFloatInterval(-100.f, 100.f);
 #endif // WITH_EDITORONLY_DATA
 
-	UPROPERTY(EditAnywhere, Category = "Database")
+	// use GetNumAnimationAssets() and GetDatabaseAnimationAsset(...) APIs to query for AnimationAssets
+	UE_DEPRECATED(5.4, "This property will be made private")
+	UPROPERTY()
 	TArray<FInstancedStruct> AnimationAssets;
 
 	/** Array of tags that can be used as metadata. */
@@ -510,9 +512,8 @@ public:
 	template<typename TDatabaseAnimationAsset> TDatabaseAnimationAsset* GetMutableDatabaseAnimationAsset(int32 AnimationAssetIndex);
 	template<typename TDatabaseAnimationAsset> TDatabaseAnimationAsset* GetMutableDatabaseAnimationAsset(const UE::PoseSearch::FSearchIndexAsset& SearchIndexAsset);
 
-	// @todo: deprecate GetAnimationAssets
-	// UE_DEPRECATED(5.4, "Use GetNumAnimationAssets to iterate over GetDatabaseAnimationAsset instead")
-	const TArray<FInstancedStruct>& GetAnimationAssets() const { return AnimationAssets; }
+	UE_DEPRECATED(5.4, "Use GetNumAnimationAssets to iterate over GetDatabaseAnimationAsset instead")
+	const TArray<FInstancedStruct>& GetAnimationAssets() const;
 
 	UE_DEPRECATED(5.4, "Use GetDatabaseAnimationAsset instead")
 	const FInstancedStruct& GetAnimationAssetStruct(int32 AnimationAssetIndex) const;
@@ -586,10 +587,12 @@ private:
 template<typename TDatabaseAnimationAsset>
 inline const TDatabaseAnimationAsset* UPoseSearchDatabase::GetDatabaseAnimationAsset(int32 AnimationAssetIndex) const
 {
+	PRAGMA_DISABLE_DEPRECATION_WARNINGS
 	if (AnimationAssets.IsValidIndex(AnimationAssetIndex))
 	{
 		return AnimationAssets[AnimationAssetIndex].GetPtr<TDatabaseAnimationAsset>();
 	}
+	PRAGMA_ENABLE_DEPRECATION_WARNINGS
 	return nullptr;
 }
 
@@ -602,10 +605,12 @@ inline const TDatabaseAnimationAsset* UPoseSearchDatabase::GetDatabaseAnimationA
 template<typename TDatabaseAnimationAsset>
 inline TDatabaseAnimationAsset* UPoseSearchDatabase::GetMutableDatabaseAnimationAsset(int32 AnimationAssetIndex)
 {
+	PRAGMA_DISABLE_DEPRECATION_WARNINGS
 	if (AnimationAssets.IsValidIndex(AnimationAssetIndex))
 	{
 		return AnimationAssets[AnimationAssetIndex].GetMutablePtr<TDatabaseAnimationAsset>();
 	}
+	PRAGMA_ENABLE_DEPRECATION_WARNINGS
 	return nullptr;
 }
 
