motion matching - hashing UMirrorDataTable, USkeletalMeshSocket, USkeleton properties into pose search database DDC key

diff --git a/Engine/Plugins/Animation/PoseSearch/Source/Runtime/Private/PoseSearchDerivedData.cpp b/Engine/Plugins/Animation/PoseSearch/Source/Runtime/Private/PoseSearchDerivedData.cpp
--- a/Engine/Plugins/Animation/PoseSearch/Source/Runtime/Private/PoseSearchDerivedData.cpp
+++ b/Engine/Plugins/Animation/PoseSearch/Source/Runtime/Private/PoseSearchDerivedData.cpp
@@ -2278,7 +2278,11 @@ void FAsyncPoseSearchDatabasesManagement::Tick(float DeltaTime)
 					{
 						UE_LOG(LogPoseSearch, Error, TEXT("FAnimationAssetSampler - Loaded the wrong amount of data!"));
 					}
-					else if (TestSample.SerializedData.Num() != LoadedData.Num() || FMemory::Memcmp(TestSample.SerializedData.GetData(), LoadedData.GetData(), TestSample.SerializedData.Num()) != 0)
+					else if (TestSample.SerializedData.Num() != LoadedData.Num())
+					{
+						UE_LOG(LogPoseSearch, Error, TEXT("FAnimationAssetSampler - Loaded data mismatch expected amount of data!"));
+					}
+					else if (FMemory::Memcmp(TestSample.SerializedData.GetData(), LoadedData.GetData(), TestSample.SerializedData.Num()) != 0)
 					{
 						const int32 NumTransforms = LoadedData.Num() / sizeof(FTransform);
 
diff --git a/Engine/Plugins/Animation/PoseSearch/Source/Runtime/Private/PoseSearchDerivedDataKey.cpp b/Engine/Plugins/Animation/PoseSearch/Source/Runtime/Private/PoseSearchDerivedDataKey.cpp
--- a/Engine/Plugins/Animation/PoseSearch/Source/Runtime/Private/PoseSearchDerivedDataKey.cpp
+++ b/Engine/Plugins/Animation/PoseSearch/Source/Runtime/Private/PoseSearchDerivedDataKey.cpp
@@ -8,10 +8,7 @@
 #include "Animation/AnimData/IAnimationDataModel.h"
 #include "AnimationModifier.h"
 #include "Animation/AnimCompress.h"
-#include "Animation/MirrorDataTable.h"
-#include "Animation/Skeleton.h"
 #include "Engine/SkeletalMesh.h"
-#include "Engine/SkeletalMeshSocket.h"
 #include "Engine/SkinnedAsset.h"
 #include "Engine/StreamableRenderAsset.h"
 #include "Factories/FbxAnimSequenceImportData.h"
@@ -149,7 +146,13 @@ FArchive& FKeyBuilder::operator<<(class UObject*& Object)
 				bAnyAssetNotReady = true;
 			}
 		}
-		
+
+		// @todo: should we WaitOnExistingCompression?
+		//if (UAnimSequence* AnimSequence = Cast<UAnimSequence>(Object))
+		//{
+		//	AnimSequence->WaitOnExistingCompression();
+		//}
+
 		if (!bAnyAssetNotReady)
 		{
 			#if UE_POSE_SEARCH_DERIVED_DATA_LOGGING
@@ -257,10 +260,7 @@ bool FKeyBuilder::IsAddNameOnlyType(class UObject* Object)
 		nullptr != Cast<UAnimCurveCompressionSettings>(Object) ||
 		nullptr != Cast<UAssetImportData>(Object) ||
 		nullptr != Cast<UFunction>(Object) ||
-		nullptr != Cast<UMirrorDataTable>(Object) ||
 		nullptr != Cast<USkeletalMesh>(Object) ||
-		nullptr != Cast<USkeletalMeshSocket>(Object) ||
-		nullptr != Cast<USkeleton>(Object) ||
 		nullptr != Cast<UStreamableRenderAsset>(Object);
 }
 
