PoseSearch: add description to plugin and normalization set

diff --git a/Engine/Plugins/Animation/PoseSearch/PoseSearch.uplugin b/Engine/Plugins/Animation/PoseSearch/PoseSearch.uplugin
--- a/Engine/Plugins/Animation/PoseSearch/PoseSearch.uplugin
+++ b/Engine/Plugins/Animation/PoseSearch/PoseSearch.uplugin
@@ -3,7 +3,7 @@
 	"Version": 1,
 	"VersionName": "1.0",
 	"FriendlyName": "Pose Search",
-	"Description": "Pose Search API",
+	"Description": "Framework for indexing and searching pose features. Used in techniques such as Motion Matching.",
 	"Category": "Animation",
 	"CreatedBy": "Epic Games, Inc.",
 	"CreatedByURL": "https://epicgames.com",
diff --git a/Engine/Plugins/Animation/PoseSearch/Source/Runtime/Public/PoseSearch/PoseSearchNormalizationSet.h b/Engine/Plugins/Animation/PoseSearch/Source/Runtime/Public/PoseSearch/PoseSearchNormalizationSet.h
--- a/Engine/Plugins/Animation/PoseSearch/Source/Runtime/Public/PoseSearch/PoseSearchNormalizationSet.h
+++ b/Engine/Plugins/Animation/PoseSearch/Source/Runtime/Public/PoseSearch/PoseSearchNormalizationSet.h
@@ -8,6 +8,9 @@
 
 class UPoseSearchDatabase;
 
+/**
+* Data asset used to allow multiple Pose Search Databases to be normalized together.
+*/
 UCLASS(BlueprintType, Category = "Animation|Pose Search", meta = (DisplayName = "Pose Search Normalization Set"))
 class POSESEARCH_API UPoseSearchNormalizationSet : public UDataAsset
 {
