Pose Search: Expose curve name from channel to blueprint

diff --git a/Engine/Plugins/Animation/PoseSearch/Source/Runtime/Private/PoseSearchFeatureChannel_Curve.h b/Engine/Plugins/Animation/PoseSearch/Source/Runtime/Private/PoseSearchFeatureChannel_Curve.h
--- a/Engine/Plugins/Animation/PoseSearch/Source/Runtime/Private/PoseSearchFeatureChannel_Curve.h
+++ b/Engine/Plugins/Animation/PoseSearch/Source/Runtime/Private/PoseSearchFeatureChannel_Curve.h
@@ -12,7 +12,7 @@ class POSESEARCH_API UPoseSearchFeatureChannel_Curve : public UPoseSearchFeature
 
 public:
 
-	UPROPERTY(EditAnywhere, Category = "Settings")
+	UPROPERTY(BlueprintReadOnly, EditAnywhere, Category = "Settings")
 	FName CurveName = NAME_None;
 
 	UPROPERTY(EditAnywhere, Category = "Settings")
