PoseSearch: Added dynamic blendspace inputs support to motion matching

diff --git a/Engine/Plugins/Animation/PoseSearch/Source/Runtime/Private/AnimNode_MotionMatching.cpp b/Engine/Plugins/Animation/PoseSearch/Source/Runtime/Private/AnimNode_MotionMatching.cpp
--- a/Engine/Plugins/Animation/PoseSearch/Source/Runtime/Private/AnimNode_MotionMatching.cpp
+++ b/Engine/Plugins/Animation/PoseSearch/Source/Runtime/Private/AnimNode_MotionMatching.cpp
@@ -209,11 +209,31 @@ void FAnimNode_MotionMatching::UpdateAssetPlayer(const FAnimationUpdateContext&
 						AssetPlayerNode->SetGroupName(NAME_None);
 					}
 				}
+
+				bool bAbortBlend = false;
+				if (!AnimPlayers.IsEmpty() && GetBlendspaceParametersDeltaThreshold() > 0.0f)
+				{
+					if (AnimPlayers[0].GetAnimationAsset() == AnimationAsset)
+					{
+						const FVector CurrentBlendParameters = AnimPlayers[0].GetBlendParameters();
+						const FVector DesiredBlendParameters = GetBlendspaceParameters();
+						const float Delta = (CurrentBlendParameters - DesiredBlendParameters).SizeSquared();
+						if (Delta < FMath::Square(GetBlendspaceParametersDeltaThreshold()))
+						{
+							// If we haven't changed assets, and our currently playing blendspace xy is within threshold of change, then keep playing it.
+							// Time differences should be OK because of the call to AdjustAssetTime before the search.
+							bAbortBlend = true;
+						}
+					}
+				}
 				
-				FAnimNode_BlendStack_Standalone::BlendTo(Context, AnimationAsset, MotionMatchingState.CurrentSearchResult.AssetTime,
-					SearchIndexAsset->IsLooping(), SearchIndexAsset->IsMirrored(), CurrentResultDatabase->Schema->GetMirrorDataTable(DefaultRole), BlendTime,
-					BlendProfile, BlendOption, bUseInertialBlend, SearchIndexAsset->GetBlendParameters(), MotionMatchingState.WantedPlayRate * PlayRateMultiplier, 0,
-					GetGroupName(), GetGroupRole(), GetGroupMethod(), GetOverridePositionWhenJoiningSyncGroupAsLeader());
+				if (!bAbortBlend)
+				{
+					FAnimNode_BlendStack_Standalone::BlendTo(Context, AnimationAsset, MotionMatchingState.CurrentSearchResult.AssetTime,
+						SearchIndexAsset->IsLooping(), SearchIndexAsset->IsMirrored(), CurrentResultDatabase->Schema->GetMirrorDataTable(DefaultRole), BlendTime,
+						BlendProfile, BlendOption, bUseInertialBlend, SearchIndexAsset->GetBlendParameters(), MotionMatchingState.WantedPlayRate * PlayRateMultiplier, 0,
+						GetGroupName(), GetGroupRole(), GetGroupMethod(), GetOverridePositionWhenJoiningSyncGroupAsLeader());
+				}
 			}
 			else
 			{
@@ -226,6 +246,8 @@ void FAnimNode_MotionMatching::UpdateAssetPlayer(const FAnimationUpdateContext&
 	UE::Anim::TOptionalScopedGraphMessage<UE::Anim::FAnimInertializationSyncScope> InertializationSync(bDidBlendToRequestAnInertialBlend, Context);
 	
 	FAnimNode_BlendStack_Standalone::UpdatePlayRate(MotionMatchingState.WantedPlayRate * PlayRateMultiplier);
+	FAnimNode_BlendStack_Standalone::UpdateBlendspaceParameters(GetBlendspaceUpdateMode(), GetBlendspaceParameters());
+
 	FAnimNode_BlendStack_Standalone::UpdateAssetPlayer(Context);
 
 	NextUpdateInterruptMode = EPoseSearchInterruptMode::DoNotInterrupt;
@@ -290,6 +312,21 @@ bool FAnimNode_MotionMatching::SetIgnoreForRelevancyTest(bool bInIgnoreForReleva
 	return false;
 }
 
+const FVector& FAnimNode_MotionMatching::GetBlendspaceParameters() const
+{
+	return GET_ANIM_NODE_DATA(FVector, BlendParameters);
+}
+
+float FAnimNode_MotionMatching::GetBlendspaceParametersDeltaThreshold() const
+{
+	return GET_ANIM_NODE_DATA(float, BlendParametersDeltaThreshold);
+}
+
+EBlendStack_BlendspaceUpdateMode FAnimNode_MotionMatching::GetBlendspaceUpdateMode() const
+{
+	return GET_ANIM_NODE_DATA(EBlendStack_BlendspaceUpdateMode, BlendspaceUpdateMode);
+}
+
 FName FAnimNode_MotionMatching::GetGroupName() const
 {
 	return GET_ANIM_NODE_DATA(FName, GroupName);
diff --git a/Engine/Plugins/Animation/PoseSearch/Source/Runtime/Public/PoseSearch/AnimNode_MotionMatching.h b/Engine/Plugins/Animation/PoseSearch/Source/Runtime/Public/PoseSearch/AnimNode_MotionMatching.h
--- a/Engine/Plugins/Animation/PoseSearch/Source/Runtime/Public/PoseSearch/AnimNode_MotionMatching.h
+++ b/Engine/Plugins/Animation/PoseSearch/Source/Runtime/Public/PoseSearch/AnimNode_MotionMatching.h
@@ -52,12 +52,25 @@ private:
 	virtual bool SetOverridePositionWhenJoiningSyncGroupAsLeader(bool InOverridePositionWhenJoiningSyncGroupAsLeader) override;
 	// End of FAnimNode_Base interface
 
+	const FVector& GetBlendspaceParameters() const;
+	float GetBlendspaceParametersDeltaThreshold() const;
+	EBlendStack_BlendspaceUpdateMode GetBlendspaceUpdateMode() const;
+
 	// FAnimNode_AssetPlayerBase interface
 	virtual void UpdateAssetPlayer(const FAnimationUpdateContext& Context) override;
 	// End of FAnimNode_AssetPlayerBase interface
 
 #if WITH_EDITORONLY_DATA
 	
+	// requested blend space blend parameters (if AnimationAsset is a blend space)
+	UPROPERTY(EditAnywhere, Category = Blendspace, meta = (PinHiddenByDefault, FoldProperty))
+	FVector BlendParameters = FVector::Zero();
+
+	// Use this to define a threshold to trigger a new blend when blendspace xy input pins change.
+	// By default, any delta will trigger a blend.
+	UPROPERTY(EditAnywhere, Category = Blendspace)
+	float BlendParametersDeltaThreshold = 0.0f;
+
 	// The group name that we synchronize with (NAME_None if it is not part of any group). Note that
 	// this is the name of the group used to sync the output of this node - it will not force
 	// syncing of animations contained by it.
@@ -73,6 +86,10 @@ private:
 	UPROPERTY(VisibleAnywhere, Category = Sync, meta = (FoldProperty, EditCondition = "GroupRole == EAnimGroupRole::AlwaysLeader || GroupRole == EAnimGroupRole::ExclusiveAlwaysLeader || GroupRole == EAnimGroupRole::TransitionLeader", EditConditionHides))
 	bool bOverridePositionWhenJoiningSyncGroupAsLeader = true;
 	
+	// How we should update individual blend space parameters. See dropdown options tooltips.
+	UPROPERTY(EditAnywhere, Category = Blendspace, meta = (FoldProperty))
+	EBlendStack_BlendspaceUpdateMode BlendspaceUpdateMode = EBlendStack_BlendspaceUpdateMode::InitialOnly;
+
 	// How this node will synchronize with other animations. Note that this determines how the output
 	// of this node is used for synchronization, not of animations contained by it.
 	UPROPERTY(EditAnywhere, Category = Sync, meta = (FoldProperty))
