motion matching - performance improvements

diff --git a/Engine/Plugins/Animation/PoseSearch/Source/Editor/Private/PoseSearchCustomization.cpp b/Engine/Plugins/Animation/PoseSearch/Source/Editor/Private/PoseSearchCustomization.cpp
--- a/Engine/Plugins/Animation/PoseSearch/Source/Editor/Private/PoseSearchCustomization.cpp
+++ b/Engine/Plugins/Animation/PoseSearch/Source/Editor/Private/PoseSearchCustomization.cpp
@@ -44,9 +44,9 @@ void FPoseSearchDatabaseSequenceCustomization::CustomizeHeader(TSharedRef<IPrope
 		FPoseSearchDatabaseSequence* PoseSearchDatabaseSequence = (FPoseSearchDatabaseSequence*)InStructPropertyHandle->GetValueBaseAddress((uint8*)Objects[0]);
 		check(PoseSearchDatabaseSequence);
 		
-		if (PoseSearchDatabaseSequence->Sequence)
+		if (const UObject* AnimationAsset = PoseSearchDatabaseSequence->GetAnimationAsset())
 		{
-			SequenceNameText = FText::FromName(PoseSearchDatabaseSequence->Sequence->GetFName());
+			SequenceNameText = FText::FromName(AnimationAsset->GetFName());
 		}
 		else
 		{
diff --git a/Engine/Plugins/Animation/PoseSearch/Source/Editor/Private/PoseSearchDatabaseViewModel.cpp b/Engine/Plugins/Animation/PoseSearch/Source/Editor/Private/PoseSearchDatabaseViewModel.cpp
--- a/Engine/Plugins/Animation/PoseSearch/Source/Editor/Private/PoseSearchDatabaseViewModel.cpp
+++ b/Engine/Plugins/Animation/PoseSearch/Source/Editor/Private/PoseSearchDatabaseViewModel.cpp
@@ -139,9 +139,10 @@ bool FDatabasePreviewActor::SpawnPreviewActor(UWorld* World, const UPoseSearchDa
 
 		for (int32 Index = 1; Index < NumPoses; ++Index)
 		{
+			const float DeltaAccumulatedSeconds = FMath::Max(UE_KINDA_SMALL_NUMBER, Trajectory.Samples[Index].AccumulatedSeconds - Trajectory.Samples[Index - 1].AccumulatedSeconds);
 			const FVector& Start = Trajectory.Samples[Index - 1].Position;
 			const FVector& End = Trajectory.Samples[Index].Position;
-			TrajectorySpeed[Index] = (Start - End).Length() * PoseSearchDatabase->Schema->SampleRate;
+			TrajectorySpeed[Index] = (Start - End).Length() / DeltaAccumulatedSeconds;
 		}
 		TrajectorySpeed[0] = TrajectorySpeed[1];
 	}
diff --git a/Engine/Plugins/Animation/PoseSearch/Source/Runtime/Private/PoseSearchDatabase.cpp b/Engine/Plugins/Animation/PoseSearch/Source/Runtime/Private/PoseSearchDatabase.cpp
--- a/Engine/Plugins/Animation/PoseSearch/Source/Runtime/Private/PoseSearchDatabase.cpp
+++ b/Engine/Plugins/Animation/PoseSearch/Source/Runtime/Private/PoseSearchDatabase.cpp
@@ -40,140 +40,6 @@
 static TAutoConsoleVariable<bool> CVarMotionMatchValidateKNNSearch(TEXT("a.MotionMatch.ValidateKNNSearch"), false, TEXT("Validate KNN search"));
 #endif
 
-typedef TArray<int32, TInlineAllocator<256, TMemStackAllocator<>>> FSelectableAssetIdx;
-static void PopulateSelectableAssetIdx(FSelectableAssetIdx& SelectableAssetIdx, TConstArrayView<const UObject*> AssetsToConsider, const UPoseSearchDatabase* Database)
-{
-	check(Database);
-	SelectableAssetIdx.Reset();
-	if (!AssetsToConsider.IsEmpty())
-	{
-		const FSearchIndex& SearchIndex = Database->GetSearchIndex();
-
-		for (int32 AssetIndex = 0; AssetIndex < SearchIndex.Assets.Num(); ++AssetIndex)
-		{
-			if (const FPoseSearchDatabaseAnimationAssetBase* DatabaseAnimationAssetBase = Database->GetDatabaseAnimationAsset<FPoseSearchDatabaseAnimationAssetBase>(SearchIndex.Assets[AssetIndex]))
-			{
-				if (AssetsToConsider.Contains(DatabaseAnimationAssetBase->GetAnimationAsset()))
-				{
-					SelectableAssetIdx.Add(AssetIndex);
-				}
-			}
-		}
-	}
-}
-
-typedef TArray<int32, TInlineAllocator<256, TMemStackAllocator<>>> FNonSelectableIdx;
-static void PopulateNonSelectableIdx(FNonSelectableIdx& NonSelectableIdx, FSearchContext& SearchContext, const UPoseSearchDatabase* Database
-#if UE_POSE_SEARCH_TRACE_ENABLED
-	, TConstArrayView<float> QueryValues
-#endif //UE_POSE_SEARCH_TRACE_ENABLED
-)
-{
-	check(Database);
-	const FSearchIndex& SearchIndex = Database->GetSearchIndex();
-
-	NonSelectableIdx.Reset();
-	if (SearchContext.IsCurrentResultFromDatabase(Database))
-	{
-		if (const FSearchIndexAsset* CurrentIndexAsset = SearchContext.GetCurrentResult().GetSearchIndexAsset(true))
-		{
-			if (CurrentIndexAsset->IsDisableReselection())
-			{
-				// excluding all the poses with CurrentIndexAsset->GetSourceAssetIdx()
-				// @todo: optimize this code!
-				for (const FSearchIndexAsset& SearchIndexAsset : SearchIndex.Assets)
-				{
-					if (SearchIndexAsset.GetSourceAssetIdx() == CurrentIndexAsset->GetSourceAssetIdx())
-					{
-						const int32 FirstPoseIdx = SearchIndexAsset.GetFirstPoseIdx();
-						const int32 LastPoseIdx = FirstPoseIdx + SearchIndexAsset.GetNumPoses();
-						for (int32 PoseIdx = FirstPoseIdx; PoseIdx < LastPoseIdx; ++PoseIdx)
-						{
-							// no need to AddUnique since there's no overlapping between pose indexes in the FSearchIndexAsset(s)
-							NonSelectableIdx.Add(PoseIdx);
-
-#if UE_POSE_SEARCH_TRACE_ENABLED
-							const TArray<float> PoseValues = SearchIndex.GetPoseValuesSafe(PoseIdx);
-							const FPoseSearchCost PoseCost = SearchIndex.ComparePoses(PoseIdx, 0.f, PoseValues, QueryValues);
-							SearchContext.Track(Database, PoseIdx, EPoseCandidateFlags::DiscardedBy_PoseJumpThresholdTime, PoseCost);
-#endif // UE_POSE_SEARCH_TRACE_ENABLED
-						}
-					}
-				}
-			}
-			else if (!FMath::IsNearlyEqual(SearchContext.GetPoseJumpThresholdTime().Min, SearchContext.GetPoseJumpThresholdTime().Max))
-			{
-				const int32 CurrentResultPoseIdx = SearchContext.GetCurrentResult().PoseIdx;
-				const int32 UnboundMinPoseIdx = CurrentResultPoseIdx + FMath::FloorToInt(SearchContext.GetPoseJumpThresholdTime().Min * Database->Schema->SampleRate);
-				const int32 UnboundMaxPoseIdx = CurrentResultPoseIdx + FMath::CeilToInt(SearchContext.GetPoseJumpThresholdTime().Max * Database->Schema->SampleRate);
-				const int32 CurrentIndexAssetFirstPoseIdx = CurrentIndexAsset->GetFirstPoseIdx();
-				const int32 CurrentIndexAssetNumPoses = CurrentIndexAsset->GetNumPoses();
-				const bool bIsLooping = CurrentIndexAsset->IsLooping();
-
-				if (bIsLooping)
-				{
-					for (int32 UnboundPoseIdx = UnboundMinPoseIdx; UnboundPoseIdx < UnboundMaxPoseIdx; ++UnboundPoseIdx)
-					{
-						const int32 Modulo = (UnboundPoseIdx - CurrentIndexAssetFirstPoseIdx) % CurrentIndexAssetNumPoses;
-						const int32 CurrentIndexAssetFirstPoseIdxPlusModulo = CurrentIndexAssetFirstPoseIdx + Modulo;
-						const int32 PoseIdx = Modulo >= 0 ? CurrentIndexAssetFirstPoseIdxPlusModulo : CurrentIndexAssetFirstPoseIdxPlusModulo + CurrentIndexAssetNumPoses;
-
-						NonSelectableIdx.AddUnique(PoseIdx);
-
-#if UE_POSE_SEARCH_TRACE_ENABLED
-						const TArray<float> PoseValues = SearchIndex.GetPoseValuesSafe(PoseIdx);
-						const FPoseSearchCost PoseCost = SearchIndex.ComparePoses(PoseIdx, 0.f, PoseValues, QueryValues);
-						SearchContext.Track(Database, PoseIdx, EPoseCandidateFlags::DiscardedBy_PoseJumpThresholdTime, PoseCost);
-#endif // UE_POSE_SEARCH_TRACE_ENABLED
-					}
-				}
-				else
-				{
-					const int32 MinPoseIdx = FMath::Max(CurrentIndexAssetFirstPoseIdx, UnboundMinPoseIdx);
-					const int32 MaxPoseIdx = FMath::Min(CurrentIndexAssetFirstPoseIdx + CurrentIndexAssetNumPoses, UnboundMaxPoseIdx);
-
-					for (int32 PoseIdx = MinPoseIdx; PoseIdx < MaxPoseIdx; ++PoseIdx)
-					{
-						NonSelectableIdx.AddUnique(PoseIdx);
-
-#if UE_POSE_SEARCH_TRACE_ENABLED
-						const TArray<float> PoseValues = SearchIndex.GetPoseValuesSafe(PoseIdx);
-						const FPoseSearchCost PoseCost = SearchIndex.ComparePoses(PoseIdx, 0.f, PoseValues, QueryValues);
-						SearchContext.Track(Database, PoseIdx, EPoseCandidateFlags::DiscardedBy_PoseJumpThresholdTime, PoseCost);
-#endif // UE_POSE_SEARCH_TRACE_ENABLED
-					}
-				}
-			}
-		}
-	}
-
-	if (SearchContext.GetPoseIndicesHistory())
-	{
-		const FObjectKey DatabaseKey(Database);
-		for (auto It = SearchContext.GetPoseIndicesHistory()->IndexToTime.CreateConstIterator(); It; ++It)
-		{
-			const FHistoricalPoseIndex& HistoricalPoseIndex = It.Key();
-			if (HistoricalPoseIndex.DatabaseKey == DatabaseKey)
-			{
-				NonSelectableIdx.AddUnique(HistoricalPoseIndex.PoseIndex);
-
-#if UE_POSE_SEARCH_TRACE_ENABLED
-				check(HistoricalPoseIndex.PoseIndex >= 0);
-
-				// if we're editing the database and removing assets it's possible that the PoseIndicesHistory contains invalid pose indexes
-				if (HistoricalPoseIndex.PoseIndex < SearchIndex.GetNumPoses())
-				{
-					const FPoseSearchCost PoseCost = SearchIndex.ComparePoses(HistoricalPoseIndex.PoseIndex, 0.f, SearchIndex.GetPoseValuesSafe(HistoricalPoseIndex.PoseIndex), QueryValues);
-					SearchContext.Track(Database, HistoricalPoseIndex.PoseIndex, EPoseCandidateFlags::DiscardedBy_PoseReselectHistory, PoseCost);
-				}
-#endif // UE_POSE_SEARCH_TRACE_ENABLED
-			}
-		}
-	}
-
-	NonSelectableIdx.Sort();
-}
-
 struct FSearchFilters
 {
 	FSearchFilters(const UPoseSearchSchema* Schema, TConstArrayView<int32> NonSelectableIdx, TConstArrayView<int32> SelectableAssetIdx, bool bAddBlockTransitionFilter)
@@ -304,6 +170,43 @@
 	TArray<const IPoseSearchFilter*, TInlineAllocator<64, TMemStackAllocator<>>> Filters;
 };
 
+template<bool bReconstructPoseValues, bool bAlignedAndPadded>
+static inline void EvaluatePoseKernel(UE::PoseSearch::FSearchResult& Result, const UE::PoseSearch::FSearchIndex& SearchIndex, TConstArrayView<float> QueryValues, TArrayView<float> ReconstructedPoseValuesBuffer,
+	int32 PoseIdx, const UE::PoseSearch::FSearchFilters& SearchFilters, UE::PoseSearch::FSearchContext& SearchContext, const UPoseSearchDatabase* Database, bool bUpdateBestCandidates, int32 ResultIndex = -1)
+{
+	using namespace UE::PoseSearch;
+
+	const TConstArrayView<float> PoseValues = bReconstructPoseValues ? SearchIndex.GetReconstructedPoseValues(PoseIdx, ReconstructedPoseValuesBuffer) : SearchIndex.GetPoseValues(PoseIdx);
+
+	if (SearchFilters.AreFiltersValid(SearchIndex, PoseValues, QueryValues, PoseIdx
+#if UE_POSE_SEARCH_TRACE_ENABLED
+		, SearchContext, Database
+#endif // UE_POSE_SEARCH_TRACE_ENABLED
+	))
+	{
+		const FPoseSearchCost PoseCost = bAlignedAndPadded ? SearchIndex.CompareAlignedPoses(PoseIdx, 0.f, PoseValues, QueryValues) : SearchIndex.ComparePoses(PoseIdx, 0.f, PoseValues, QueryValues);
+		if (PoseCost < Result.PoseCost)
+		{
+			Result.PoseCost = PoseCost;
+			Result.PoseIdx = PoseIdx;
+
+#if WITH_EDITOR && ENABLE_ANIM_DEBUG && UE_POSE_SEARCH_TRACE_ENABLED
+			if (bUpdateBestCandidates)
+			{
+				Result.BestPosePos = ResultIndex;
+			}
+#endif // WITH_EDITOR && ENABLE_ANIM_DEBUG && UE_POSE_SEARCH_TRACE_ENABLED
+		}
+
+#if UE_POSE_SEARCH_TRACE_ENABLED
+		if (bUpdateBestCandidates)
+		{
+			SearchContext.Track(Database, PoseIdx, EPoseCandidateFlags::Valid_Pose, PoseCost);
+		}
+#endif // UE_POSE_SEARCH_TRACE_ENABLED
+	}
+}
+
 } // namespace UE::PoseSearch
 
 //////////////////////////////////////////////////////////////////////////
@@ -678,8 +581,57 @@
 {
 	check(IsInGameThread());
 	SearchIndexPrivate = SearchIndex;
+
+	UpdateCachedProperties();
 }
 
+void UPoseSearchDatabase::UpdateCachedProperties()
+{
+	using namespace UE::PoseSearch;
+
+	CachedAssetMap.Reset();
+	for (int32 AssetIdx = 0; AssetIdx != SearchIndexPrivate.Assets.Num(); ++AssetIdx)
+	{
+		const FSearchIndexAsset& SearchIndexAsset = SearchIndexPrivate.Assets[AssetIdx];
+
+		if (const FPoseSearchDatabaseAnimationAssetBase* DatabaseAnimationAssetBase = GetDatabaseAnimationAsset<FPoseSearchDatabaseAnimationAssetBase>(SearchIndexAsset))
+		{
+			CachedAssetMap.FindOrAdd(DatabaseAnimationAssetBase->GetAnimationAsset()).Add(AssetIdx);
+		}
+	}
+
+	for (TPair<TObjectPtr<UObject>, TArray<int32>>& CachedAssetMapPair : CachedAssetMap)
+	{
+		CachedAssetMapPair.Value.Sort();
+	}
+}
+
+TConstArrayView<int32> UPoseSearchDatabase::GetAssetIndexesForSourceAsset(const UObject* SourceAsset) const
+{
+	using namespace UE::PoseSearch;
+
+	if (const TArray<int32>* IndexesForSourceAsset = CachedAssetMap.Find(SourceAsset))
+	{
+#if DO_CHECK
+		// validating the consistency of IndexesForSourceAsset retrieved from SourceAsset
+		const UE::PoseSearch::FSearchIndex& SearchIndex = GetSearchIndex();
+		for (int32 AssetIndex : *IndexesForSourceAsset)
+		{
+			const FSearchIndexAsset& SearchIndexAsset = SearchIndex.Assets[AssetIndex];
+			const FPoseSearchDatabaseAnimationAssetBase* DatabaseAnimationAssetBase = GetDatabaseAnimationAsset<FPoseSearchDatabaseAnimationAssetBase>(SearchIndexAsset);
+	
+			// if those checks fail the calling code hasn't been protected by FAsyncPoseSearchDatabasesManagement::RequestAsyncBuildIndex
+			check(DatabaseAnimationAssetBase);
+			check(DatabaseAnimationAssetBase->GetAnimationAsset() == SourceAsset);
+		}
+#endif // DO_CHECK
+
+		return *IndexesForSourceAsset;
+	}
+	return TConstArrayView<int32>();
+}
+
+
 const UE::PoseSearch::FSearchIndex& UPoseSearchDatabase::GetSearchIndex() const
 {
 	// making sure the search index is consistent. if it fails the calling code hasn't been protected by FAsyncPoseSearchDatabasesManagement::RequestAsyncBuildIndex
@@ -769,11 +721,10 @@
 	using namespace UE::PoseSearch;
 
 	bool bRequiresSynchronization = false;
-	PRAGMA_DISABLE_DEPRECATION_WARNINGS
-	for (FInstancedStruct& AnimationAsset : AnimationAssets)
-	PRAGMA_ENABLE_DEPRECATION_WARNINGS
+
+	for (int32 AnimationAssetIndex = 0; AnimationAssetIndex < GetNumAnimationAssets(); ++AnimationAssetIndex)
 	{
-		if (FPoseSearchDatabaseAnimationAssetBase* AnimationAssetBase = AnimationAsset.GetMutablePtr<FPoseSearchDatabaseAnimationAssetBase>())
+		if (FPoseSearchDatabaseAnimationAssetBase* AnimationAssetBase = GetMutableDatabaseAnimationAsset<FPoseSearchDatabaseAnimationAssetBase>(AnimationAssetIndex))
 		{
 			if (AnimationAssetBase->bSynchronizeWithExternalDependency_DEPRECATED)
 			{
@@ -788,11 +739,9 @@
 	{
 		SynchronizeWithExternalDependencies();
 
-		PRAGMA_DISABLE_DEPRECATION_WARNINGS
-		for (FInstancedStruct& AnimationAsset : AnimationAssets)
-		PRAGMA_ENABLE_DEPRECATION_WARNINGS
+		for (int32 AnimationAssetIndex = 0; AnimationAssetIndex < GetNumAnimationAssets(); ++AnimationAssetIndex)
 		{
-			if (FPoseSearchDatabaseAnimationAssetBase* AnimationAssetBase = AnimationAsset.GetMutablePtr<FPoseSearchDatabaseAnimationAssetBase>())
+			if (FPoseSearchDatabaseAnimationAssetBase* AnimationAssetBase = GetMutableDatabaseAnimationAsset<FPoseSearchDatabaseAnimationAssetBase>(AnimationAssetIndex))
 			{
 				AnimationAssetBase->bSynchronizeWithExternalDependency_DEPRECATED = false;
 			}
@@ -816,19 +765,7 @@
 
 bool UPoseSearchDatabase::Contains(const UObject* Object) const
 {
-	PRAGMA_DISABLE_DEPRECATION_WARNINGS
-	for (const FInstancedStruct& AnimationAsset : AnimationAssets)
-	PRAGMA_ENABLE_DEPRECATION_WARNINGS
-	{
-		if (const FPoseSearchDatabaseAnimationAssetBase* AnimationAssetBase = AnimationAsset.GetPtr<FPoseSearchDatabaseAnimationAssetBase>())
-		{
-			if (AnimationAssetBase->GetAnimationAsset() == Object)
-			{
-				return true;
-			}
-		}
-	}
-	return false;
+	return !GetAssetIndexesForSourceAsset(Object).IsEmpty();
 }
 
 int32 UPoseSearchDatabase::GetNumAnimationAssets() const
@@ -904,7 +841,8 @@
 	// collecting all the database AnimationAsset(s) that don't require synchronization
 	TArray<bool> DisableReselection;
 	PRAGMA_DISABLE_DEPRECATION_WARNINGS
-	DisableReselection.Reserve(AnimationAssets.Num());
+	DisableReselection.Reserve(GetNumAnimationAssets());
+
 	for (FInstancedStruct& AnimationAsset : AnimationAssets)
 	{
 		FPoseSearchDatabaseAnimationAssetBase& AnimationAssetBase = AnimationAsset.GetMutable<FPoseSearchDatabaseAnimationAssetBase>();
@@ -971,7 +909,7 @@
 
 	// updating AnimationAssets from NewAnimationAssets preserving the original sorting
 	bool bModified = false;
-	for (int32 AnimationAssetIndex = AnimationAssets.Num() - 1; AnimationAssetIndex >= 0; --AnimationAssetIndex)
+	for (int32 AnimationAssetIndex = GetNumAnimationAssets() - 1; AnimationAssetIndex >= 0; --AnimationAssetIndex)
 	{
 		const int32 FoundIndex = NewAnimationAssets.Find(AnimationAssets[AnimationAssetIndex]);
 		if (FoundIndex >= 0)
@@ -1066,6 +1004,8 @@
 		if (Ar.IsLoading() || Ar.IsCooking())
 		{
 			Ar << SearchIndexPrivate;
+
+			UpdateCachedProperties();
 		}
 	}
 }
@@ -1160,40 +1100,141 @@
 	return Result;
 }
 
-template<bool bReconstructPoseValues, bool bAlignedAndPadded>
-static inline void EvaluatePoseKernel(UE::PoseSearch::FSearchResult& Result, const UE::PoseSearch::FSearchIndex& SearchIndex, TConstArrayView<float> QueryValues, TArrayView<float> ReconstructedPoseValuesBuffer,
-	int32 PoseIdx, const UE::PoseSearch::FSearchFilters& SearchFilters, UE::PoseSearch::FSearchContext& SearchContext, const UPoseSearchDatabase* Database, bool bUpdateBestCandidates, int32 ResultIndex = -1)
+void UPoseSearchDatabase::PopulateNonSelectableIdx(FNonSelectableIdx& NonSelectableIdx, UE::PoseSearch::FSearchContext& SearchContext
+#if UE_POSE_SEARCH_TRACE_ENABLED
+	, TConstArrayView<float> QueryValues
+#endif //UE_POSE_SEARCH_TRACE_ENABLED
+) const
 {
 	using namespace UE::PoseSearch;
 
-	const TConstArrayView<float> PoseValues = bReconstructPoseValues ? SearchIndex.GetReconstructedPoseValues(PoseIdx, ReconstructedPoseValuesBuffer) : SearchIndex.GetPoseValues(PoseIdx);
+	const FSearchIndex& SearchIndex = GetSearchIndex();
 
-	if (SearchFilters.AreFiltersValid(SearchIndex, PoseValues, QueryValues, PoseIdx
+	NonSelectableIdx.Reset();
+	if (SearchContext.IsCurrentResultFromDatabase(this))
+	{
+		if (const FSearchIndexAsset* CurrentIndexAsset = SearchContext.GetCurrentResult().GetSearchIndexAsset(true))
+		{
+			if (CurrentIndexAsset->IsDisableReselection())
+			{
+				// excluding all the poses with CurrentIndexAsset->GetSourceAssetIdx()
+				const FPoseSearchDatabaseAnimationAssetBase* DatabaseAnimationAssetBase = GetDatabaseAnimationAsset<FPoseSearchDatabaseAnimationAssetBase>(CurrentIndexAsset->GetSourceAssetIdx());
+				check(DatabaseAnimationAssetBase);
+
+				for (int32 AssetIndex : GetAssetIndexesForSourceAsset(DatabaseAnimationAssetBase->GetAnimationAsset()))
+				{
+					const FSearchIndexAsset& SearchIndexAsset = SearchIndex.Assets[AssetIndex];
+					const int32 FirstPoseIdx = SearchIndexAsset.GetFirstPoseIdx();
+					const int32 LastPoseIdx = FirstPoseIdx + SearchIndexAsset.GetNumPoses();
+					for (int32 PoseIdx = FirstPoseIdx; PoseIdx < LastPoseIdx; ++PoseIdx)
+					{
+						// no need to AddUnique since there's no overlapping between pose indexes in the FSearchIndexAsset(s)
+						NonSelectableIdx.Add(PoseIdx);
+
 #if UE_POSE_SEARCH_TRACE_ENABLED
-		, SearchContext, Database
+						const TArray<float> PoseValues = SearchIndex.GetPoseValuesSafe(PoseIdx);
+						const FPoseSearchCost PoseCost = SearchIndex.ComparePoses(PoseIdx, 0.f, PoseValues, QueryValues);
+						SearchContext.Track(this, PoseIdx, EPoseCandidateFlags::DiscardedBy_PoseJumpThresholdTime, PoseCost);
 #endif // UE_POSE_SEARCH_TRACE_ENABLED
-	))
+					}
+				}
+			}
+			else if (!FMath::IsNearlyEqual(SearchContext.GetPoseJumpThresholdTime().Min, SearchContext.GetPoseJumpThresholdTime().Max))
+			{
+				const int32 CurrentResultPoseIdx = SearchContext.GetCurrentResult().PoseIdx;
+				const int32 UnboundMinPoseIdx = CurrentResultPoseIdx + FMath::FloorToInt(SearchContext.GetPoseJumpThresholdTime().Min * Schema->SampleRate);
+				const int32 UnboundMaxPoseIdx = CurrentResultPoseIdx + FMath::CeilToInt(SearchContext.GetPoseJumpThresholdTime().Max * Schema->SampleRate);
+				const int32 CurrentIndexAssetFirstPoseIdx = CurrentIndexAsset->GetFirstPoseIdx();
+				const int32 CurrentIndexAssetNumPoses = CurrentIndexAsset->GetNumPoses();
+				const bool bIsLooping = CurrentIndexAsset->IsLooping();
+
+				if (bIsLooping)
+				{
+					for (int32 UnboundPoseIdx = UnboundMinPoseIdx; UnboundPoseIdx < UnboundMaxPoseIdx; ++UnboundPoseIdx)
+					{
+						const int32 Modulo = (UnboundPoseIdx - CurrentIndexAssetFirstPoseIdx) % CurrentIndexAssetNumPoses;
+						const int32 CurrentIndexAssetFirstPoseIdxPlusModulo = CurrentIndexAssetFirstPoseIdx + Modulo;
+						const int32 PoseIdx = Modulo >= 0 ? CurrentIndexAssetFirstPoseIdxPlusModulo : CurrentIndexAssetFirstPoseIdxPlusModulo + CurrentIndexAssetNumPoses;
+
+						NonSelectableIdx.AddUnique(PoseIdx);
+
+#if UE_POSE_SEARCH_TRACE_ENABLED
+						const TArray<float> PoseValues = SearchIndex.GetPoseValuesSafe(PoseIdx);
+						const FPoseSearchCost PoseCost = SearchIndex.ComparePoses(PoseIdx, 0.f, PoseValues, QueryValues);
+						SearchContext.Track(this, PoseIdx, EPoseCandidateFlags::DiscardedBy_PoseJumpThresholdTime, PoseCost);
+#endif // UE_POSE_SEARCH_TRACE_ENABLED
+					}
+				}
+				else
+				{
+					const int32 MinPoseIdx = FMath::Max(CurrentIndexAssetFirstPoseIdx, UnboundMinPoseIdx);
+					const int32 MaxPoseIdx = FMath::Min(CurrentIndexAssetFirstPoseIdx + CurrentIndexAssetNumPoses, UnboundMaxPoseIdx);
+
+					for (int32 PoseIdx = MinPoseIdx; PoseIdx < MaxPoseIdx; ++PoseIdx)
+					{
+						NonSelectableIdx.AddUnique(PoseIdx);
+
+#if UE_POSE_SEARCH_TRACE_ENABLED
+						const TArray<float> PoseValues = SearchIndex.GetPoseValuesSafe(PoseIdx);
+						const FPoseSearchCost PoseCost = SearchIndex.ComparePoses(PoseIdx, 0.f, PoseValues, QueryValues);
+						SearchContext.Track(this, PoseIdx, EPoseCandidateFlags::DiscardedBy_PoseJumpThresholdTime, PoseCost);
+#endif // UE_POSE_SEARCH_TRACE_ENABLED
+					}
+				}
+			}
+		}
+	}
+
+	if (SearchContext.GetPoseIndicesHistory())
 	{
-		const FPoseSearchCost PoseCost = bAlignedAndPadded ? SearchIndex.CompareAlignedPoses(PoseIdx, 0.f, PoseValues, QueryValues) : SearchIndex.ComparePoses(PoseIdx, 0.f, PoseValues, QueryValues);
-		if (PoseCost < Result.PoseCost)
+		const FObjectKey DatabaseKey(this);
+		for (auto It = SearchContext.GetPoseIndicesHistory()->IndexToTime.CreateConstIterator(); It; ++It)
 		{
-			Result.PoseCost = PoseCost;
-			Result.PoseIdx = PoseIdx;
+			const FHistoricalPoseIndex& HistoricalPoseIndex = It.Key();
+			if (HistoricalPoseIndex.DatabaseKey == DatabaseKey)
+			{
+				NonSelectableIdx.AddUnique(HistoricalPoseIndex.PoseIndex);
 
-#if WITH_EDITOR && ENABLE_ANIM_DEBUG && UE_POSE_SEARCH_TRACE_ENABLED
-			if (bUpdateBestCandidates)
-			{
-				Result.BestPosePos = ResultIndex;
+#if UE_POSE_SEARCH_TRACE_ENABLED
+				check(HistoricalPoseIndex.PoseIndex >= 0);
+
+				// if we're editing the database and removing assets it's possible that the PoseIndicesHistory contains invalid pose indexes
+				if (HistoricalPoseIndex.PoseIndex < SearchIndex.GetNumPoses())
+				{
+					const FPoseSearchCost PoseCost = SearchIndex.ComparePoses(HistoricalPoseIndex.PoseIndex, 0.f, SearchIndex.GetPoseValuesSafe(HistoricalPoseIndex.PoseIndex), QueryValues);
+					SearchContext.Track(this, HistoricalPoseIndex.PoseIndex, EPoseCandidateFlags::DiscardedBy_PoseReselectHistory, PoseCost);
+				}
+#endif // UE_POSE_SEARCH_TRACE_ENABLED
 			}
-#endif // WITH_EDITOR && ENABLE_ANIM_DEBUG && UE_POSE_SEARCH_TRACE_ENABLED
+		}
+	}
+
+	NonSelectableIdx.Sort();
+}
+
+void UPoseSearchDatabase::PopulateSelectableAssetIdx(FSelectableAssetIdx& SelectableAssetIdx, TConstArrayView<const UObject*> AssetsToConsider) const
+{
+	SelectableAssetIdx.Reset();
+	if (!AssetsToConsider.IsEmpty())
+	{
+		for (const UObject* AssetToConsider : AssetsToConsider)
+		{
+			SelectableAssetIdx.Append(GetAssetIndexesForSourceAsset(AssetToConsider));
 		}
 
-#if UE_POSE_SEARCH_TRACE_ENABLED
-		if (bUpdateBestCandidates)
+		if (!SelectableAssetIdx.IsEmpty())
 		{
-			SearchContext.Track(Database, PoseIdx, EPoseCandidateFlags::Valid_Pose, PoseCost);
+			if (SelectableAssetIdx.Num() != GetSearchIndex().Assets.Num())
+			{
+				SelectableAssetIdx.Sort();
+			}
+			else
+			{
+				// SelectableAssetIdx contains ALL the Database->GetSearchIndex().Assets. 
+				// We reset SelectableAssetIdx since it has the same meaning, and it'll perform better
+				SelectableAssetIdx.Reset();
+			}
 		}
-#endif // UE_POSE_SEARCH_TRACE_ENABLED
 	}
 }
 
@@ -1319,10 +1360,10 @@
 		TConstArrayView<float> QueryValues = SearchContext.GetOrBuildQuery(Schema);
 
 		FSelectableAssetIdx SelectableAssetIdx;
-		PopulateSelectableAssetIdx(SelectableAssetIdx, SearchContext.GetAssetsToConsider(), this);
+		PopulateSelectableAssetIdx(SelectableAssetIdx, SearchContext.GetAssetsToConsider());
 
 		FNonSelectableIdx NonSelectableIdx;
-		PopulateNonSelectableIdx(NonSelectableIdx, SearchContext, this
+		PopulateNonSelectableIdx(NonSelectableIdx, SearchContext
 #if UE_POSE_SEARCH_TRACE_ENABLED
 			, QueryValues
 #endif // UE_POSE_SEARCH_TRACE_ENABLED
@@ -1463,7 +1504,7 @@
 		// calling just for reporting non selectable poses
 		TConstArrayView<float> QueryValues = SearchContext.GetOrBuildQuery(Schema);
 		FNonSelectableIdx NonSelectableIdx;
-		PopulateNonSelectableIdx(NonSelectableIdx, SearchContext, this, QueryValues);
+		PopulateNonSelectableIdx(NonSelectableIdx, SearchContext, QueryValues);
 #endif // UE_POSE_SEARCH_TRACE_ENABLED
 	}
 
@@ -1494,11 +1535,11 @@
 		TConstArrayView<float> QueryValues = SearchContext.GetOrBuildQuery(Schema);
 
 		FSelectableAssetIdx SelectableAssetIdx;
-		PopulateSelectableAssetIdx(SelectableAssetIdx, SearchContext.GetAssetsToConsider(), this);
+		PopulateSelectableAssetIdx(SelectableAssetIdx, SearchContext.GetAssetsToConsider());
 
 		// @todo: implement filtering within the VPTree as KDTree does
 		FNonSelectableIdx NonSelectableIdx;
-		PopulateNonSelectableIdx(NonSelectableIdx, SearchContext, this
+		PopulateNonSelectableIdx(NonSelectableIdx, SearchContext
 #if UE_POSE_SEARCH_TRACE_ENABLED
 			, QueryValues
 #endif // UE_POSE_SEARCH_TRACE_ENABLED
@@ -1550,7 +1591,7 @@
 		// calling just for reporting non selectable poses
 		TConstArrayView<float> QueryValues = SearchContext.GetOrBuildQuery(Schema);
 		FNonSelectableIdx NonSelectableIdx;
-		PopulateNonSelectableIdx(NonSelectableIdx, SearchContext, this, QueryValues);
+		PopulateNonSelectableIdx(NonSelectableIdx, SearchContext, QueryValues);
 #endif // UE_POSE_SEARCH_TRACE_ENABLED
 	}
 
@@ -1581,10 +1622,10 @@
 		TConstArrayView<float> QueryValues = SearchContext.GetOrBuildQuery(Schema);
 
 		FSelectableAssetIdx SelectableAssetIdx;
-		PopulateSelectableAssetIdx(SelectableAssetIdx, SearchContext.GetAssetsToConsider(), this);
+		PopulateSelectableAssetIdx(SelectableAssetIdx, SearchContext.GetAssetsToConsider());
 
 		FNonSelectableIdx NonSelectableIdx;
-		PopulateNonSelectableIdx(NonSelectableIdx, SearchContext, this
+		PopulateNonSelectableIdx(NonSelectableIdx, SearchContext
 #if UE_POSE_SEARCH_TRACE_ENABLED
 			, QueryValues
 #endif // UE_POSE_SEARCH_TRACE_ENABLED
@@ -1683,7 +1724,7 @@
 		// calling just for reporting non selectable poses
 		TConstArrayView<float> QueryValues = SearchContext.GetOrBuildQuery(Schema);
 		FNonSelectableIdx NonSelectableIdx;
-		PopulateNonSelectableIdx(NonSelectableIdx, SearchContext, this, QueryValues);
+		PopulateNonSelectableIdx(NonSelectableIdx, SearchContext, QueryValues);
 #endif // UE_POSE_SEARCH_TRACE_ENABLED
 	}
 
diff --git a/Engine/Plugins/Animation/PoseSearch/Source/Runtime/Private/PoseSearchDerivedDataKey.cpp b/Engine/Plugins/Animation/PoseSearch/Source/Runtime/Private/PoseSearchDerivedDataKey.cpp
--- a/Engine/Plugins/Animation/PoseSearch/Source/Runtime/Private/PoseSearchDerivedDataKey.cpp
+++ b/Engine/Plugins/Animation/PoseSearch/Source/Runtime/Private/PoseSearchDerivedDataKey.cpp
@@ -39,7 +39,7 @@ FKeyBuilder::FKeyBuilder(const UObject* Object, bool bUseDataVer, bool bPerformC
 	if (bUseDataVer)
 	{
 		// used to invalidate the key without having to change POSESEARCHDB_DERIVEDDATA_VER all the times
-		int32 POSESEARCHDB_DERIVEDDATA_VER_SMALL = 259;
+		int32 POSESEARCHDB_DERIVEDDATA_VER_SMALL = 261;
 		FGuid VersionGuid = FDevSystemGuids::GetSystemGuid(FDevSystemGuids::Get().POSESEARCHDB_DERIVEDDATA_VER);
 		FString AnimationCompressionVersionString = UE::Anim::Compression::AnimationCompressionVersionString;
 
diff --git a/Engine/Plugins/Animation/PoseSearch/Source/Runtime/Private/PoseSearchLibrary.cpp b/Engine/Plugins/Animation/PoseSearch/Source/Runtime/Private/PoseSearchLibrary.cpp
--- a/Engine/Plugins/Animation/PoseSearch/Source/Runtime/Private/PoseSearchLibrary.cpp
+++ b/Engine/Plugins/Animation/PoseSearch/Source/Runtime/Private/PoseSearchLibrary.cpp
@@ -75,6 +75,12 @@ namespace UE::PoseSearch
 						continue;
 					}
 #if WITH_EDITOR
+					if (EAsyncBuildIndexResult::Success != FAsyncPoseSearchDatabasesManagement::RequestAsyncBuildIndex(PoseSearchBranchIn->Database, ERequestAsyncBuildFlag::ContinueRequest))
+					{
+						// database is still indexing.. moving on
+						continue;
+					}
+
 					if (!PoseSearchBranchIn->Database->Contains(SequenceBase))
 					{
 						UE_LOG(LogPoseSearch, Error, TEXT("improperly setup UAnimSequenceBase. Database %s doesn't contain UAnimSequenceBase %s"), *PoseSearchBranchIn->Database->GetName(), *SequenceBase->GetName());
@@ -902,47 +908,42 @@ UE::PoseSearch::FSearchResult UPoseSearchLibrary::MotionMatch(
 				check(Database);
 
 				const FSearchIndex& SearchIndex = Database->GetSearchIndex();
-				for (int32 AssetIndex = 0; AssetIndex < SearchIndex.Assets.Num(); ++AssetIndex)
+				for (int32 AssetIndex : Database->GetAssetIndexesForSourceAsset(PlayingAnimationAsset))
 				{
 					const FSearchIndexAsset& SearchIndexAsset = SearchIndex.Assets[AssetIndex];
-					if (const FPoseSearchDatabaseAnimationAssetBase* DatabaseAnimationAssetBase = Database->GetDatabaseAnimationAsset<FPoseSearchDatabaseAnimationAssetBase>(SearchIndexAsset.GetSourceAssetIdx()))
+
+					bool bCanAdvance = true;
+					if (!SearchIndexAsset.IsLooping())
 					{
-						if (PlayingAnimationAsset == DatabaseAnimationAssetBase->GetAnimationAsset())
-						{
-							bool bCanAdvance = true;
-							if (!SearchIndexAsset.IsLooping())
-							{
-								const float FirstSampleTime = SearchIndexAsset.GetFirstSampleTime(Database->Schema->SampleRate);
-								const float LastSampleTime = SearchIndexAsset.GetLastSampleTime(Database->Schema->SampleRate) - DeltaSeconds;
-								const float MaxTimeToBeAbleToContinuingPlayingAnimation = LastSampleTime - DeltaSeconds;
+						const float FirstSampleTime = SearchIndexAsset.GetFirstSampleTime(Database->Schema->SampleRate);
+						const float LastSampleTime = SearchIndexAsset.GetLastSampleTime(Database->Schema->SampleRate) - DeltaSeconds;
+						const float MaxTimeToBeAbleToContinuingPlayingAnimation = LastSampleTime - DeltaSeconds;
 
-								bCanAdvance = ContinuingProperties.PlayingAssetAccumulatedTime >= FirstSampleTime && ContinuingProperties.PlayingAssetAccumulatedTime < MaxTimeToBeAbleToContinuingPlayingAnimation;
-							}
+						bCanAdvance = ContinuingProperties.PlayingAssetAccumulatedTime >= FirstSampleTime && ContinuingProperties.PlayingAssetAccumulatedTime < MaxTimeToBeAbleToContinuingPlayingAnimation;
+					}
 
-							if (bCanAdvance)
-							{
-								ReconstructedPreviousSearchResult.Database = Database;
-								ReconstructedPreviousSearchResult.AssetTime = ContinuingProperties.PlayingAssetAccumulatedTime;
-								ReconstructedPreviousSearchResult.PoseIdx = Database->GetPoseIndexFromTime(ContinuingProperties.PlayingAssetAccumulatedTime, SearchIndexAsset);
-								SearchContext.UpdateCurrentResultPoseVector();
+					if (bCanAdvance)
+					{
+						ReconstructedPreviousSearchResult.Database = Database;
+						ReconstructedPreviousSearchResult.AssetTime = ContinuingProperties.PlayingAssetAccumulatedTime;
+						ReconstructedPreviousSearchResult.PoseIdx = Database->GetPoseIndexFromTime(ContinuingProperties.PlayingAssetAccumulatedTime, SearchIndexAsset);
+						SearchContext.UpdateCurrentResultPoseVector();
 
-								const FSearchResult NewSearchResult = Database->SearchContinuingPose(SearchContext);
+						const FSearchResult NewSearchResult = Database->SearchContinuingPose(SearchContext);
 
 #if WITH_EDITOR && ENABLE_ANIM_DEBUG && UE_POSE_SEARCH_TRACE_ENABLED
-								const FPoseSearchCost BestBruteForcePoseCost = NewSearchResult.BruteForcePoseCost < SearchResult.BruteForcePoseCost ? NewSearchResult.BruteForcePoseCost : SearchResult.BruteForcePoseCost;
+						const FPoseSearchCost BestBruteForcePoseCost = NewSearchResult.BruteForcePoseCost < SearchResult.BruteForcePoseCost ? NewSearchResult.BruteForcePoseCost : SearchResult.BruteForcePoseCost;
 #endif // WITH_EDITOR && ENABLE_ANIM_DEBUG && UE_POSE_SEARCH_TRACE_ENABLED
 
-								if (NewSearchResult.PoseCost < SearchResult.PoseCost)
-								{
-									SearchResult = NewSearchResult;
-									SearchContext.UpdateCurrentBestCost(SearchResult.PoseCost);
-								}
+						if (NewSearchResult.PoseCost < SearchResult.PoseCost)
+						{
+							SearchResult = NewSearchResult;
+							SearchContext.UpdateCurrentBestCost(SearchResult.PoseCost);
+						}
 								
 #if WITH_EDITOR && ENABLE_ANIM_DEBUG && UE_POSE_SEARCH_TRACE_ENABLED
-								SearchResult.BruteForcePoseCost = BestBruteForcePoseCost;
+						SearchResult.BruteForcePoseCost = BestBruteForcePoseCost;
 #endif // WITH_EDITOR && ENABLE_ANIM_DEBUG && UE_POSE_SEARCH_TRACE_ENABLED
-							}
-						}
 					}
 				}
 			}
diff --git a/Engine/Plugins/Animation/PoseSearch/Source/Runtime/Private/PoseSearchResult.cpp b/Engine/Plugins/Animation/PoseSearch/Source/Runtime/Private/PoseSearchResult.cpp
--- a/Engine/Plugins/Animation/PoseSearch/Source/Runtime/Private/PoseSearchResult.cpp
+++ b/Engine/Plugins/Animation/PoseSearch/Source/Runtime/Private/PoseSearchResult.cpp
@@ -11,38 +11,31 @@ namespace UE::PoseSearch
 
 void FSearchResult::Update(float NewAssetTime)
 {
-	if (!IsValid())
-	{
-		Reset();
-	}
-	else
+	if (IsValid())
 	{
+		float RealTime = NewAssetTime;
 		const FSearchIndexAsset& SearchIndexAsset = Database->GetSearchIndex().GetAssetForPose(PoseIdx);
-		if (Database->GetDatabaseAnimationAsset<FPoseSearchDatabaseSequence>(SearchIndexAsset) || Database->GetDatabaseAnimationAsset<FPoseSearchDatabaseAnimComposite>(SearchIndexAsset))
-		{
-			PoseIdx = Database->GetPoseIndexFromTime(NewAssetTime, SearchIndexAsset);
-			AssetTime = NewAssetTime;
-		}
-		else if (const FPoseSearchDatabaseBlendSpace* DatabaseBlendSpace = Database->GetDatabaseAnimationAsset<FPoseSearchDatabaseBlendSpace>(SearchIndexAsset))
+		const FPoseSearchDatabaseAnimationAssetBase* DatabaseAnimationAssetBase = Database->GetDatabaseAnimationAsset<FPoseSearchDatabaseAnimationAssetBase>(SearchIndexAsset);
+		check(DatabaseAnimationAssetBase);
+		if (const UBlendSpace* BlendSpace = Cast<UBlendSpace>(DatabaseAnimationAssetBase->GetAnimationAsset()))
 		{
 			TArray<FBlendSampleData> BlendSamples;
 			int32 TriangulationIndex = 0;
-			DatabaseBlendSpace->BlendSpace->GetSamplesFromBlendInput(SearchIndexAsset.GetBlendParameters(), BlendSamples, TriangulationIndex, true);
-
-			const float PlayLength = DatabaseBlendSpace->BlendSpace->GetAnimationLengthFromSampleData(BlendSamples);
+			BlendSpace->GetSamplesFromBlendInput(SearchIndexAsset.GetBlendParameters(), BlendSamples, TriangulationIndex, true);
+			const float PlayLength = BlendSpace->GetAnimationLengthFromSampleData(BlendSamples);
 
 			// Asset player time for blendspaces is normalized [0, 1] so we need to convert 
 			// to a real time before we advance it
 			check(NewAssetTime >= 0.f && NewAssetTime <= 1.f);
-			const float RealTime = NewAssetTime * PlayLength;
-
-			PoseIdx = Database->GetPoseIndexFromTime(RealTime, SearchIndexAsset);
-			AssetTime = NewAssetTime;
-		}
-		else
-		{
-			Reset();
+			RealTime = NewAssetTime * PlayLength;
 		}
+		
+		PoseIdx = Database->GetPoseIndexFromTime(RealTime, SearchIndexAsset);
+		AssetTime = NewAssetTime;
+	}
+	else
+	{
+		Reset();
 	}
 }
 
@@ -72,23 +65,23 @@ bool FSearchResult::CanAdvance(float DeltaTime) const
 	bool bCanAdvance = false;
 	if (const FSearchIndexAsset* SearchIndexAsset = GetSearchIndexAsset())
 	{
+		const FPoseSearchDatabaseAnimationAssetBase* DatabaseAnimationAssetBase = Database->GetDatabaseAnimationAsset<FPoseSearchDatabaseAnimationAssetBase>(*SearchIndexAsset);
+		check(DatabaseAnimationAssetBase);
+
 		float SteppedTime = AssetTime;
-		if (const FPoseSearchDatabaseBlendSpace* DatabaseBlendSpace = Database->GetDatabaseAnimationAsset<FPoseSearchDatabaseBlendSpace>(*SearchIndexAsset))
+		if (const UBlendSpace* BlendSpace = Cast<UBlendSpace>(DatabaseAnimationAssetBase->GetAnimationAsset()))
 		{
-			if (const UBlendSpace* BlendSpace = DatabaseBlendSpace->BlendSpace.Get())
-			{
-				TArray<FBlendSampleData> BlendSamples;
-				int32 TriangulationIndex = 0;
-				BlendSpace->GetSamplesFromBlendInput(SearchIndexAsset->GetBlendParameters(), BlendSamples, TriangulationIndex, true);
+			TArray<FBlendSampleData> BlendSamples;
+			int32 TriangulationIndex = 0;
+			BlendSpace->GetSamplesFromBlendInput(SearchIndexAsset->GetBlendParameters(), BlendSamples, TriangulationIndex, true);
 
-				const float PlayLength = BlendSpace->GetAnimationLengthFromSampleData(BlendSamples);
+			const float PlayLength = BlendSpace->GetAnimationLengthFromSampleData(BlendSamples);
 
-				// Asset player time for blend spaces is normalized [0, 1] so we need to convert it back to real time before we advance it
-				SteppedTime = AssetTime * PlayLength;
-				bCanAdvance = ETAA_Finished != FAnimationRuntime::AdvanceTime(SearchIndexAsset->IsLooping(), DeltaTime, SteppedTime, PlayLength);
-			}
+			// Asset player time for blend spaces is normalized [0, 1] so we need to convert it back to real time before we advance it
+			SteppedTime = AssetTime * PlayLength;
+			bCanAdvance = ETAA_Finished != FAnimationRuntime::AdvanceTime(SearchIndexAsset->IsLooping(), DeltaTime, SteppedTime, PlayLength);
 		}
-		else if (const FPoseSearchDatabaseAnimationAssetBase* DatabaseAnimationAssetBase = Database->GetDatabaseAnimationAsset<FPoseSearchDatabaseAnimationAssetBase>(*SearchIndexAsset))
+		else
 		{
 			const float AssetLength = DatabaseAnimationAssetBase->GetPlayLength();
 			bCanAdvance = ETAA_Finished != FAnimationRuntime::AdvanceTime(SearchIndexAsset->IsLooping(), DeltaTime, SteppedTime, AssetLength);
diff --git a/Engine/Plugins/Animation/PoseSearch/Source/Runtime/Public/PoseSearch/PoseSearchDatabase.h b/Engine/Plugins/Animation/PoseSearch/Source/Runtime/Public/PoseSearch/PoseSearchDatabase.h
--- a/Engine/Plugins/Animation/PoseSearch/Source/Runtime/Public/PoseSearch/PoseSearchDatabase.h
+++ b/Engine/Plugins/Animation/PoseSearch/Source/Runtime/Public/PoseSearch/PoseSearchDatabase.h
@@ -461,6 +461,7 @@ public:
 	
 	// @todo: rename to KNNQueryNumNeighbors to be usable with the VPTree as well
 	// Out of a kdtree search, results will have only an approximate cost, so the database search will select the best “KDTree Query Num Neighbors” poses to perform the full cost analysis, and be able to elect the best pose.
+	// Memory & Performance Optimization! If KDTreeQueryNumNeighbors is 1 all the SearchIndexPrivate::Values will be stripped away, and the search will exclusively rely on the KDTree query result from the PCA space encoded values (SearchIndexPrivate::PCAValues).
 	UPROPERTY(EditAnywhere, Category = "Performance", meta = (DisplayName = "KNNQueryNumNeighbors", EditCondition = "PoseSearchMode == EPoseSearchMode::PCAKDTree || PoseSearchMode == EPoseSearchMode::VPTree", EditConditionHides, ClampMin = "1", ClampMax = "600", UIMin = "1"))
 	int32 KDTreeQueryNumNeighbors = 200;
 
@@ -485,6 +486,9 @@ public:
 private:
 	// Do not use it directly. Use GetSearchIndex / SetSearchIndex interact with it and validate that is ok to do so.
 	UE::PoseSearch::FSearchIndex SearchIndexPrivate;
+	
+	// CachedAssetMap is NOT serialized in operator<< but recalculated by UpdateCachedProperties every time SearchIndexPrivate changes
+	TMap<TObjectPtr<UObject>, TArray<int32>> CachedAssetMap;
 
 #if WITH_EDITOR
 	DECLARE_MULTICAST_DELEGATE(FOnDerivedDataRebuildMulticaster);
@@ -578,10 +582,24 @@ public:
 	void TestSynchronizeWithExternalDependencies();
 #endif // WITH_EDITOR && ENABLE_ANIM_DEBUG
 
+	TConstArrayView<int32> GetAssetIndexesForSourceAsset(const UObject* SourceAsset) const;
+
 private:
 	UE::PoseSearch::FSearchResult SearchPCAKDTree(UE::PoseSearch::FSearchContext& SearchContext) const;
 	UE::PoseSearch::FSearchResult SearchVPTree(UE::PoseSearch::FSearchContext& SearchContext) const;
 	UE::PoseSearch::FSearchResult SearchBruteForce(UE::PoseSearch::FSearchContext& SearchContext) const;
+
+	typedef TArray<int32, TInlineAllocator<256, TMemStackAllocator<>>> FSelectableAssetIdx;
+	void PopulateSelectableAssetIdx(FSelectableAssetIdx& SelectableAssetIdx, TConstArrayView<const UObject*> AssetsToConsider) const;
+
+	typedef TArray<int32, TInlineAllocator<256, TMemStackAllocator<>>> FNonSelectableIdx;
+	void PopulateNonSelectableIdx(FNonSelectableIdx& NonSelectableIdx, UE::PoseSearch::FSearchContext& SearchContext
+#if UE_POSE_SEARCH_TRACE_ENABLED
+		, TConstArrayView<float> QueryValues
+#endif //UE_POSE_SEARCH_TRACE_ENABLED
+	) const;
+
+	void UpdateCachedProperties();
 };
 
 template<typename TDatabaseAnimationAsset>
@@ -618,4 +636,4 @@ template<typename TDatabaseAnimationAsset>
 inline TDatabaseAnimationAsset* UPoseSearchDatabase::GetMutableDatabaseAnimationAsset(const UE::PoseSearch::FSearchIndexAsset& SearchIndexAsset)
 {
 	return GetMutableDatabaseAnimationAsset<TDatabaseAnimationAsset>(SearchIndexAsset.GetSourceAssetIdx());
-}
+}
\ No newline at end of file
