Add anim node function to override motion matching blend parameters from the results of the search

diff --git a/Engine/Plugins/Animation/PoseSearch/Source/Runtime/Private/MotionMatchingAnimNodeLibrary.cpp b/Engine/Plugins/Animation/PoseSearch/Source/Runtime/Private/MotionMatchingAnimNodeLibrary.cpp
--- a/Engine/Plugins/Animation/PoseSearch/Source/Runtime/Private/MotionMatchingAnimNodeLibrary.cpp
+++ b/Engine/Plugins/Animation/PoseSearch/Source/Runtime/Private/MotionMatchingAnimNodeLibrary.cpp
@@ -3,6 +3,15 @@
 #include "PoseSearch/MotionMatchingAnimNodeLibrary.h"
 #include "PoseSearch/AnimNode_MotionMatching.h"
 #include "PoseSearch/PoseSearchDatabase.h"
+#include "AlphaBlend.h"
+
+FMotionMatchingBlueprintBlendSettings::FMotionMatchingBlueprintBlendSettings()
+	: BlendTime(0.2f)
+	, BlendProfile(nullptr)
+	, BlendOption(EAlphaBlendOption::Linear)
+	, bUseInertialBlend(false)
+{
+}
 
 FMotionMatchingAnimNodeReference UMotionMatchingAnimNodeLibrary::ConvertToMotionMatchingNode(const FAnimNodeReference& Node, EAnimNodeReferenceConversionResult& Result)
 {
@@ -42,6 +51,37 @@ void UMotionMatchingAnimNodeLibrary::GetMotionMatchingSearchResult(const FMotion
 	}
 }
 
+void UMotionMatchingAnimNodeLibrary::GetMotionMatchingBlendSettings(const FMotionMatchingAnimNodeReference& MotionMatchingNode, FMotionMatchingBlueprintBlendSettings& BlendSettings, bool& bIsResultValid)
+{
+	if (FAnimNode_MotionMatching* MotionMatchingNodePtr = MotionMatchingNode.GetAnimNodePtr<FAnimNode_MotionMatching>())
+	{
+		BlendSettings.BlendOption = MotionMatchingNodePtr->BlendOption;
+		BlendSettings.BlendProfile = MotionMatchingNodePtr->BlendProfile;
+		BlendSettings.BlendTime = MotionMatchingNodePtr->BlendTime;
+		BlendSettings.bUseInertialBlend = MotionMatchingNodePtr->bUseInertialBlend;
+	}
+	else
+	{
+		BlendSettings = FMotionMatchingBlueprintBlendSettings();
+		UE_LOG(LogPoseSearch, Warning, TEXT("UMotionMatchingAnimNodeLibrary::GetMotionMatchingBlendSettings called on an invalid context or with an invalid type"));
+	}
+}
+
+void UMotionMatchingAnimNodeLibrary::OverrideMotionMatchingBlendSettings(const FMotionMatchingAnimNodeReference& MotionMatchingNode, const FMotionMatchingBlueprintBlendSettings& BlendSettings, bool& bIsResultValid)
+{
+	if (FAnimNode_MotionMatching* MotionMatchingNodePtr = MotionMatchingNode.GetAnimNodePtr<FAnimNode_MotionMatching>())
+	{
+		 MotionMatchingNodePtr->BlendOption = BlendSettings.BlendOption;
+		 MotionMatchingNodePtr->BlendProfile = BlendSettings.BlendProfile;
+		 MotionMatchingNodePtr->BlendTime = BlendSettings.BlendTime;
+		 MotionMatchingNodePtr->bUseInertialBlend = BlendSettings.bUseInertialBlend;
+	}
+	else
+	{
+		UE_LOG(LogPoseSearch, Warning, TEXT("UMotionMatchingAnimNodeLibrary::OverrideMotionMatchingBlendSettings called on an invalid context or with an invalid type"));
+	}
+}
+
 void UMotionMatchingAnimNodeLibrary::SetDatabaseToSearch(const FMotionMatchingAnimNodeReference& MotionMatchingNode, UPoseSearchDatabase* Database, EPoseSearchInterruptMode InterruptMode)
 {
 	if (FAnimNode_MotionMatching* MotionMatchingNodePtr = MotionMatchingNode.GetAnimNodePtr<FAnimNode_MotionMatching>())
diff --git a/Engine/Plugins/Animation/PoseSearch/Source/Runtime/Public/PoseSearch/AnimNode_MotionMatching.h b/Engine/Plugins/Animation/PoseSearch/Source/Runtime/Public/PoseSearch/AnimNode_MotionMatching.h
--- a/Engine/Plugins/Animation/PoseSearch/Source/Runtime/Public/PoseSearch/AnimNode_MotionMatching.h
+++ b/Engine/Plugins/Animation/PoseSearch/Source/Runtime/Public/PoseSearch/AnimNode_MotionMatching.h
@@ -178,4 +178,5 @@ private:
 #endif // WITH_EDITORONLY_DATA
 
 	friend class UAnimGraphNode_MotionMatching;
+	friend class UMotionMatchingAnimNodeLibrary;
 };
\ No newline at end of file
diff --git a/Engine/Plugins/Animation/PoseSearch/Source/Runtime/Public/PoseSearch/MotionMatchingAnimNodeLibrary.h b/Engine/Plugins/Animation/PoseSearch/Source/Runtime/Public/PoseSearch/MotionMatchingAnimNodeLibrary.h
--- a/Engine/Plugins/Animation/PoseSearch/Source/Runtime/Public/PoseSearch/MotionMatchingAnimNodeLibrary.h
+++ b/Engine/Plugins/Animation/PoseSearch/Source/Runtime/Public/PoseSearch/MotionMatchingAnimNodeLibrary.h
@@ -8,8 +8,30 @@
 #include "MotionMatchingAnimNodeLibrary.generated.h"
 
 struct FAnimNode_MotionMatching;
+enum class EAlphaBlendOption : uint8;
+class UBlendProfile;
 class UPoseSearchDatabase;
 
+USTRUCT(BlueprintType)
+struct FMotionMatchingBlueprintBlendSettings
+{
+	GENERATED_BODY()
+
+	FMotionMatchingBlueprintBlendSettings();
+
+	UPROPERTY(EditAnywhere, BlueprintReadWrite, Category = Settings)
+	float BlendTime;
+
+	UPROPERTY(EditAnywhere, BlueprintReadWrite, Category = Settings)
+	TObjectPtr<UBlendProfile> BlendProfile;
+
+	UPROPERTY(EditAnywhere, BlueprintReadWrite, Category = Settings)
+	EAlphaBlendOption BlendOption;
+
+	UPROPERTY(EditAnywhere, BlueprintReadWrite, Category = Settings)
+	bool bUseInertialBlend;
+};
+
 USTRUCT(BlueprintType)
 struct FMotionMatchingAnimNodeReference : public FAnimNodeReference
 {
@@ -41,6 +63,14 @@ public:
 	UFUNCTION(BlueprintPure, Category = "Animation|MotionMatching", meta = (BlueprintThreadSafe, DisplayName = "Get Motion Matching Search Result"))
 	static void GetMotionMatchingSearchResult(const FMotionMatchingAnimNodeReference& MotionMatchingNode, FPoseSearchBlueprintResult& Result, bool& bIsResultValid);
 
+	// Get current blend settings used when blending into a new asset
+	UFUNCTION(BlueprintPure, Category = "Animation|MotionMatching", meta = (BlueprintThreadSafe, DisplayName = "Get Motion Matching Blend Settings"))
+	static void GetMotionMatchingBlendSettings(const FMotionMatchingAnimNodeReference& MotionMatchingNode, FMotionMatchingBlueprintBlendSettings& BlendSettings, bool& bIsResultValid);
+
+	// Override current blend settings for motion matching. Note that any pinned parameters will stomp this override on the next update.
+	UFUNCTION(BlueprintCallable, Category = "Animation|MotionMatching", meta = (BlueprintThreadSafe, DisplayName = "Override Motion Matching Blend Settings"))
+	static void OverrideMotionMatchingBlendSettings(const FMotionMatchingAnimNodeReference& MotionMatchingNode, const FMotionMatchingBlueprintBlendSettings& BlendSettings, bool& bIsResultValid);
+
 	/**
 	 * Set the database to search on the motion matching node. This overrides the Database property on the motion matching node.
 	 * @param MotionMatchingNode - The motion matching node to operate on.
