Fix for crash in AnimNode_MotionMatching

diff --git a/Engine/Plugins/Animation/PoseSearch/Source/Runtime/Public/PoseSearch/AnimNode_MotionMatching.h b/Engine/Plugins/Animation/PoseSearch/Source/Runtime/Public/PoseSearch/AnimNode_MotionMatching.h
--- a/Engine/Plugins/Animation/PoseSearch/Source/Runtime/Public/PoseSearch/AnimNode_MotionMatching.h
+++ b/Engine/Plugins/Animation/PoseSearch/Source/Runtime/Public/PoseSearch/AnimNode_MotionMatching.h
@@ -68,7 +68,7 @@ private:
 
 	// Use this to define a threshold to trigger a new blend when blendspace xy input pins change.
 	// By default, any delta will trigger a blend.
-	UPROPERTY(EditAnywhere, Category = Blendspace)
+	UPROPERTY(EditAnywhere, Category = Blendspace, meta = (FoldProperty))
 	float BlendParametersDeltaThreshold = 0.0f;
 
 	// The group name that we synchronize with (NAME_None if it is not part of any group). Note that
