motion matching - fix for previeweing channels debug lines in in pose search database for skeleton with non identity root bone

diff --git a/Engine/Plugins/Animation/PoseSearch/Source/Editor/Private/PoseSearchDatabaseViewModel.cpp b/Engine/Plugins/Animation/PoseSearch/Source/Editor/Private/PoseSearchDatabaseViewModel.cpp
--- a/Engine/Plugins/Animation/PoseSearch/Source/Editor/Private/PoseSearchDatabaseViewModel.cpp
+++ b/Engine/Plugins/Animation/PoseSearch/Source/Editor/Private/PoseSearchDatabaseViewModel.cpp
@@ -210,17 +210,20 @@ bool FDatabasePreviewActor::DrawPreviewActors(TConstArrayView<FDatabasePreviewAc
 	int32 CommonIndexAssetIndex = INDEX_NONE;
 #endif // DO_CHECK
 
-	TArray<const USkinnedMeshComponent*> Meshes;
+	TArray<const USkinnedMeshComponent*, TInlineAllocator<PreallocatedRolesNum>> Meshes;
 	FRoleToIndex RoleToIndex;
-	TArray<FArchivedPoseHistory> ArchivedPoseHistories;
-	TArray<const IPoseHistory*> PoseHistories;
+	TArray<FArchivedPoseHistory, TInlineAllocator<PreallocatedRolesNum>> ArchivedPoseHistories;
+	TArray<const IPoseHistory*, TInlineAllocator<PreallocatedRolesNum>> PoseHistories;
 
 	const int32 NumPreviewActors = PreviewActors.Num();
-
-	Meshes.Reserve(NumPreviewActors);
-	RoleToIndex.Reserve(NumPreviewActors);
-	ArchivedPoseHistories.Reserve(NumPreviewActors);
-	PoseHistories.Reserve(NumPreviewActors);
+	if (NumPreviewActors > PreallocatedRolesNum)
+	{
+		// reserve the needed amount of memory for containers
+		Meshes.Reserve(NumPreviewActors);
+		RoleToIndex.Reserve(NumPreviewActors);
+		ArchivedPoseHistories.Reserve(NumPreviewActors);
+		PoseHistories.Reserve(NumPreviewActors);
+	}
 
 	for (const FDatabasePreviewActor& PreviewActor : PreviewActors)
 	{
@@ -275,10 +278,24 @@ bool FDatabasePreviewActor::DrawPreviewActors(TConstArrayView<FDatabasePreviewAc
 		RoleToIndex.Add(PreviewActor.ActorRole) = Meshes.Num();
 		Meshes.Add(Mesh);
 
-		// @todo: reconstruct ArchivedPoseHistory::BoneToTransformMap and ArchivedPoseHistory::Entries if needed.
-		//        So far has been unnecessary, becasue all the animations have an identity root bone (root bone lock)
 		FArchivedPoseHistory& ArchivedPoseHistory = ArchivedPoseHistories.AddDefaulted_GetRef();
 		ArchivedPoseHistory.Trajectory = PreviewActor.Trajectory;
+
+		check(PoseSearchDatabase && PoseSearchDatabase->Schema);
+		if (const USkeleton* Skeleton = PoseSearchDatabase->Schema->GetSkeleton(PreviewActor.ActorRole))
+		{
+			// reconstructing ArchivedPoseHistory::BoneToTransformMap and ArchivedPoseHistory::Entries ONLY for the root bone.
+			// @todo: add more bones if needed
+			const TArray<FTransform>& RefBonePose = Skeleton->GetReferenceSkeleton().GetRefBonePose();
+			const FTransform& RefRootBone = RefBonePose[RootBoneIndexType];
+
+			ArchivedPoseHistory.BoneToTransformMap.Add(RootBoneIndexType) = RootBoneIndexType;
+			FPoseHistoryEntry& PoseHistoryEntry = ArchivedPoseHistory.Entries.AddDefaulted_GetRef();
+			// saving space for the root bone only
+			PoseHistoryEntry.SetNum(1, true);
+			PoseHistoryEntry.SetComponentSpaceTransform(RootBoneIndexType, RefRootBone);
+		}
+
 		for (FPoseSearchQueryTrajectorySample& TrajectorySample : ArchivedPoseHistory.Trajectory.Samples)
 		{
 			TrajectorySample.AccumulatedSeconds -= PreviewActor.QuantizedTime;
diff --git a/Engine/Plugins/Animation/PoseSearch/Source/Runtime/Private/PoseSearchAssetIndexer.cpp b/Engine/Plugins/Animation/PoseSearch/Source/Runtime/Private/PoseSearchAssetIndexer.cpp
--- a/Engine/Plugins/Animation/PoseSearch/Source/Runtime/Private/PoseSearchAssetIndexer.cpp
+++ b/Engine/Plugins/Animation/PoseSearch/Source/Runtime/Private/PoseSearchAssetIndexer.cpp
@@ -226,8 +226,8 @@ void FAssetIndexer::AssignWorkingData(int32 InStartPoseIdx, TArrayView<float> In
 	const int32 NumIndexedPoses = GetNumIndexedPoses();
 
 	StartPoseIdx = InStartPoseIdx;
-	FeatureVectorTable = MakeArrayView(InOutFeatureVectorTable.GetData() + Schema.SchemaCardinality * StartPoseIdx, Schema.SchemaCardinality * NumIndexedPoses);
-	PoseMetadata = MakeArrayView(InOutPoseMetadata.GetData() + StartPoseIdx, NumIndexedPoses);
+	FeatureVectorTable = InOutFeatureVectorTable.Slice(Schema.SchemaCardinality * StartPoseIdx, Schema.SchemaCardinality * NumIndexedPoses);
+	PoseMetadata = InOutPoseMetadata.Slice(StartPoseIdx, NumIndexedPoses);
 }
 
 void FAssetIndexer::Process(int32 AssetIdx)
@@ -260,9 +260,11 @@ void FAssetIndexer::Process(int32 AssetIdx)
 			CostAddend += SamplingContext.LoopingCostBias;
 		}
 
-		const int32 ValueOffset = (StartPoseIdx + GetVectorIdx(SampleIdx)) * Schema.SchemaCardinality;
+		const int32 VectorIdx = GetVectorIdx(SampleIdx);
+		const int32 PoseIdx = StartPoseIdx + VectorIdx;
+		const int32 ValueOffset = PoseIdx * Schema.SchemaCardinality;
 		check(ValueOffset >= 0 && AssetIdx >= 0);
-		PoseMetadata[GetVectorIdx(SampleIdx)] = FPoseMetadata(ValueOffset, AssetIdx, bBlockTransition, CostAddend);
+		PoseMetadata[VectorIdx] = FPoseMetadata(ValueOffset, AssetIdx, bBlockTransition, CostAddend);
 	}
 
 	// Generate pose features data
diff --git a/Engine/Plugins/Animation/PoseSearch/Source/Runtime/Private/PoseSearchContext.cpp b/Engine/Plugins/Animation/PoseSearch/Source/Runtime/Private/PoseSearchContext.cpp
--- a/Engine/Plugins/Animation/PoseSearch/Source/Runtime/Private/PoseSearchContext.cpp
+++ b/Engine/Plugins/Animation/PoseSearch/Source/Runtime/Private/PoseSearchContext.cpp
@@ -481,7 +481,7 @@ void FSearchContext::UpdateCurrentResultPoseVector()
 		{
 			const int32 NumDimensions = CurrentResult.Database->Schema->SchemaCardinality;
 			CurrentResultPoseVectorData.AddUninitialized(NumDimensions);
-			CurrentResultPoseVector = SearchIndex.GetReconstructedPoseValues(CurrentResult.PoseIdx, MakeArrayView(CurrentResultPoseVectorData.GetData() + NumDimensions, NumDimensions));
+			CurrentResultPoseVector = SearchIndex.GetReconstructedPoseValues(CurrentResult.PoseIdx, MakeArrayView(CurrentResultPoseVectorData).Slice(NumDimensions, NumDimensions));
 		}
 		else
 		{
diff --git a/Engine/Plugins/Animation/PoseSearch/Source/Runtime/Private/PoseSearchDerivedData.cpp b/Engine/Plugins/Animation/PoseSearch/Source/Runtime/Private/PoseSearchDerivedData.cpp
--- a/Engine/Plugins/Animation/PoseSearch/Source/Runtime/Private/PoseSearchDerivedData.cpp
+++ b/Engine/Plugins/Animation/PoseSearch/Source/Runtime/Private/PoseSearchDerivedData.cpp
@@ -992,7 +992,7 @@ static bool IndexDatabase(FSearchIndexBase& SearchIndexBase, const UPoseSearchDa
 	}
 
 	FAnimationAssetSamplers TempAssetSamplers;
-	TArray<FBoneContainer> TempBoneContainers;
+	TArray<FBoneContainer, TInlineAllocator<PreallocatedRolesNum>> TempBoneContainers;
 	FRoleToIndex TempRoleToIndex;
 
 	int32 TotalPoses = 0;
@@ -1176,8 +1176,8 @@ static void CompareSearchIndexBase(const FSearchIndexBase& A, const FSearchIndex
 		const int32 NumValuePoses = A.Values.Num() / Schema->SchemaCardinality;
 		for (int32 ValuePoseIndex = 0; ValuePoseIndex < NumValuePoses; ++ValuePoseIndex)
 		{
-			const TConstArrayView<float> PoseA = MakeArrayView(A.Values.GetData() + ValuePoseIndex * Schema->SchemaCardinality, Schema->SchemaCardinality);
-			const TConstArrayView<float> PoseB = MakeArrayView(B.Values.GetData() + ValuePoseIndex * Schema->SchemaCardinality, Schema->SchemaCardinality);
+			const TConstArrayView<float> PoseA = MakeArrayView(A.Values).Slice(ValuePoseIndex * Schema->SchemaCardinality, Schema->SchemaCardinality);
+			const TConstArrayView<float> PoseB = MakeArrayView(B.Values).Slice(ValuePoseIndex * Schema->SchemaCardinality, Schema->SchemaCardinality);
 			CompareChannelValues(0, ValuePoseIndex, PoseA, PoseB, Schema->GetChannels(), StringBuilder);
 		}
 	}
diff --git a/Engine/Plugins/Animation/PoseSearch/Source/Runtime/Private/PoseSearchIndex.cpp b/Engine/Plugins/Animation/PoseSearch/Source/Runtime/Private/PoseSearchIndex.cpp
--- a/Engine/Plugins/Animation/PoseSearch/Source/Runtime/Private/PoseSearchIndex.cpp
+++ b/Engine/Plugins/Animation/PoseSearch/Source/Runtime/Private/PoseSearchIndex.cpp
@@ -603,7 +603,7 @@ void FSearchIndex::PruneDuplicatePCAValues(float SimilarityThreshold, int32 Numb
 	PCAValuesVectorToPoseIndexes = FSparsePoseMultiMap<int32>();
 
 	const uint32 NumPoses = GetNumPoses();
-	if (SimilarityThreshold > 0.f && NumPoses >= 2 && NumberOfPrincipalComponents > 0)
+	if (SimilarityThreshold > 0.f && NumPoses >= 2 && NumberOfPrincipalComponents > 0 && !PCAValues.IsEmpty())
 	{
 		check(PCAValues.Num() % NumberOfPrincipalComponents == 0);
 		const int32 NumPCAValuesVectors = PCAValues.Num() / NumberOfPrincipalComponents;
diff --git a/Engine/Plugins/Animation/PoseSearch/Source/Runtime/Public/PoseSearch/PoseSearchAssetIndexer.h b/Engine/Plugins/Animation/PoseSearch/Source/Runtime/Public/PoseSearch/PoseSearchAssetIndexer.h
--- a/Engine/Plugins/Animation/PoseSearch/Source/Runtime/Public/PoseSearch/PoseSearchAssetIndexer.h
+++ b/Engine/Plugins/Animation/PoseSearch/Source/Runtime/Public/PoseSearch/PoseSearchAssetIndexer.h
@@ -46,8 +46,8 @@ struct FAnimationAssetSamplers
 	FTransform MirrorTransform(const FTransform& InTransform, int32 RoleIndex) const;
 	void MirrorPose(FCompactPose& Pose, int32 RoleIndex) const;
 
-	TArray<const FAnimationAssetSampler*> AnimationAssetSamplers;
-	TArray<const FMirrorDataCache*> MirrorDataCaches;
+	TArray<const FAnimationAssetSampler*, TInlineAllocator<PreallocatedRolesNum>> AnimationAssetSamplers;
+	TArray<const FMirrorDataCache*, TInlineAllocator<PreallocatedRolesNum>> MirrorDataCaches;
 };
 
 class POSESEARCH_API FAssetIndexer
diff --git a/Engine/Plugins/Animation/PoseSearch/Source/Runtime/Public/PoseSearch/PoseSearchHistory.h b/Engine/Plugins/Animation/PoseSearch/Source/Runtime/Public/PoseSearch/PoseSearchHistory.h
--- a/Engine/Plugins/Animation/PoseSearch/Source/Runtime/Public/PoseSearch/PoseSearchHistory.h
+++ b/Engine/Plugins/Animation/PoseSearch/Source/Runtime/Public/PoseSearch/PoseSearchHistory.h
@@ -37,10 +37,10 @@ struct FPoseHistoryEntry
 
 	void Update(float Time, FCSPose<FCompactPose>& ComponentSpacePose, const FBoneToTransformMap& BoneToTransformMap, bool bStoreScales);
 
-	void SetNum(int32 Num, bool bStoreScales);
+	POSESEARCH_API void SetNum(int32 Num, bool bStoreScales);
 	int32 Num() const;
 
-	void SetComponentSpaceTransform(int32 Index, const FTransform& Transform);
+	POSESEARCH_API void SetComponentSpaceTransform(int32 Index, const FTransform& Transform);
 	FTransform GetComponentSpaceTransform(int32 Index) const;
 };
 FArchive& operator<<(FArchive& Ar, FPoseHistoryEntry& Entry);
diff --git a/Engine/Plugins/Animation/PoseSearch/Source/Runtime/Public/PoseSearch/PoseSearchIndex.h b/Engine/Plugins/Animation/PoseSearch/Source/Runtime/Public/PoseSearch/PoseSearchIndex.h
--- a/Engine/Plugins/Animation/PoseSearch/Source/Runtime/Public/PoseSearch/PoseSearchIndex.h
+++ b/Engine/Plugins/Animation/PoseSearch/Source/Runtime/Public/PoseSearch/PoseSearchIndex.h
@@ -297,7 +297,7 @@ struct FSparsePoseMultiMap
 		const Type Value = DataValues[Key];
 		if (Value <= MaxValue)
 		{
-			return MakeArrayView(DataValues.GetData() + Key, 1);
+			return MakeArrayView(DataValues).Slice(Key, 1);
 		}
 
 		check(Value >= DeltaKeyValue);
@@ -308,7 +308,7 @@ struct FSparsePoseMultiMap
 		// ..and it's data starts at the next location DecodedArrayStartLocation + 1
 		const Type DataOffset = DecodedArrayStartLocation + 1;
 		check(int32(DataOffset + Size) <= DataValues.Num());
-		return MakeArrayView(DataValues.GetData() + DataOffset, Size);
+		return MakeArrayView(DataValues).Slice(DataOffset, Size);
 	}
 
 	Type Num() const
