[NO NON-PII COMMIT MESSAGE]

diff --git a/Engine/Plugins/Animation/PoseSearch/Source/Editor/Private/PoseSearchDatabaseViewModel.cpp b/Engine/Plugins/Animation/PoseSearch/Source/Editor/Private/PoseSearchDatabaseViewModel.cpp
--- a/Engine/Plugins/Animation/PoseSearch/Source/Editor/Private/PoseSearchDatabaseViewModel.cpp
+++ b/Engine/Plugins/Animation/PoseSearch/Source/Editor/Private/PoseSearchDatabaseViewModel.cpp
@@ -232,7 +232,13 @@ bool FDatabasePreviewActor::DrawPreviewActors(TConstArrayView<FDatabasePreviewAc
 			return false;
 		}
 
-		if (!PoseSearchDatabase->GetSearchIndex().IsValidPoseIndex(PreviewActor.GetCurrentPoseIndex()))
+		const FSearchIndex& SearchIndex = PoseSearchDatabase->GetSearchIndex();
+
+		// This condition happens when the database got reindexed and the new valid SearchIndex has different cardinality for assets or poses.
+		// Since we didn't refresh the PreviewActor IndexAssetIndex nor CurrentPoseIndex, now the PreviewActor is invalid
+		// @todo: we should refresh the PreviewActor and restore it's preview time etc
+		if (!SearchIndex.Assets.IsValidIndex(PreviewActor.GetIndexAssetIndex()) ||
+			!SearchIndex.IsValidPoseIndex(PreviewActor.GetCurrentPoseIndex()))
 		{
 			return false;
 		}
@@ -266,9 +272,9 @@ bool FDatabasePreviewActor::DrawPreviewActors(TConstArrayView<FDatabasePreviewAc
 #if DO_CHECK
 		if (CommonIndexAssetIndex == INDEX_NONE)
 		{
-			CommonIndexAssetIndex = PreviewActor.IndexAssetIndex;
+			CommonIndexAssetIndex = PreviewActor.GetIndexAssetIndex();
 		}
-		else if (CommonIndexAssetIndex != PreviewActor.IndexAssetIndex)
+		else if (CommonIndexAssetIndex != PreviewActor.GetIndexAssetIndex())
 		{
 			checkNoEntry();
 			return false;
@@ -316,7 +322,7 @@ bool FDatabasePreviewActor::DrawPreviewActors(TConstArrayView<FDatabasePreviewAc
 	{
 		const UDebugSkelMeshComponent* Mesh = PreviewActor.GetDebugSkelMeshComponent();
 		const FSearchIndex& SearchIndex = PoseSearchDatabase->GetSearchIndex();
-		const FSearchIndexAsset& IndexAsset = SearchIndex.Assets[PreviewActor.IndexAssetIndex];
+		const FSearchIndexAsset& IndexAsset = SearchIndex.Assets[PreviewActor.GetIndexAssetIndex()];
 		const int32 SamplesNum = PreviewActor.Trajectory.Samples.Num();
 
 		if (bDisplayRootMotionSpeed)
