PoseSearch: Added GetTrajectoryAngularVelocity to trajectory library and added logging and early out to invalid time horizons.

diff --git a/Engine/Plugins/Animation/PoseSearch/Source/Runtime/Private/PoseSearchTrajectoryLibrary.cpp b/Engine/Plugins/Animation/PoseSearch/Source/Runtime/Private/PoseSearchTrajectoryLibrary.cpp
--- a/Engine/Plugins/Animation/PoseSearch/Source/Runtime/Private/PoseSearchTrajectoryLibrary.cpp
+++ b/Engine/Plugins/Animation/PoseSearch/Source/Runtime/Private/PoseSearchTrajectoryLibrary.cpp
@@ -6,6 +6,7 @@
 #include "Components/SkeletalMeshComponent.h"
 #include "GameFramework/Character.h"
 #include "GameFramework/CharacterMovementComponent.h"
+#include "PoseSearch/PoseSearchDefines.h"
 
 void FPoseSearchTrajectoryData::UpdateData(
 	float DeltaTime,
@@ -583,8 +584,38 @@ void UPoseSearchTrajectoryLibrary::GetTrajectorySampleAtTime(UPARAM(ref) const F
 
 void UPoseSearchTrajectoryLibrary::GetTrajectoryVelocity(UPARAM(ref) const FPoseSearchQueryTrajectory& InTrajectory, float Time1, float Time2, FVector& OutVelocity, bool bExtrapolate)
 {
+	if (FMath::IsNearlyEqual(Time1, Time2))
+	{
+		UE_LOG(LogPoseSearch, Warning, TEXT("UPoseSearchTrajectoryLibrary::GetTrajectoryVelocity - Time1 is same as Time2. Invalid time horizon."));
+		OutVelocity = FVector::ZeroVector;
+		return;
+	}
+
 	FPoseSearchQueryTrajectorySample Sample1 = InTrajectory.GetSampleAtTime(Time1, bExtrapolate);
 	FPoseSearchQueryTrajectorySample Sample2 = InTrajectory.GetSampleAtTime(Time2, bExtrapolate);
 
 	OutVelocity = (Sample2.Position - Sample1.Position) / (Time2 - Time1);
 }
+
+void UPoseSearchTrajectoryLibrary::GetTrajectoryAngularVelocity(const FPoseSearchQueryTrajectory& InTrajectory, float Time1, float Time2, FVector& OutAngularVelocity, bool bExtrapolate /*= false*/)
+{
+	if (FMath::IsNearlyEqual(Time1, Time2))
+	{
+		UE_LOG(LogPoseSearch, Warning, TEXT("UPoseSearchTrajectoryLibrary::GetTrajectoryAngularVelocity - Time1 is same as Time2. Invalid time horizon."));
+		OutAngularVelocity = FVector::ZeroVector;
+		return;
+	}
+
+	FPoseSearchQueryTrajectorySample Sample1 = InTrajectory.GetSampleAtTime(Time1, bExtrapolate);
+	FPoseSearchQueryTrajectorySample Sample2 = InTrajectory.GetSampleAtTime(Time2, bExtrapolate);
+
+	FQuat DeltaRotation = (Sample2.Facing * Sample1.Facing.Inverse());
+	DeltaRotation.EnforceShortestArcWith(FQuat::Identity);
+
+	const FVector AngularVelocityInRadians = DeltaRotation.ToRotationVector() / (Time2 - Time1);
+
+	OutAngularVelocity = FVector(
+		FMath::RadiansToDegrees(AngularVelocityInRadians.X),
+		FMath::RadiansToDegrees(AngularVelocityInRadians.Y),
+		FMath::RadiansToDegrees(AngularVelocityInRadians.Z));
+}
\ No newline at end of file
diff --git a/Engine/Plugins/Animation/PoseSearch/Source/Runtime/Public/PoseSearch/PoseSearchTrajectoryLibrary.h b/Engine/Plugins/Animation/PoseSearch/Source/Runtime/Public/PoseSearch/PoseSearchTrajectoryLibrary.h
--- a/Engine/Plugins/Animation/PoseSearch/Source/Runtime/Public/PoseSearch/PoseSearchTrajectoryLibrary.h
+++ b/Engine/Plugins/Animation/PoseSearch/Source/Runtime/Public/PoseSearch/PoseSearchTrajectoryLibrary.h
@@ -157,6 +157,9 @@ public:
 	UFUNCTION(BlueprintPure, meta = (BlueprintThreadSafe), Category="Animation|PoseSearch")
 	static void GetTrajectoryVelocity(UPARAM(ref) const FPoseSearchQueryTrajectory& InTrajectory, float Time1, float Time2, FVector& OutVelocity, bool bExtrapolate = false);
 	
+	UFUNCTION(BlueprintPure, meta = (BlueprintThreadSafe), Category="Animation|PoseSearch")
+	static void GetTrajectoryAngularVelocity(UPARAM(ref) const FPoseSearchQueryTrajectory& InTrajectory, float Time1, float Time2, FVector& OutAngularVelocity, bool bExtrapolate = false);
+
 	UFUNCTION(BlueprintPure, meta = (BlueprintThreadSafe), Category="Animation|PoseSearch")
 	static FTransform GetTransform(UPARAM(ref) const FPoseSearchQueryTrajectorySample& InTrajectorySample)
 	{
