Pose Search: Fix for playrate and selection at variable framerates

diff --git a/Engine/Plugins/Animation/PoseSearch/Source/Runtime/Private/PoseSearchTrajectoryLibrary.cpp b/Engine/Plugins/Animation/PoseSearch/Source/Runtime/Private/PoseSearchTrajectoryLibrary.cpp
--- a/Engine/Plugins/Animation/PoseSearch/Source/Runtime/Private/PoseSearchTrajectoryLibrary.cpp
+++ b/Engine/Plugins/Animation/PoseSearch/Source/Runtime/Private/PoseSearchTrajectoryLibrary.cpp
@@ -240,10 +240,14 @@ void UPoseSearchTrajectoryLibrary::UpdateHistory_TransformHistory(
 
 		check(NumHistorySamples <= Trajectory.Samples.Num());
 
+		// Our trajectory's "current" position assumes the we have the same delta time as the previous frame.
+		// so use predicted trajectory with current time step.
+		const FVector PredictedPositionAdjusted = Trajectory.GetSampleAtTime(DeltaTime).Position;
+
 		// converting all the history samples relative to the previous character position (Trajectory.Samples[NumHistorySamples].Position)
 		for (int32 Index = 0; Index < NumHistorySamples; ++Index)
 		{
-			Trajectory.Samples[Index].Position = Trajectory.Samples[NumHistorySamples].Position - Trajectory.Samples[Index].Position;
+			Trajectory.Samples[Index].Position = PredictedPositionAdjusted - Trajectory.Samples[Index].Position;
 		}
 
 		FVector CurrentTranslation = CurrentVelocity * DeltaTime;
