pose search - PoseSearchAssetSamplerLibrary

diff --git a/Engine/Plugins/Animation/PoseSearch/Source/Runtime/Private/PoseSearchAssetSamplerLibrary.cpp b/Engine/Plugins/Animation/PoseSearch/Source/Runtime/Private/PoseSearchAssetSamplerLibrary.cpp
new file mode 100644
--- /dev/null
+++ b/Engine/Plugins/Animation/PoseSearch/Source/Runtime/Private/PoseSearchAssetSamplerLibrary.cpp
@@ -0,0 +1,142 @@
+// Copyright Epic Games, Inc. All Rights Reserved.
+
+#include "PoseSearch/PoseSearchAssetSamplerLibrary.h"
+#include "Animation/AnimInstance.h"
+#include "DrawDebugHelpers.h"
+#include "PoseSearch/PoseSearchDefines.h"
+#include "PoseSearch/PoseSearchMirrorDataCache.h"
+
+FPoseSearchAssetSamplerPose UPoseSearchAssetSamplerLibrary::SamplePose(const UAnimInstance* AnimInstance, const FPoseSearchAssetSamplerInput Input)
+{
+	FPoseSearchAssetSamplerPose AssetSamplerPose;
+	using namespace UE::PoseSearch;
+	if (!Input.Animation)
+	{
+		UE_LOG(LogPoseSearch, Error, TEXT("UPoseSearchAssetSamplerLibrary::SamplePose invalid Input.Animation"));
+	}
+	else if (!AnimInstance)
+	{
+		UE_LOG(LogPoseSearch, Error, TEXT("UPoseSearchAssetSamplerLibrary::SamplePose invalid AnimInstance"));
+	}
+	else if (Input.bMirrored && !Input.MirrorDataTable)
+	{
+		UE_LOG(LogPoseSearch, Error, TEXT("UPoseSearchAssetSamplerLibrary::SamplePose unable to mirror the pose from %s at time %f because of invalid MirrorDataTable"), *Input.Animation->GetName(), Input.AnimationTime);
+	}
+	else
+	{
+		const FBoneContainer& BoneContainer = AnimInstance->GetRequiredBonesOnAnyThread();
+
+		FMemMark Mark(FMemStack::Get());
+
+		bool bPreProcessRootTransform = true;
+		const FAnimationAssetSampler Sampler(Input.Animation, Input.RootTransformOrigin, Input.BlendParameters, Input.RootTransformSamplingRate, bPreProcessRootTransform);
+
+		FBlendedCurve Curve;
+		FCompactPose Pose;
+		Pose.SetBoneContainer(&BoneContainer);
+
+		Sampler.ExtractPose(Input.AnimationTime, Pose, Curve);
+		AssetSamplerPose.RootTransform = Sampler.ExtractRootTransform(Input.AnimationTime);
+
+		if (Input.bMirrored)
+		{
+			const FMirrorDataCache MirrorDataCache(Input.MirrorDataTable, BoneContainer);
+			MirrorDataCache.MirrorPose(Pose);
+			AssetSamplerPose.RootTransform = MirrorDataCache.MirrorTransform(AssetSamplerPose.RootTransform);
+		}
+
+		AssetSamplerPose.Pose.CopyBonesFrom(Pose);
+		AssetSamplerPose.ComponentSpacePose.InitPose(AssetSamplerPose.Pose);
+	}
+	return AssetSamplerPose;
+}
+
+FTransform UPoseSearchAssetSamplerLibrary::GetTransform(UPARAM(ref) FPoseSearchAssetSamplerPose& AssetSamplerPose, int32 BoneIndex, EPoseSearchAssetSamplerSpace Space)
+{
+	const FCompactPoseBoneIndex CompactPoseBoneIndex(BoneIndex);
+
+	if (!AssetSamplerPose.Pose.IsValid())
+	{
+		UE_LOG(LogPoseSearch, Error, TEXT("UPoseSearchAssetSamplerLibrary::GetBoneTransform invalid AssetSamplerPose.Pose"));
+		return FTransform::Identity;
+	}
+
+	if (BoneIndex == INDEX_NONE)
+	{
+		if (Space != EPoseSearchAssetSamplerSpace::World)
+		{
+			UE_LOG(LogPoseSearch, Error, TEXT("UPoseSearchAssetSamplerLibrary::GetBoneTransform invalid Space %s to get the RootTransform. Expected space: %s"), *UEnum::GetDisplayValueAsText(Space).ToString(), *UEnum::GetDisplayValueAsText(EPoseSearchAssetSamplerSpace::World).ToString());
+		}
+		return AssetSamplerPose.RootTransform;
+	}
+	
+	if (!AssetSamplerPose.Pose.IsValidIndex(CompactPoseBoneIndex))
+	{
+		UE_LOG(LogPoseSearch, Error, TEXT("UPoseSearchAssetSamplerLibrary::GetBoneTransform invalid BoneIndex %d"), BoneIndex);
+		return FTransform::Identity;
+	}
+	
+	switch (Space)
+	{
+	case EPoseSearchAssetSamplerSpace::Local:
+		return AssetSamplerPose.Pose[CompactPoseBoneIndex];
+
+	case EPoseSearchAssetSamplerSpace::Component:
+		return AssetSamplerPose.ComponentSpacePose.GetComponentSpaceTransform(CompactPoseBoneIndex);
+
+	case EPoseSearchAssetSamplerSpace::World:
+		return AssetSamplerPose.ComponentSpacePose.GetComponentSpaceTransform(CompactPoseBoneIndex) * AssetSamplerPose.RootTransform;
+	}
+
+	checkNoEntry();
+	return FTransform::Identity;
+}
+
+FTransform UPoseSearchAssetSamplerLibrary::GetTransformByName(UPARAM(ref) FPoseSearchAssetSamplerPose& AssetSamplerPose, FName BoneName, EPoseSearchAssetSamplerSpace Space)
+{
+	if (!AssetSamplerPose.Pose.IsValid())
+	{
+		UE_LOG(LogPoseSearch, Error, TEXT("UPoseSearchAssetSamplerLibrary::GetTransformByName invalid AssetSamplerPose.Pose"));
+		return FTransform::Identity;
+	}
+
+	const USkeleton* Skeleton = AssetSamplerPose.Pose.GetBoneContainer().GetSkeletonAsset();
+
+	FBoneReference BoneReference;
+	BoneReference.BoneName = BoneName;
+	BoneReference.Initialize(Skeleton);
+	if (!BoneReference.HasValidSetup())
+	{
+		UE_LOG(LogPoseSearch, Error, TEXT("UPoseSearchAssetSamplerLibrary::GetTransformByName invalid BoneName %s for Skeleton %s"), *BoneName.ToString(), *GetNameSafe(Skeleton));
+		return FTransform::Identity;
+	}
+
+	return GetTransform(AssetSamplerPose, BoneReference.BoneIndex, Space);
+}
+
+void UPoseSearchAssetSamplerLibrary::Draw(const UAnimInstance* AnimInstance, UPARAM(ref) FPoseSearchAssetSamplerPose& AssetSamplerPose)
+{
+	static float DebugDrawSamplerRootAxisLength = 20.f;
+	static float DebugDrawSamplerSize = 6.f;
+
+	if (AnimInstance)
+	{
+		if (UWorld* World = AnimInstance->GetWorld())
+		{
+			if (DebugDrawSamplerRootAxisLength > 0.f)
+			{
+				const FTransform RootTransform = GetTransform(AssetSamplerPose);
+				DrawDebugLine(World, RootTransform.GetTranslation(), RootTransform.GetTranslation() + RootTransform.GetScaledAxis(EAxis::X) * DebugDrawSamplerRootAxisLength, FColor::Red, false, 0.f, SDPG_Foreground);
+				DrawDebugLine(World, RootTransform.GetTranslation(), RootTransform.GetTranslation() + RootTransform.GetScaledAxis(EAxis::Y) * DebugDrawSamplerRootAxisLength, FColor::Green, false, 0.f, SDPG_Foreground);
+				DrawDebugLine(World, RootTransform.GetTranslation(), RootTransform.GetTranslation() + RootTransform.GetScaledAxis(EAxis::Z) * DebugDrawSamplerRootAxisLength, FColor::Blue, false, 0.f, SDPG_Foreground);
+			}
+
+			for (int32 BoneIndex = 0; BoneIndex < AssetSamplerPose.ComponentSpacePose.GetPose().GetNumBones(); ++BoneIndex)
+			{
+				const FTransform BoneWorldTransforms = GetTransform(AssetSamplerPose, BoneIndex, EPoseSearchAssetSamplerSpace::World);
+				DrawDebugPoint(World, BoneWorldTransforms.GetTranslation(), DebugDrawSamplerSize, FColor::Red, false, 0.f, SDPG_Foreground);
+			}
+		}
+	}
+}
+
diff --git a/Engine/Plugins/Animation/PoseSearch/Source/Runtime/Public/PoseSearch/MultiAnimAsset.h b/Engine/Plugins/Animation/PoseSearch/Source/Runtime/Public/PoseSearch/MultiAnimAsset.h
--- a/Engine/Plugins/Animation/PoseSearch/Source/Runtime/Public/PoseSearch/MultiAnimAsset.h
+++ b/Engine/Plugins/Animation/PoseSearch/Source/Runtime/Public/PoseSearch/MultiAnimAsset.h
@@ -30,4 +30,7 @@ public:
 
 	UFUNCTION(BlueprintPure, Category = "Animation", meta=(BlueprintThreadSafe, DisplayName = "Get Animation Asset"))
 	UAnimationAsset* BP_GetAnimationAsset(const FName& Role) const { return GetAnimationAsset(Role); }
+
+	UFUNCTION(BlueprintPure, Category = "Animation", meta=(BlueprintThreadSafe, DisplayName = "Get Origin"))
+	FTransform BP_GetOrigin(const FName& Role) const { return GetOrigin(Role); }
 };
\ No newline at end of file
diff --git a/Engine/Plugins/Animation/PoseSearch/Source/Runtime/Public/PoseSearch/PoseSearchAssetSamplerLibrary.h b/Engine/Plugins/Animation/PoseSearch/Source/Runtime/Public/PoseSearch/PoseSearchAssetSamplerLibrary.h
new file mode 100644
--- /dev/null
+++ b/Engine/Plugins/Animation/PoseSearch/Source/Runtime/Public/PoseSearch/PoseSearchAssetSamplerLibrary.h
@@ -0,0 +1,82 @@
+// Copyright Epic Games, Inc. All Rights Reserved.
+
+#pragma once
+
+#include "Animation/AnimCurveTypes.h"
+#include "Kismet/BlueprintFunctionLibrary.h"
+#include "PoseSearch/PoseSearchAssetSampler.h"
+#include "PoseSearchAssetSamplerLibrary.generated.h"
+
+USTRUCT(Experimental, BlueprintType, Category="Animation|Pose Search")
+struct POSESEARCH_API FPoseSearchAssetSamplerInput
+{
+	GENERATED_BODY()
+
+	// Animation to sample
+	UPROPERTY(EditAnywhere, BlueprintReadWrite, Category=State)
+	TObjectPtr<const UAnimationAsset> Animation;
+
+	// Sampling time for Animation
+	UPROPERTY(EditAnywhere, BlueprintReadWrite, Category=State)
+	float AnimationTime = 0.f;
+
+	// origin used to start sampling Animation at time of zero
+	UPROPERTY(EditAnywhere, BlueprintReadWrite, Category=State)
+	FTransform RootTransformOrigin = FTransform::Identity;
+
+	UPROPERTY(EditAnywhere, BlueprintReadWrite, Category=State)
+	bool bMirrored = false;
+
+	UPROPERTY(EditAnywhere, BlueprintReadWrite, Category=State)
+	TObjectPtr<const UMirrorDataTable> MirrorDataTable;
+
+	// blend parameters if Animation is a blend space
+	UPROPERTY(EditAnywhere, BlueprintReadWrite, Category=State)
+	FVector BlendParameters = FVector::ZeroVector;
+	
+	// frequency of sampling while sampling the root transform of blend spaces
+	UPROPERTY(EditAnywhere, BlueprintReadWrite, Category=State)
+	int32 RootTransformSamplingRate = UE::PoseSearch::FAnimationAssetSampler::DefaultRootTransformSamplingRate;
+};
+
+USTRUCT(Experimental, BlueprintType, Category="Animation|Pose Search")
+struct POSESEARCH_API FPoseSearchAssetSamplerPose
+{
+	GENERATED_BODY()
+
+	FTransform RootTransform = FTransform::Identity;
+	FCompactHeapPose Pose;
+	FBlendedHeapCurve Curve;
+	// @todo: add Attribute(s)
+	//FHeapAttributeContainer Attribute;
+
+	FCSPose<FCompactHeapPose> ComponentSpacePose;
+};
+
+UENUM()
+enum class EPoseSearchAssetSamplerSpace : uint8
+{
+	Local,
+	Component,
+	World
+};
+
+UCLASS(Experimental)
+class POSESEARCH_API UPoseSearchAssetSamplerLibrary : public UBlueprintFunctionLibrary
+{
+	GENERATED_BODY()
+
+	UFUNCTION(BlueprintPure, Category = "Animation|Pose Search|Experimental", meta = (BlueprintThreadSafe))
+	static FPoseSearchAssetSamplerPose SamplePose(const UAnimInstance* AnimInstance, const FPoseSearchAssetSamplerInput Input);
+
+	UFUNCTION(BlueprintPure, Category = "Animation|Pose Search|Experimental", meta = (BlueprintThreadSafe))
+	static FTransform GetTransform(UPARAM(ref) FPoseSearchAssetSamplerPose& AssetSamplerPose, int32 BoneIndex = -1, EPoseSearchAssetSamplerSpace Space = EPoseSearchAssetSamplerSpace::World);
+
+	UFUNCTION(BlueprintPure, Category = "Animation|Pose Search|Experimental", meta = (BlueprintThreadSafe))
+	static FTransform GetTransformByName(UPARAM(ref) FPoseSearchAssetSamplerPose& AssetSamplerPose, FName BoneName, EPoseSearchAssetSamplerSpace Space = EPoseSearchAssetSamplerSpace::World);
+
+	// @todo: it'd be nice if it was threadsafe...
+	UFUNCTION(BlueprintCallable, Category="Animation|Pose Search|Experimental")
+	static void Draw(const UAnimInstance* AnimInstance, UPARAM(ref) FPoseSearchAssetSamplerPose& AssetSamplerPose);
+};
+
diff --git a/Engine/Plugins/Animation/PoseSearch/Source/Runtime/Public/PoseSearch/PoseSearchMirrorDataCache.h b/Engine/Plugins/Animation/PoseSearch/Source/Runtime/Public/PoseSearch/PoseSearchMirrorDataCache.h
--- a/Engine/Plugins/Animation/PoseSearch/Source/Runtime/Public/PoseSearch/PoseSearchMirrorDataCache.h
+++ b/Engine/Plugins/Animation/PoseSearch/Source/Runtime/Public/PoseSearch/PoseSearchMirrorDataCache.h
@@ -11,6 +11,7 @@ struct FCompactPose;
 namespace UE::PoseSearch
 {
 
+// @todo: make a memstack allocator version of FMirrorDataCache
 struct POSESEARCH_API FMirrorDataCache
 {
 	FMirrorDataCache();
