PoseSearch: Fix anim velocity reporting in debugger + added non-timescaled version for debugging

diff --git a/Engine/Plugins/Animation/PoseSearch/Source/Editor/Private/PoseSearchDebuggerReflection.h b/Engine/Plugins/Animation/PoseSearch/Source/Editor/Private/PoseSearchDebuggerReflection.h
--- a/Engine/Plugins/Animation/PoseSearch/Source/Editor/Private/PoseSearchDebuggerReflection.h
+++ b/Engine/Plugins/Animation/PoseSearch/Source/Editor/Private/PoseSearchDebuggerReflection.h
@@ -35,4 +35,13 @@ public:
 
 	UPROPERTY(VisibleAnywhere, Category = "Motion Matching State", meta = (ForceUnits = "deg/s"))
 	float AnimAngularVelocity;
+
+	UPROPERTY(VisibleAnywhere, Category = "Motion Matching State", meta = (ForceUnits = "s"))
+	float Playrate;
+
+	UPROPERTY(VisibleAnywhere, Category = "Motion Matching State", meta = (ForceUnits = "cm/s"))
+	float AnimLinearVelocityNoTimescale;
+
+	UPROPERTY(VisibleAnywhere, Category = "Motion Matching State", meta = (ForceUnits = "deg/s"))
+	float AnimAngularVelocityNoTimescale;
 };
diff --git a/Engine/Plugins/Animation/PoseSearch/Source/Editor/Private/PoseSearchDebuggerView.cpp b/Engine/Plugins/Animation/PoseSearch/Source/Editor/Private/PoseSearchDebuggerView.cpp
--- a/Engine/Plugins/Animation/PoseSearch/Source/Editor/Private/PoseSearchDebuggerView.cpp
+++ b/Engine/Plugins/Animation/PoseSearch/Source/Editor/Private/PoseSearchDebuggerView.cpp
@@ -100,6 +100,9 @@ void SDebuggerDetailsView::UpdateReflection(const FTraceMotionMatchingStateMessa
 	Reflection->SimAngularVelocity = State.SimAngularVelocity;
 	Reflection->AnimLinearVelocity = State.AnimLinearVelocity;
 	Reflection->AnimAngularVelocity = State.AnimAngularVelocity;
+	Reflection->Playrate = State.Playrate;
+	Reflection->AnimLinearVelocityNoTimescale = State.AnimLinearVelocityNoTimescale;
+	Reflection->AnimAngularVelocityNoTimescale = State.AnimAngularVelocityNoTimescale;
 }
 
 void SDebuggerView::Construct(const FArguments& InArgs, uint64 InAnimInstanceId, int32 InWantedSearchId)
diff --git a/Engine/Plugins/Animation/PoseSearch/Source/Runtime/Private/AnimNode_MotionMatching.cpp b/Engine/Plugins/Animation/PoseSearch/Source/Runtime/Private/AnimNode_MotionMatching.cpp
--- a/Engine/Plugins/Animation/PoseSearch/Source/Runtime/Private/AnimNode_MotionMatching.cpp
+++ b/Engine/Plugins/Animation/PoseSearch/Source/Runtime/Private/AnimNode_MotionMatching.cpp
@@ -84,10 +84,6 @@ void FAnimNode_MotionMatching::Evaluate_AnyThread(FPoseContext& Output)
 	}
 
 	MotionMatchingState.AnimationDeltaYaw = FRotator(RootMotionTransformDelta.GetRotation()).Yaw;
-
-#if UE_POSE_SEARCH_TRACE_ENABLED
-	MotionMatchingState.RootMotionTransformDelta = RootMotionTransformDelta;
-#endif //UE_POSE_SEARCH_TRACE_ENABLED
 }
 
 void FAnimNode_MotionMatching::UpdateAssetPlayer(const FAnimationUpdateContext& Context)
diff --git a/Engine/Plugins/Animation/PoseSearch/Source/Runtime/Private/PoseSearchLibrary.cpp b/Engine/Plugins/Animation/PoseSearch/Source/Runtime/Private/PoseSearchLibrary.cpp
--- a/Engine/Plugins/Animation/PoseSearch/Source/Runtime/Private/PoseSearchLibrary.cpp
+++ b/Engine/Plugins/Animation/PoseSearch/Source/Runtime/Private/PoseSearchLibrary.cpp
@@ -202,10 +202,6 @@ void FMotionMatchingState::Reset(const FTransform& ComponentTransform)
 	AnimationDeltaYaw = 0.f;
 
 	PoseIndicesHistory.Reset();
-
-#if UE_POSE_SEARCH_TRACE_ENABLED
-	RootMotionTransformDelta = FTransform::Identity;
-#endif // UE_POSE_SEARCH_TRACE_ENABLED
 }
 
 void FMotionMatchingState::AdjustAssetTime(float AssetTime)
@@ -281,21 +277,16 @@ void FMotionMatchingState::UpdateWantedPlayRate(const UE::PoseSearch::FSearchCon
 }
 
 #if UE_POSE_SEARCH_TRACE_ENABLED
-void UPoseSearchLibrary::TraceMotionMatching(
-	UE::PoseSearch::FSearchContext& SearchContext,
-	const UE::PoseSearch::FSearchResult& CurrentResult,
-	float ElapsedPoseSearchTime,
-	const FTransform& RootMotionTransformDelta,
-	float DeltaTime,
-	bool bSearch,
-	float RecordingTime)
+void UPoseSearchLibrary::TraceMotionMatching(UE::PoseSearch::FSearchContext& SearchContext, FMotionMatchingState& CurrentState, float DeltaTime, bool bSearch, float RecordingTime)
 {
 	using namespace UE::PoseSearch;
 	
 	uint32 SearchId = 787;
 
 	FTraceMotionMatchingStateMessage TraceState;
-	
+	FSearchResult& CurrentResult = CurrentState.CurrentSearchResult;
+	const float ElapsedPoseSearchTime = CurrentState.ElapsedPoseSearchTime;
+
 	const int32 AnimInstancesNum = SearchContext.GetAnimInstances().Num();
 	TraceState.SkeletalMeshComponentIds.SetNum(AnimInstancesNum);
 
@@ -393,9 +384,32 @@ void UPoseSearchLibrary::TraceMotionMatching(
 			}
 		}
 
-		// animation
-		TraceState.AnimLinearVelocity = RootMotionTransformDelta.GetTranslation().Size() / DeltaTime;
-		TraceState.AnimAngularVelocity = FMath::RadiansToDegrees(RootMotionTransformDelta.GetRotation().GetAngle()) / DeltaTime;
+		const FSearchIndexAsset* SearchIndexAsset = CurrentResult.GetSearchIndexAsset();
+		const UPoseSearchDatabase* CurrentResultDatabase = CurrentResult.Database.Get();
+		if (SearchIndexAsset && CurrentResultDatabase)
+		{
+			const FPoseSearchDatabaseAnimationAssetBase* DatabaseAsset = CurrentResultDatabase->GetDatabaseAnimationAsset<FPoseSearchDatabaseAnimationAssetBase>(*SearchIndexAsset);
+			check(DatabaseAsset);
+			if (UAnimationAsset* AnimationAsset = Cast<UAnimationAsset>(DatabaseAsset->GetAnimationAsset()))
+			{
+				// Simulate the time step to get accurate root motion prediction for this frame.
+				FAnimationAssetSampler Sampler(AnimationAsset);
+
+				const float TimeStep = DeltaTime * CurrentState.WantedPlayRate;
+				const FTransform PrevRoot = Sampler.ExtractRootTransform(CurrentResult.AssetTime);
+				const FTransform CurrRoot = Sampler.ExtractRootTransform(CurrentResult.AssetTime + TimeStep);
+				const FTransform RootMotionTransformDelta = PrevRoot.GetRelativeTransform(CurrRoot);
+				TraceState.AnimLinearVelocity = RootMotionTransformDelta.GetTranslation().Size() / DeltaTime;
+				TraceState.AnimAngularVelocity = FMath::RadiansToDegrees(RootMotionTransformDelta.GetRotation().GetAngle()) / DeltaTime;
+
+				// Need another root motion extraction for non-playrate version in case acceleration isn't the same.
+				const FTransform CurrRootNoTimescale = Sampler.ExtractRootTransform(CurrentResult.AssetTime + DeltaTime);
+				const FTransform RootMotionTransformDeltaNoTimescale = PrevRoot.GetRelativeTransform(CurrRootNoTimescale);
+				TraceState.AnimLinearVelocityNoTimescale = RootMotionTransformDeltaNoTimescale.GetTranslation().Size() / DeltaTime;
+				TraceState.AnimAngularVelocityNoTimescale = FMath::RadiansToDegrees(RootMotionTransformDeltaNoTimescale.GetRotation().GetAngle()) / DeltaTime;
+			}
+		}
+		TraceState.Playrate = CurrentState.WantedPlayRate;
 	}
 
 	TraceState.ElapsedPoseSearchTime = ElapsedPoseSearchTime;
@@ -599,8 +613,7 @@ void UPoseSearchLibrary::UpdateMotionMatchingState(
 	// Record debugger details
 	if (IsTracing(Context))
 	{
-		TraceMotionMatching(SearchContext, InOutMotionMatchingState.CurrentSearchResult, InOutMotionMatchingState.ElapsedPoseSearchTime,
-			InOutMotionMatchingState.RootMotionTransformDelta, DeltaTime, bSearch,
+		TraceMotionMatching(SearchContext, InOutMotionMatchingState, DeltaTime, bSearch,
 			AnimInstance ? FObjectTrace::GetWorldElapsedTime(AnimInstance->GetWorld()) : 0.f);
 	}
 #endif // UE_POSE_SEARCH_TRACE_ENABLED
@@ -1019,7 +1032,10 @@ UE::PoseSearch::FSearchResult UPoseSearchLibrary::MotionMatch(
 #endif // ENABLE_DRAW_DEBUG && ENABLE_ANIM_DEBUG
 
 #if UE_POSE_SEARCH_TRACE_ENABLED
-	TraceMotionMatching(SearchContext, SearchResult, 0.f, FTransform::Identity,
+	FMotionMatchingState MotionMatchingState;
+	MotionMatchingState.CurrentSearchResult = SearchResult;
+	MotionMatchingState.ElapsedPoseSearchTime = 0.0f;
+	TraceMotionMatching(SearchContext, MotionMatchingState,
 		DeltaSeconds, true, FObjectTrace::GetWorldElapsedTime(AnimInstances[0]->GetWorld()));
 #endif // UE_POSE_SEARCH_TRACE_ENABLED
 
diff --git a/Engine/Plugins/Animation/PoseSearch/Source/Runtime/Private/Trace/PoseSearchTraceLogger.cpp b/Engine/Plugins/Animation/PoseSearch/Source/Runtime/Private/Trace/PoseSearchTraceLogger.cpp
--- a/Engine/Plugins/Animation/PoseSearch/Source/Runtime/Private/Trace/PoseSearchTraceLogger.cpp
+++ b/Engine/Plugins/Animation/PoseSearch/Source/Runtime/Private/Trace/PoseSearchTraceLogger.cpp
@@ -77,6 +77,9 @@ FArchive& operator<<(FArchive& Ar, FTraceMotionMatchingStateMessage& State)
 	Ar << State.SimAngularVelocity;
 	Ar << State.AnimLinearVelocity;
 	Ar << State.AnimAngularVelocity;
+	Ar << State.Playrate;
+	Ar << State.AnimLinearVelocityNoTimescale;
+	Ar << State.AnimAngularVelocityNoTimescale;
 	Ar << State.RecordingTime;
 	Ar << State.SearchBestCost;
 	Ar << State.SearchBruteForceCost;
diff --git a/Engine/Plugins/Animation/PoseSearch/Source/Runtime/Public/PoseSearch/PoseSearchLibrary.h b/Engine/Plugins/Animation/PoseSearch/Source/Runtime/Public/PoseSearch/PoseSearchLibrary.h
--- a/Engine/Plugins/Animation/PoseSearch/Source/Runtime/Public/PoseSearch/PoseSearchLibrary.h
+++ b/Engine/Plugins/Animation/PoseSearch/Source/Runtime/Public/PoseSearch/PoseSearchLibrary.h
@@ -42,8 +42,10 @@ enum class EPoseSearchInterruptMode : uint8
 	ForceInterruptAndInvalidateContinuingPose,
 };
 
+PRAGMA_DISABLE_DEPRECATION_WARNINGS
 struct FMotionMatchingState
 {
+PRAGMA_ENABLE_DEPRECATION_WARNINGS
 	// Reset the state to a default state using the current Database
 	void Reset(const FTransform& ComponentTransform);
 
@@ -82,7 +84,7 @@ struct FMotionMatchingState
 	float AnimationDeltaYaw = 0.f;
 
 #if UE_POSE_SEARCH_TRACE_ENABLED
-	// Root motion delta for currently playing animation (or animation tree if from the blend stack)
+	UE_DEPRECATED(5.4, "Debug RootMotionDelta is now calculated in-place.")
 	FTransform RootMotionTransformDelta = FTransform::Identity;
 #endif //UE_POSE_SEARCH_TRACE_ENABLED
 };
@@ -127,6 +129,15 @@ class POSESEARCH_API UPoseSearchLibrary : public UBlueprintFunctionLibrary
 	GENERATED_BODY()
 
 #if UE_POSE_SEARCH_TRACE_ENABLED
+
+	static void TraceMotionMatching(
+		UE::PoseSearch::FSearchContext& SearchContext,
+		FMotionMatchingState& CurrentState,
+		float DeltaTime,
+		bool bSearch,
+		float RecordingTime);
+	
+	UE_DEPRECATED(5.5, "Use TraceMotionMatching with different signature instead.")
 	static void TraceMotionMatching(
 		UE::PoseSearch::FSearchContext& SearchContext,
 		const UE::PoseSearch::FSearchResult& CurrentResult,
@@ -134,7 +145,13 @@ class POSESEARCH_API UPoseSearchLibrary : public UBlueprintFunctionLibrary
 		const FTransform& RootMotionTransformDelta,
 		float DeltaTime,
 		bool bSearch,
-		float RecordingTime);
+		float RecordingTime)
+	{
+		FMotionMatchingState MotionMatchingState;
+		MotionMatchingState.CurrentSearchResult = CurrentResult;
+		MotionMatchingState.ElapsedPoseSearchTime = ElapsedPoseSearchTime;
+		TraceMotionMatching(SearchContext, MotionMatchingState, DeltaTime, bSearch, RecordingTime);
+	}
 
 	UE_DEPRECATED(5.4, "Use TraceMotionMatching instead")
 	static void TraceMotionMatchingState(
@@ -147,7 +164,10 @@ class POSESEARCH_API UPoseSearchLibrary : public UBlueprintFunctionLibrary
 		bool bSearch,
 		float RecordingTime)
 	{
-		TraceMotionMatching(SearchContext, CurrentResult, ElapsedPoseSearchTime, RootMotionTransformDelta, DeltaTime, bSearch, RecordingTime);
+		FMotionMatchingState MotionMatchingState;
+		MotionMatchingState.CurrentSearchResult = CurrentResult;
+		MotionMatchingState.ElapsedPoseSearchTime = ElapsedPoseSearchTime;
+		TraceMotionMatching(SearchContext, MotionMatchingState, DeltaTime, bSearch, RecordingTime);
 	}
 	
 #endif // UE_POSE_SEARCH_TRACE_ENABLED
diff --git a/Engine/Plugins/Animation/PoseSearch/Source/Runtime/Public/PoseSearch/Trace/PoseSearchTraceLogger.h b/Engine/Plugins/Animation/PoseSearch/Source/Runtime/Public/PoseSearch/Trace/PoseSearchTraceLogger.h
--- a/Engine/Plugins/Animation/PoseSearch/Source/Runtime/Public/PoseSearch/Trace/PoseSearchTraceLogger.h
+++ b/Engine/Plugins/Animation/PoseSearch/Source/Runtime/Public/PoseSearch/Trace/PoseSearchTraceLogger.h
@@ -68,6 +68,9 @@ struct POSESEARCH_API FTraceMotionMatchingStateMessage : public FTraceMessage
 	float SimAngularVelocity = 0.f;
 	float AnimLinearVelocity = 0.f;
 	float AnimAngularVelocity = 0.f;
+	float Playrate = 0.f;
+	float AnimLinearVelocityNoTimescale = 0.f;
+	float AnimAngularVelocityNoTimescale = 0.f;
 	
 	float RecordingTime = 0.f;
 	float SearchBestCost = 0.f;
