motion matching - deprecating FMotionMatchingState properties ComponentDeltaYaw, ComponentWorldYaw, AnimationDeltaYaw

diff --git a/Engine/Plugins/Animation/PoseSearch/Source/Runtime/Private/AnimNode_MotionMatching.cpp b/Engine/Plugins/Animation/PoseSearch/Source/Runtime/Private/AnimNode_MotionMatching.cpp
--- a/Engine/Plugins/Animation/PoseSearch/Source/Runtime/Private/AnimNode_MotionMatching.cpp
+++ b/Engine/Plugins/Animation/PoseSearch/Source/Runtime/Private/AnimNode_MotionMatching.cpp
@@ -62,6 +62,7 @@ void FAnimNode_MotionMatching::Evaluate_AnyThread(FPoseContext& Output)
 		RootMotionProvider = nullptr;
 	}
 
+PRAGMA_DISABLE_DEPRECATION_WARNINGS
 	// applying MotionMatchingState.ComponentDeltaYaw (considered as root bone delta yaw) to the root bone and the root motion delta transform
 	if (!FMath::IsNearlyZero(MotionMatchingState.ComponentDeltaYaw))
 	{
@@ -84,6 +85,7 @@ void FAnimNode_MotionMatching::Evaluate_AnyThread(FPoseContext& Output)
 	}
 
 	MotionMatchingState.AnimationDeltaYaw = FRotator(RootMotionTransformDelta.GetRotation()).Yaw;
+PRAGMA_ENABLE_DEPRECATION_WARNINGS
 }
 
 void FAnimNode_MotionMatching::UpdateAssetPlayer(const FAnimationUpdateContext& Context)
diff --git a/Engine/Plugins/Animation/PoseSearch/Source/Runtime/Private/PoseSearchLibrary.cpp b/Engine/Plugins/Animation/PoseSearch/Source/Runtime/Private/PoseSearchLibrary.cpp
--- a/Engine/Plugins/Animation/PoseSearch/Source/Runtime/Private/PoseSearchLibrary.cpp
+++ b/Engine/Plugins/Animation/PoseSearch/Source/Runtime/Private/PoseSearchLibrary.cpp
@@ -197,9 +197,12 @@ void FMotionMatchingState::Reset(const FTransform& ComponentTransform)
 	ElapsedPoseSearchTime = std::numeric_limits<float>::infinity();
 	WantedPlayRate = 1.f;
 	bJumpedToPose = false;
+	
+PRAGMA_DISABLE_DEPRECATION_WARNINGS
 	ComponentDeltaYaw = 0.f;
 	ComponentWorldYaw = FRotator(ComponentTransform.GetRotation()).Yaw;
 	AnimationDeltaYaw = 0.f;
+PRAGMA_ENABLE_DEPRECATION_WARNINGS
 
 	PoseIndicesHistory.Reset();
 }
diff --git a/Engine/Plugins/Animation/PoseSearch/Source/Runtime/Public/PoseSearch/PoseSearchLibrary.h b/Engine/Plugins/Animation/PoseSearch/Source/Runtime/Public/PoseSearch/PoseSearchLibrary.h
--- a/Engine/Plugins/Animation/PoseSearch/Source/Runtime/Public/PoseSearch/PoseSearchLibrary.h
+++ b/Engine/Plugins/Animation/PoseSearch/Source/Runtime/Public/PoseSearch/PoseSearchLibrary.h
@@ -75,12 +75,15 @@ PRAGMA_ENABLE_DEPRECATION_WARNINGS
 	UE::PoseSearch::FPoseIndicesHistory PoseIndicesHistory;
 
 	// Component delta yaw (also considered as root bone delta yaw)
+	UE_DEPRECATED(5.4, "Use Steering, OrientationWarping, OffsetRootBone nodes instead")
 	float ComponentDeltaYaw = 0.f;
 
 	// Internal component yaw in world space. Initialized as FRotator(AnimInstanceProxy->GetComponentTransform().GetRotation()).Yaw, but then integrated by ComponentDeltaYaw
+	UE_DEPRECATED(5.4, "Use Steering, OrientationWarping, OffsetRootBone nodes instead")
 	float ComponentWorldYaw = 0.f;
 	
 	// RootMotionTransformDelta yaw at the end of FAnimNode_MotionMatching::Evaluate_AnyThread (it represents the previous frame animation delta yaw)
+	UE_DEPRECATED(5.4, "Use Steering, OrientationWarping, OffsetRootBone nodes instead")
 	float AnimationDeltaYaw = 0.f;
 
 #if UE_POSE_SEARCH_TRACE_ENABLED
