pose search - introducing interaction motion matching

diff --git a/Engine/Plugins/Animation/PoseSearch/Source/Editor/Private/AnimGraphNode_MotionMatchingInteraction.cpp b/Engine/Plugins/Animation/PoseSearch/Source/Editor/Private/AnimGraphNode_MotionMatchingInteraction.cpp
new file mode 100644
--- /dev/null
+++ b/Engine/Plugins/Animation/PoseSearch/Source/Editor/Private/AnimGraphNode_MotionMatchingInteraction.cpp
@@ -0,0 +1,26 @@
+// Copyright Epic Games, Inc. All Rights Reserved.
+
+#include "AnimGraphNode_MotionMatchingInteraction.h"
+
+#define LOCTEXT_NAMESPACE "MotionMatchingInteraction"
+
+UAnimGraphNode_MotionMatchingInteraction::UAnimGraphNode_MotionMatchingInteraction(const FObjectInitializer& ObjectInitializer) : Super(ObjectInitializer)
+{
+}
+
+FText UAnimGraphNode_MotionMatchingInteraction::GetTooltipText() const
+{
+	return LOCTEXT("MotionMatchingInteractionTooltip", "Motion Matching Interaction Node");
+}
+
+FText UAnimGraphNode_MotionMatchingInteraction::GetNodeTitle(ENodeTitleType::Type TitleType) const
+{
+	return LOCTEXT("MotionMatchingInteraction", "MotionMatchingInteraction");
+}
+
+FLinearColor UAnimGraphNode_MotionMatchingInteraction::GetNodeTitleColor() const
+{
+	return FLinearColor(FColor(6, 9, 53));
+}
+
+#undef LOCTEXT_NAMESPACE
diff --git a/Engine/Plugins/Animation/PoseSearch/Source/Editor/Private/PoseSearchAssetDefinitions.cpp b/Engine/Plugins/Animation/PoseSearch/Source/Editor/Private/PoseSearchAssetDefinitions.cpp
--- a/Engine/Plugins/Animation/PoseSearch/Source/Editor/Private/PoseSearchAssetDefinitions.cpp
+++ b/Engine/Plugins/Animation/PoseSearch/Source/Editor/Private/PoseSearchAssetDefinitions.cpp
@@ -2,6 +2,7 @@
 
 #include "PoseSearchAssetDefinitions.h"
 #include "PoseSearchDatabaseEditor.h"
+#include "PoseSearchInteractionAssetEditor.h"
 #include "ThumbnailRendering/SceneThumbnailInfo.h"
 
 namespace UE::PoseSearch
@@ -43,3 +44,21 @@ EAssetCommandResult UAssetDefinition_PoseSearchDatabase::OpenAssets(const FAsset
 	
 	return EAssetCommandResult::Handled;
 }
+
+EAssetCommandResult UAssetDefinition_PoseSearchInteractionAsset::OpenAssets(const FAssetOpenArgs& OpenArgs) const
+{
+	using namespace UE::PoseSearch;
+
+	TArray<UPoseSearchInteractionAsset*> Objects = OpenArgs.LoadObjects<UPoseSearchInteractionAsset>();
+	const EToolkitMode::Type Mode = OpenArgs.ToolkitHost.IsValid() ? EToolkitMode::WorldCentric : EToolkitMode::Standalone;
+	for (UPoseSearchInteractionAsset* InteractionAsset : Objects)
+	{
+		if (InteractionAsset)
+		{
+			TSharedRef<FInteractionAssetEditor> NewEditor(new FInteractionAssetEditor());
+			NewEditor->InitAssetEditor(Mode, OpenArgs.ToolkitHost, InteractionAsset);
+		}
+	}
+	
+	return EAssetCommandResult::Handled;
+}
\ No newline at end of file
diff --git a/Engine/Plugins/Animation/PoseSearch/Source/Editor/Private/PoseSearchAssetDefinitions.h b/Engine/Plugins/Animation/PoseSearch/Source/Editor/Private/PoseSearchAssetDefinitions.h
--- a/Engine/Plugins/Animation/PoseSearch/Source/Editor/Private/PoseSearchAssetDefinitions.h
+++ b/Engine/Plugins/Animation/PoseSearch/Source/Editor/Private/PoseSearchAssetDefinitions.h
@@ -4,6 +4,7 @@
 
 #include "AssetDefinitionDefault.h"
 #include "PoseSearch/PoseSearchDatabase.h"
+#include "PoseSearch/PoseSearchInteractionAsset.h"
 #include "PoseSearch/PoseSearchNormalizationSet.h"
 #include "PoseSearch/PoseSearchSchema.h"
 #include "PoseSearchAssetDefinitions.generated.h"
@@ -58,4 +59,19 @@ public:
 	virtual TSoftClassPtr<UObject> GetAssetClass() const override { return UPoseSearchNormalizationSet::StaticClass(); }
 };
 
+UCLASS(Experimental)
+class UAssetDefinition_PoseSearchInteractionAsset : public UAssetDefinitionDefault
+{
+	GENERATED_BODY()
+
+public:
+	virtual FLinearColor GetAssetColor() const override { return UE::PoseSearch::GetAssetColor(); }
+	virtual TConstArrayView<FAssetCategoryPath> GetAssetCategories() const override { return UE::PoseSearch::GetAssetCategories(); }
+	virtual UThumbnailInfo* LoadThumbnailInfo(const FAssetData& InAssetData) const override { return UE::PoseSearch::LoadThumbnailInfo(InAssetData); }
+
+	virtual FText GetAssetDisplayName() const override { return NSLOCTEXT("PoseSearchInteractionAssetDefinition", "DisplayName_PoseSearchInteractionAsset", "Pose Search Interaction Asset"); }
+	virtual TSoftClassPtr<UObject> GetAssetClass() const override { return UPoseSearchInteractionAsset::StaticClass(); }
+	virtual EAssetCommandResult OpenAssets(const FAssetOpenArgs& OpenArgs) const override;
+};
+
 
diff --git a/Engine/Plugins/Animation/PoseSearch/Source/Editor/Private/PoseSearchDatabaseAssetListItem.cpp b/Engine/Plugins/Animation/PoseSearch/Source/Editor/Private/PoseSearchDatabaseAssetListItem.cpp
--- a/Engine/Plugins/Animation/PoseSearch/Source/Editor/Private/PoseSearchDatabaseAssetListItem.cpp
+++ b/Engine/Plugins/Animation/PoseSearch/Source/Editor/Private/PoseSearchDatabaseAssetListItem.cpp
@@ -266,7 +266,7 @@ namespace UE::PoseSearch
 									PreviewComponent->PreviewInstance->SetPlaying(false);
 									PreviewComponent->PreviewInstance->SetBlendSpacePosition(AnimationAssetBlendParameters);
 								}
-								else if (Editor->GetEditorName() == "ChimeraAssetEditor")
+								else if (Editor->GetEditorName() == "PoseSearchInteractionAssetEditor")
 								{
 									IMultiAnimAssetEditor* MultiAnimAssetEditor = static_cast<IMultiAnimAssetEditor*>(Editor);
 
diff --git a/Engine/Plugins/Animation/PoseSearch/Source/Editor/Private/PoseSearchDatabaseViewModel.cpp b/Engine/Plugins/Animation/PoseSearch/Source/Editor/Private/PoseSearchDatabaseViewModel.cpp
--- a/Engine/Plugins/Animation/PoseSearch/Source/Editor/Private/PoseSearchDatabaseViewModel.cpp
+++ b/Engine/Plugins/Animation/PoseSearch/Source/Editor/Private/PoseSearchDatabaseViewModel.cpp
@@ -37,8 +37,6 @@ static TAutoConsoleVariable<float> CVarDatabasePreviewDebugDrawSamplerTimeOffset
 static TAutoConsoleVariable<float> CVarDatabasePreviewDebugDrawSamplerRootAxisLength(TEXT("a.DatabasePreview.DebugDrawSamplerRootAxisLength"), 0.f, TEXT("Debug Draw Sampler Root Axis Length"));
 #endif
 
-constexpr float StepDeltaTime = 1.f / 30.f;
-
 // FDatabasePreviewActor
 bool FDatabasePreviewActor::SpawnPreviewActor(UWorld* World, const UPoseSearchDatabase* PoseSearchDatabase, int32 IndexAssetIdx, const FRole& Role, const FTransform& SamplerRootTransformOrigin, int32 PoseIdxForTimeOffset)
 {
diff --git a/Engine/Plugins/Animation/PoseSearch/Source/Editor/Private/PoseSearchEditor.cpp b/Engine/Plugins/Animation/PoseSearch/Source/Editor/Private/PoseSearchEditor.cpp
--- a/Engine/Plugins/Animation/PoseSearch/Source/Editor/Private/PoseSearchEditor.cpp
+++ b/Engine/Plugins/Animation/PoseSearch/Source/Editor/Private/PoseSearchEditor.cpp
@@ -5,16 +5,18 @@
 #include "PoseSearchDebugger.h"
 #include "PoseSearchDatabaseEdMode.h"
 #include "PoseSearchDatabaseEditorCommands.h"
+#include "PoseSearchInteractionAssetEditor.h"
 
 #include "Animation/AnimSequence.h"
-#include "Modules/ModuleManager.h"
 #include "AssetToolsModule.h"
 #include "Editor.h"
-#include "PropertyEditorModule.h"
-#include "Subsystems/AssetEditorSubsystem.h"
 #include "IAnimationEditor.h"
 #include "IPersonaToolkit.h"
 #include "IPersonaPreviewScene.h"
+#include "Modules/ModuleManager.h"
+#include "PropertyEditorModule.h"
+#include "Subsystems/AssetEditorSubsystem.h"
+
 #include "Trace/PoseSearchTraceAnalyzer.h"
 #include "Trace/PoseSearchTraceModule.h"
 
@@ -67,8 +69,9 @@ void FEditorModule::StartupModule()
 		IModularFeatures::Get().RegisterModularFeature(FDebuggerTrackCreator::ModularFeatureName, DebuggerTrackCreator.Get());
 		IModularFeatures::Get().RegisterModularFeature(TraceServices::ModuleFeatureName, TraceModule.Get());
 		
-		// Register Ed Mode used by pose search database
+		// Register Ed Modes used by PoseSearchDatabase and PoseSearchInteractionAsset
 		FEditorModeRegistry::Get().RegisterMode<FDatabaseEdMode>(FDatabaseEdMode::EdModeId, LOCTEXT("PoseSearchDatabaseEdModeName", "PoseSearchDatabase"));
+		FEditorModeRegistry::Get().RegisterMode<FInteractionAssetEdMode>(FInteractionAssetEdMode::EdModeId, LOCTEXT("FPoseSearchInteractionAssetEdModeName", "PoseSearchInteractionAsset"));
 	}
 
 	RegisterPropertyTypeCustomizations();
@@ -85,9 +88,10 @@ void FEditorModule::ShutdownModule()
 	}
 	ConsoleCommands.Empty();
 
-	// Unregister Ed Mode
+	// Unregister Ed Modes
 	FEditorModeRegistry::Get().UnregisterMode(FDatabaseEdMode::EdModeId);
-	
+	FEditorModeRegistry::Get().UnregisterMode(FInteractionAssetEdMode::EdModeId);
+
 	UnregisterCustomizations();
 
 	// Unregister Asset Editor Commands
diff --git a/Engine/Plugins/Animation/PoseSearch/Source/Editor/Private/PoseSearchInteractionAssetEditor.cpp b/Engine/Plugins/Animation/PoseSearch/Source/Editor/Private/PoseSearchInteractionAssetEditor.cpp
new file mode 100644
--- /dev/null
+++ b/Engine/Plugins/Animation/PoseSearch/Source/Editor/Private/PoseSearchInteractionAssetEditor.cpp
@@ -0,0 +1,886 @@
+// Copyright Epic Games, Inc. All Rights Reserved.
+
+#include "PoseSearchInteractionAssetEditor.h"
+#include "Animation/DebugSkelMeshComponent.h"
+#include "AnimPreviewInstance.h"
+#include "AssetEditorModeManager.h"
+#include "PoseSearch/PoseSearchInteractionAsset.h"
+#include "PoseSearchEditor.h"
+#include "Components/StaticMeshComponent.h"
+#include "Framework/Application/SlateApplication.h"
+#include "EngineUtils.h"
+#include "GameFramework/WorldSettings.h"
+#include "Modules/ModuleManager.h"
+#include "PreviewProfileController.h"
+#include "PropertyEditorModule.h"
+#include "Styling/AppStyle.h"
+#include "SSimpleTimeSlider.h"
+#include "UnrealWidget.h"
+#include "Viewports.h"
+#include "Widgets/Docking/SDockTab.h"
+#include "Widgets/Images/SImage.h"
+#include "Widgets/Input/SButton.h"
+
+#define LOCTEXT_NAMESPACE "PoseSearchInteractionAssetEditor"
+
+namespace UE::PoseSearch
+{
+
+/////////////////////////////////////////////////
+// class FInteractionAssetEdMode
+bool FInteractionAssetPreviewActor::SpawnPreviewActor(UWorld* World, const UPoseSearchInteractionAsset* InteractionAsset, const UE::PoseSearch::FRole& Role)
+{
+	UAnimationAsset* PreviewAsset = InteractionAsset->GetAnimationAsset(Role);
+	if (!PreviewAsset)
+	{
+		return false;
+	}
+
+	ActorRole = Role;
+
+	FActorSpawnParameters Params;
+	Params.SpawnCollisionHandlingOverride = ESpawnActorCollisionHandlingMethod::AdjustIfPossibleButAlwaysSpawn;
+
+	for (int32 ActorIndex = 0; ActorIndex < NumActors; ++ActorIndex)
+	{
+		UE::PoseSearch::FAnimationAssetSampler& Sampler = Samplers[ActorIndex];
+
+		const FTransform Origin = InteractionAsset->GetDebugWarpOrigin(GetRole(), ActorIndex == DebugActor);
+		Sampler.Init(PreviewAsset, Origin, BlendParameters);
+
+		const FTransform ActorTransform = Sampler.ExtractRootTransform(CurrentTime);
+
+		TWeakObjectPtr<AActor>& ActorPtr = ActorPtrs[ActorIndex];
+
+		ActorPtr = World->SpawnActor<AActor>(AActor::StaticClass(), FTransform::Identity, Params);
+		ActorPtr->SetFlags(RF_Transient);
+
+		UDebugSkelMeshComponent* Mesh = NewObject<UDebugSkelMeshComponent>(ActorPtr.Get());
+		Mesh->RegisterComponentWithWorld(World);
+
+		UAnimPreviewInstance* AnimInstance = NewObject<UAnimPreviewInstance>(Mesh);
+		Mesh->PreviewInstance = AnimInstance;
+		AnimInstance->InitializeAnimation();
+
+		USkeleton* Skeleton = PreviewAsset->GetSkeleton();
+		Mesh->SetSkeletalMesh(Skeleton->GetPreviewMesh(true));
+		Mesh->EnablePreview(true, PreviewAsset);
+
+		AnimInstance->SetAnimationAsset(PreviewAsset, InteractionAsset->IsLooping(), 0.f);
+		AnimInstance->SetBlendSpacePosition(BlendParameters);
+
+		AnimInstance->PlayAnim(InteractionAsset->IsLooping(), 0.f);
+		if (!ActorPtr->GetRootComponent())
+		{
+			ActorPtr->SetRootComponent(Mesh);
+		}
+
+		AnimInstance->SetPosition(Sampler.ToNormalizedTime(CurrentTime));
+		AnimInstance->SetPlayRate(0.f);
+		AnimInstance->SetBlendSpacePosition(BlendParameters);
+
+		ActorPtr->SetActorTransform(ActorTransform);
+
+		UE_LOG(LogPoseSearchEditor, Log, TEXT("Spawned preview Actor: %s"), *GetNameSafe(ActorPtr.Get()));
+	}
+	return true;
+}
+
+void FInteractionAssetPreviewActor::UpdatePreviewActor(const UPoseSearchInteractionAsset* InteractionAsset, float PlayTime)
+{
+	if (UAnimationAsset* PreviewAsset = InteractionAsset->GetAnimationAsset(GetRole()))
+	{
+		bool bPlayTimeUpdated = false;
+
+		float NewCurrentTime = 0.f;
+		FAnimationRuntime::AdvanceTime(false, PlayTime, NewCurrentTime, PreviewAsset->GetPlayLength());
+
+		if (!FMath::IsNearlyEqual(CurrentTime, NewCurrentTime))
+		{
+			CurrentTime = NewCurrentTime;
+			bPlayTimeUpdated = true;
+		}
+
+		for (int32 ActorIndex = 0; ActorIndex < NumActors; ++ActorIndex)
+		{
+			const FTransform Origin = InteractionAsset->GetDebugWarpOrigin(GetRole(), ActorIndex == DebugActor);
+
+			bool bSamplerReinitialized = false;
+			UE::PoseSearch::FAnimationAssetSampler& Sampler = Samplers[ActorIndex];
+			if (PreviewAsset != Sampler.GetAsset() || !Origin.Equals(Sampler.GetRootTransformOrigin()))
+			{
+				// reinitializing the Sampler if the PreviewAsset or the origin transform changed
+				Sampler.Init(PreviewAsset, Origin, BlendParameters);
+				bSamplerReinitialized = true;
+			}
+
+			TWeakObjectPtr<AActor>& ActorPtr = ActorPtrs[ActorIndex];
+			if (ActorPtr != nullptr)
+			{
+				if (UDebugSkelMeshComponent* Mesh = Cast<UDebugSkelMeshComponent>(ActorPtr->GetRootComponent()))
+				{
+					if (UAnimPreviewInstance* AnimInstance = Mesh->PreviewInstance.Get())
+					{
+						bool bPreviewAssetChanged = false;
+						if (AnimInstance->GetAnimationAsset() != PreviewAsset)
+						{
+							AnimInstance->SetAnimationAsset(PreviewAsset, InteractionAsset->IsLooping(), 0.f);
+							bPreviewAssetChanged = true;
+						}
+
+						if (bPlayTimeUpdated || bSamplerReinitialized || bPreviewAssetChanged)
+						{
+							// SetPosition is in [0..1] range for blendspaces
+							AnimInstance->SetPosition(Sampler.ToNormalizedTime(CurrentTime));
+							AnimInstance->SetPlayRate(0.f);
+							AnimInstance->SetBlendSpacePosition(BlendParameters);
+
+							const FTransform ActorTransform = Sampler.ExtractRootTransform(CurrentTime);
+							ActorPtr->SetActorTransform(ActorTransform);
+						}
+					}
+				}
+			}
+		}
+	}
+}
+
+void FInteractionAssetPreviewActor::Destroy()
+{
+	for (TWeakObjectPtr<AActor>& ActorPtr : ActorPtrs)
+	{
+		if (ActorPtr != nullptr)
+		{
+			ActorPtr->Destroy();
+		}
+		ActorPtr = nullptr;
+	}
+}
+
+
+UAnimPreviewInstance* FInteractionAssetPreviewActor::GetAnimPreviewInstance()
+{
+	if (const AActor* Actor = ActorPtrs[DebugActor].Get())
+	{
+		if (const UDebugSkelMeshComponent* Mesh = Cast<UDebugSkelMeshComponent>(Actor->GetRootComponent()))
+		{
+			return Mesh->PreviewInstance.Get();
+		}
+	}
+	
+	return nullptr;
+}
+
+FTransform FInteractionAssetPreviewActor::GetDebugActorTransformFromSampler() const
+{
+	const UE::PoseSearch::FAnimationAssetSampler& Sampler = Samplers[DebugActor];
+	const FTransform ActorTransform = Sampler.ExtractRootTransform(CurrentTime);
+	return ActorTransform;
+}
+
+void FInteractionAssetPreviewActor::ForceDebugActorTransform(const FTransform& ActorTransform)
+{
+	TWeakObjectPtr<AActor>& ActorPtr = ActorPtrs[DebugActor];
+	if (ActorPtr != nullptr)
+	{
+		ActorPtr->SetActorTransform(ActorTransform);
+	}
+}
+
+/////////////////////////////////////////////////
+// class FInteractionAssetViewModel
+void FInteractionAssetViewModel::AddReferencedObjects(FReferenceCollector& Collector)
+{
+	Collector.AddReferencedObject(InteractionAssetPtr);
+}
+
+void FInteractionAssetViewModel::Initialize(UPoseSearchInteractionAsset* InteractionAsset, const TSharedRef<FInteractionAssetPreviewScene>& PreviewScene)
+{
+	InteractionAssetPtr = InteractionAsset;
+	PreviewScenePtr = PreviewScene;
+}
+
+void FInteractionAssetViewModel::PreviewBackwardEnd()
+{
+	SetPlayTime(0.f, false);
+}
+
+void FInteractionAssetViewModel::PreviewBackwardStep()
+{
+	if (const UPoseSearchInteractionAsset* InteractionAsset = GetInteractionAsset())
+	{
+		const float NewPlayTime = FMath::Clamp(PlayTime - StepDeltaTime, 0.f, InteractionAsset->GetPlayLength());
+		SetPlayTime(NewPlayTime, false);
+	}
+}
+
+void FInteractionAssetViewModel::PreviewBackward()
+{
+	DeltaTimeMultiplier = -1.f;
+}
+
+void FInteractionAssetViewModel::PreviewPause()
+{
+	DeltaTimeMultiplier = 0.f;
+}
+
+void FInteractionAssetViewModel::PreviewForward()
+{
+	DeltaTimeMultiplier = 1.f;
+}
+
+void FInteractionAssetViewModel::PreviewForwardStep()
+{
+	if (const UPoseSearchInteractionAsset* InteractionAsset = GetInteractionAsset())
+	{
+		const float NewPlayTime = FMath::Clamp(PlayTime + StepDeltaTime, 0.f, InteractionAsset->GetPlayLength());
+		SetPlayTime(NewPlayTime, false);
+	}
+}
+
+void FInteractionAssetViewModel::PreviewForwardEnd()
+{
+	if (const UPoseSearchInteractionAsset* InteractionAsset = GetInteractionAsset())
+	{
+		SetPlayTime(InteractionAsset->GetPlayLength(), false);
+	}
+}
+
+UWorld* FInteractionAssetViewModel::GetWorld()
+{
+	check(PreviewScenePtr.IsValid());
+	return PreviewScenePtr.Pin()->GetWorld();
+}
+
+void FInteractionAssetViewModel::Tick(float DeltaSeconds)
+{
+	const UPoseSearchInteractionAsset* InteractionAsset = GetInteractionAsset();
+	if (!InteractionAsset)
+	{
+		RemovePreviewActors();
+		return;
+	}
+
+	PlayTime += DeltaSeconds * DeltaTimeMultiplier;
+	PlayTime = FMath::Clamp(PlayTime, 0.f, InteractionAsset->GetPlayLength());
+
+	UE::PoseSearch::FRoleToIndex InteractionAssetRoleToIndex;
+	for (int32 RoleIndex = 0; RoleIndex < InteractionAsset->GetNumRoles(); ++RoleIndex)
+	{
+		InteractionAssetRoleToIndex.Add(InteractionAsset->GetRole(RoleIndex)) = RoleIndex;
+	}
+
+	// iterating backwards because of the possible RemoveAtSwap 
+	UE::PoseSearch::FRoleToIndex PreviewActorsRoleToIndex;
+	for (int32 ActorIndex = PreviewActors.Num() - 1; ActorIndex >= 0; --ActorIndex)
+	{
+		FInteractionAssetPreviewActor& PreviewActor = PreviewActors[ActorIndex];
+		if (!InteractionAssetRoleToIndex.Find(PreviewActor.GetRole()))
+		{
+			PreviewActor.Destroy();
+			PreviewActors.RemoveAtSwap(ActorIndex, EAllowShrinking::No);
+		}
+		else
+		{
+			PreviewActorsRoleToIndex.Add(PreviewActor.GetRole());		
+		}
+	}
+
+	if (PreviewActors.Num() != InteractionAssetRoleToIndex.Num())
+	{
+		for (int32 RoleIndex = 0; RoleIndex < InteractionAsset->GetNumRoles(); ++RoleIndex)
+		{
+			if (!PreviewActorsRoleToIndex.Find(InteractionAsset->GetRole(RoleIndex)))
+			{
+				FInteractionAssetPreviewActor PreviewActor;
+				if (PreviewActor.SpawnPreviewActor(GetWorld(), InteractionAsset, InteractionAsset->GetRole(RoleIndex)))
+				{
+					PreviewActors.Add(PreviewActor);
+				}
+			}
+		}
+	}
+
+	for (FInteractionAssetPreviewActor& PreviewActor : PreviewActors)
+	{
+		PreviewActor.UpdatePreviewActor(InteractionAsset, PlayTime);
+	}
+	
+	UE::PoseSearch::FRoleToIndex PreviewActorRoleToIndex;
+	PreviewActorRoleToIndex.Reserve(PreviewActors.Num());
+	for (int32 PreviewActorIndex = 0; PreviewActorIndex < PreviewActors.Num(); ++PreviewActorIndex)
+	{
+		PreviewActorRoleToIndex.Add(PreviewActors[PreviewActorIndex].GetRole()) = PreviewActorIndex;
+	}
+
+	// testing CalculateWarpTransforms
+#if WITH_EDITORONLY_DATA
+	if (InteractionAsset->bEnableDebugWarp && PreviewActors.Num() == InteractionAsset->GetNumRoles())
+	{
+		TArray<FTransform, TInlineAllocator<UE::PoseSearch::PreallocatedRolesNum>> ActorTransforms;
+		ActorTransforms.SetNum(InteractionAsset->GetNumRoles());
+		for (int32 InteractionAssetRoleIndex = 0; InteractionAssetRoleIndex < InteractionAsset->GetNumRoles(); ++InteractionAssetRoleIndex)
+		{
+			const UE::PoseSearch::FRole& InteractionAssetRole = InteractionAsset->GetRole(InteractionAssetRoleIndex);
+			const int32 PreviewActorIndex = PreviewActorRoleToIndex[InteractionAssetRole];
+
+			FInteractionAssetPreviewActor& PreviewActor = PreviewActors[PreviewActorIndex];
+			check(PreviewActor.GetRole() == InteractionAssetRole);
+
+			ActorTransforms[InteractionAssetRoleIndex] = PreviewActor.GetDebugActorTransformFromSampler();
+		}
+
+		TArray<FTransform, TInlineAllocator<UE::PoseSearch::PreallocatedRolesNum>> FullAlignedActorTransforms;
+		FullAlignedActorTransforms.SetNum(InteractionAsset->GetNumRoles());
+		InteractionAsset->CalculateWarpTransforms(PlayTime, ActorTransforms, FullAlignedActorTransforms);
+
+		for (int32 InteractionAssetRoleIndex = 0; InteractionAssetRoleIndex < InteractionAsset->GetNumRoles(); ++InteractionAssetRoleIndex)
+		{
+			const UE::PoseSearch::FRole& InteractionAssetRole = InteractionAsset->GetRole(InteractionAssetRoleIndex);
+			const int32 PreviewActorIndex = PreviewActorRoleToIndex[InteractionAssetRole];
+
+			FInteractionAssetPreviewActor& PreviewActor = PreviewActors[PreviewActorIndex];
+			check(PreviewActor.GetRole() == InteractionAssetRole);
+			FTransform DebugActorTransform;
+			DebugActorTransform.Blend(ActorTransforms[InteractionAssetRoleIndex], FullAlignedActorTransforms[InteractionAssetRoleIndex], InteractionAsset->DebugWarpAmount);
+			PreviewActor.ForceDebugActorTransform(DebugActorTransform);
+		}
+	}
+#endif // WITH_EDITORONLY_DATA
+}
+
+void FInteractionAssetViewModel::RemovePreviewActors()
+{
+	PlayTime = 0.f;
+	DeltaTimeMultiplier = 1.f;
+
+	for (FInteractionAssetPreviewActor& PreviewActor : PreviewActors)
+	{
+		PreviewActor.Destroy();
+	}
+
+	PreviewActors.Reset();
+}
+
+TRange<double> FInteractionAssetViewModel::GetPreviewPlayRange() const
+{
+	constexpr double ViewRangeSlack = 0.2;
+	if (const UPoseSearchInteractionAsset* InteractionAsset = GetInteractionAsset())
+	{
+		return TRange<double>(-ViewRangeSlack, InteractionAsset->GetPlayLength() + ViewRangeSlack);
+	}
+	
+	return TRange<double>(-ViewRangeSlack, ViewRangeSlack);
+}
+
+void FInteractionAssetViewModel::SetPlayTime(float NewPlayTime, bool bInTickPlayTime)
+{
+	if (const UPoseSearchInteractionAsset* InteractionAsset = GetInteractionAsset())
+	{
+		NewPlayTime = FMath::Clamp(NewPlayTime, 0.f, InteractionAsset->GetPlayLength());
+		DeltaTimeMultiplier = bInTickPlayTime ? DeltaTimeMultiplier : 0.f;
+
+		if (!FMath::IsNearlyEqual(PlayTime, NewPlayTime))
+		{
+			PlayTime = NewPlayTime;
+
+			for (FInteractionAssetPreviewActor& PreviewActor : PreviewActors)
+			{
+				PreviewActor.UpdatePreviewActor(InteractionAsset, PlayTime);
+			}
+		}
+	}
+}
+
+void FInteractionAssetViewModel::SetPreviewProperties(float AnimAssetTime, const FVector& AnimAssetBlendParameters, bool bAnimAssetPlaying)
+{
+	// @todo: add support for blend spaces (pass AnimAssetBlendParameters as input)
+	SetPlayTime(AnimAssetTime, bAnimAssetPlaying);
+}
+
+/////////////////////////////////////////////////
+// class FInteractionAssetEdMode
+const FEditorModeID FInteractionAssetEdMode::EdModeId = TEXT("PoseSearchInteractionAssetEdMode");
+
+void FInteractionAssetEdMode::Tick(FEditorViewportClient* ViewportClient, float DeltaTime)
+{
+	FEdMode::Tick(ViewportClient, DeltaTime);
+
+	if (FInteractionAssetViewportClient* InteractionAssetViewportClient = static_cast<FInteractionAssetViewportClient*>(ViewportClient))
+	{
+		// ensure we redraw even if PIE is active
+		InteractionAssetViewportClient->Invalidate();
+
+		if (!ViewModel)
+		{
+			ViewModel = InteractionAssetViewportClient->GetAssetEditor()->GetViewModel();
+		}
+	}
+
+	if (ViewModel)
+	{
+		ViewModel->Tick(DeltaTime);
+	}
+}
+
+void FInteractionAssetEdMode::Render(const FSceneView* View, FViewport* Viewport, FPrimitiveDrawInterface* PDI)
+{
+	FEdMode::Render(View, Viewport, PDI);
+}
+
+bool FInteractionAssetEdMode::AllowWidgetMove()
+{
+	return FEdMode::ShouldDrawWidget();
+}
+
+bool FInteractionAssetEdMode::ShouldDrawWidget() const
+{
+	return FEdMode::ShouldDrawWidget();
+}
+
+bool FInteractionAssetEdMode::GetCustomInputCoordinateSystem(FMatrix& InMatrix, void* InData)
+{
+	return FEdMode::GetCustomDrawingCoordinateSystem(InMatrix, InData);
+}
+
+/////////////////////////////////////////////////
+// class FInteractionAssetViewportClient
+FInteractionAssetViewportClient::FInteractionAssetViewportClient(
+	const TSharedRef<FInteractionAssetPreviewScene>& InPreviewScene,
+	const TSharedRef<SInteractionAssetViewport>& InViewport,
+	const TSharedRef<FInteractionAssetEditor>& InAssetEditor)
+	: FEditorViewportClient(nullptr, &InPreviewScene.Get(), StaticCastSharedRef<SEditorViewport>(InViewport))
+	, PreviewScenePtr(InPreviewScene)
+	, AssetEditorPtr(InAssetEditor)
+{
+	Widget->SetUsesEditorModeTools(ModeTools.Get());
+	StaticCastSharedPtr<FAssetEditorModeManager>(ModeTools)->SetPreviewScene(&InPreviewScene.Get());
+	ModeTools->SetDefaultMode(FInteractionAssetEdMode::EdModeId);
+
+	SetRealtime(true);
+
+	SetWidgetCoordSystemSpace(COORD_Local);
+	ModeTools->SetWidgetMode(UE::Widget::WM_Translate);
+}
+
+void FInteractionAssetViewportClient::TrackingStarted(
+	const struct FInputEventState& InInputState,
+	bool bIsDraggingWidget,
+	bool bNudge)
+{
+	ModeTools->StartTracking(this, Viewport);
+}
+
+void FInteractionAssetViewportClient::TrackingStopped()
+{
+	ModeTools->EndTracking(this, Viewport);
+	Invalidate();
+}
+
+void FInteractionAssetViewportClient::Draw(const FSceneView* View, FPrimitiveDrawInterface* PDI)
+{
+	FEditorViewportClient::Draw(View, PDI);
+}
+
+/////////////////////////////////////////////////
+// class FInteractionAssetPreviewScene
+FInteractionAssetPreviewScene::FInteractionAssetPreviewScene(ConstructionValues CVs, const TSharedRef<FInteractionAssetEditor>& Editor)
+: FAdvancedPreviewScene(CVs)
+, EditorPtr(Editor)
+{
+	// Disable killing actors outside of the world
+	AWorldSettings* WorldSettings = GetWorld()->GetWorldSettings(true);
+	WorldSettings->bEnableWorldBoundsChecks = false;
+
+	// Spawn an owner for FloorMeshComponent so CharacterMovementComponent can detect it as a valid floor and slide 
+	// along it
+	{
+		AActor* FloorActor = GetWorld()->SpawnActor<AActor>(AActor::StaticClass(), FTransform());
+		check(FloorActor);
+
+		static const FString NewName = FString(TEXT("FloorComponent"));
+		FloorMeshComponent->Rename(*NewName, FloorActor);
+
+		FloorActor->SetRootComponent(FloorMeshComponent);
+	}
+}
+
+void FInteractionAssetPreviewScene::Tick(float InDeltaTime)
+{
+	FAdvancedPreviewScene::Tick(InDeltaTime);
+
+	// Trigger Begin Play in this preview world.
+	// This is needed for the CharacterMovementComponent to be able to switch to falling mode. 
+	// See: UCharacterMovementComponent::StartFalling
+	if (PreviewWorld && !PreviewWorld->GetBegunPlay())
+	{
+		for (FActorIterator It(PreviewWorld); It; ++It)
+		{
+			It->DispatchBeginPlay();
+		}
+
+		PreviewWorld->SetBegunPlay(true);
+	}
+
+	GetWorld()->Tick(LEVELTICK_All, InDeltaTime);
+}
+
+/////////////////////////////////////////////////
+// class SInteractionAssetViewport
+void SInteractionAssetViewport::Construct(
+	const FArguments& InArgs,
+	const FInteractionAssetPreviewRequiredArgs& InRequiredArgs)
+{
+	PreviewScenePtr = InRequiredArgs.PreviewScene;
+	AssetEditorPtr = InRequiredArgs.AssetEditor;
+
+	SEditorViewport::Construct(
+		SEditorViewport::FArguments()
+		.IsEnabled(FSlateApplication::Get().GetNormalExecutionAttribute())
+		.AddMetaData<FTagMetaData>(TEXT("AnimationTools.Viewport"))
+	);
+}
+
+void SInteractionAssetViewport::BindCommands()
+{
+	SEditorViewport::BindCommands();
+}
+
+TSharedRef<FEditorViewportClient> SInteractionAssetViewport::MakeEditorViewportClient()
+{
+	ViewportClient = MakeShared<FInteractionAssetViewportClient>(
+		PreviewScenePtr.Pin().ToSharedRef(),
+		SharedThis(this),
+		AssetEditorPtr.Pin().ToSharedRef());
+	ViewportClient->ViewportType = LVT_Perspective;
+	ViewportClient->bSetListenerPosition = false;
+	ViewportClient->SetViewLocation(EditorViewportDefs::DefaultPerspectiveViewLocation);
+	ViewportClient->SetViewRotation(EditorViewportDefs::DefaultPerspectiveViewRotation);
+
+	return ViewportClient.ToSharedRef();
+}
+
+TSharedPtr<SWidget> SInteractionAssetViewport::MakeViewportToolbar()
+{
+	return SAssignNew(ViewportToolbar, SCommonEditorViewportToolbarBase, SharedThis(this));
+}
+
+TSharedRef<SEditorViewport> SInteractionAssetViewport::GetViewportWidget()
+{
+	return SharedThis(this);
+}
+
+TSharedPtr<FExtender> SInteractionAssetViewport::GetExtenders() const
+{
+	TSharedPtr<FExtender> Result(MakeShareable(new FExtender));
+	return Result;
+}
+
+void SInteractionAssetViewport::OnFloatingButtonClicked()
+{
+}
+
+/////////////////////////////////////////////////
+// class SInteractionAssetPreview
+void SInteractionAssetPreview::Construct(const FArguments& InArgs, const FInteractionAssetPreviewRequiredArgs& InRequiredArgs)
+{
+	SliderColor = InArgs._SliderColor;
+	SliderScrubTime = InArgs._SliderScrubTime;
+	SliderViewRange = InArgs._SliderViewRange;
+	OnSliderScrubPositionChanged = InArgs._OnSliderScrubPositionChanged;
+
+	OnBackwardEnd = InArgs._OnBackwardEnd;
+	OnBackwardStep = InArgs._OnBackwardStep;
+	OnBackward = InArgs._OnBackward;
+	OnPause = InArgs._OnPause;
+	OnForward = InArgs._OnForward;
+	OnForwardStep = InArgs._OnForwardStep;
+	OnForwardEnd = InArgs._OnForwardEnd;
+
+	FSlimHorizontalToolBarBuilder ToolBarBuilder(
+		TSharedPtr<const FUICommandList>(), 
+		FMultiBoxCustomization::None, 
+		nullptr, true);
+
+	auto AddToolBarButton = [&ToolBarBuilder](FName ButtonImageName, FOnButtonClickedEvent& OnClicked)
+		{
+			ToolBarBuilder.AddToolBarWidget(
+				SNew(SButton)
+				.ButtonStyle(FAppStyle::Get(), "Animation.PlayControlsButton")
+				.OnClicked_Lambda([&OnClicked]()
+					{
+						if (OnClicked.IsBound())
+						{
+							OnClicked.Execute();
+							return FReply::Handled();
+						}
+						return FReply::Unhandled();
+					})
+				[
+					SNew(SImage)
+						.ColorAndOpacity(FSlateColor::UseSubduedForeground())
+						.Image(FAppStyle::Get().GetBrush(ButtonImageName))
+				]);
+		};
+
+	//ToolBarBuilder.SetStyle(&FAppStyle::Get(), "PaletteToolBar");
+	ToolBarBuilder.BeginSection("Preview");
+	{
+		AddToolBarButton("Animation.Backward_End", OnBackwardEnd);
+		AddToolBarButton("Animation.Backward_Step", OnBackwardStep);
+		AddToolBarButton("Animation.Backward", OnBackward);
+		AddToolBarButton("Animation.Pause", OnPause);
+		AddToolBarButton("Animation.Forward", OnForward);
+		AddToolBarButton("Animation.Forward_Step", OnForwardStep);
+		AddToolBarButton("Animation.Forward_End", OnForwardEnd);
+	}
+
+	ChildSlot
+	[
+		SNew(SVerticalBox)
+		+SVerticalBox::Slot()
+		.FillHeight(1.0f)
+		[
+			SNew(SInteractionAssetViewport, InRequiredArgs)
+		]
+		+SVerticalBox::Slot()
+		.AutoHeight()
+		[
+			SNew(SHorizontalBox)
+			+SHorizontalBox::Slot()
+			.AutoWidth()
+			[
+				ToolBarBuilder.MakeWidget()
+			]
+			+SHorizontalBox::Slot()
+			.FillWidth(1.0f)
+			[
+				SNew(SSimpleTimeSlider)
+				.ClampRangeHighlightSize(0.15f)
+				.ClampRangeHighlightColor_Lambda([this]()
+					{
+						return SliderColor.Get();
+					})
+				.ScrubPosition_Lambda([this]()
+					{
+						return SliderScrubTime.Get();
+					})
+				.ViewRange_Lambda([this]()
+					{
+						return SliderViewRange.Get();
+					})
+				.ClampRange_Lambda([this]()
+					{
+						return SliderViewRange.Get();
+					})
+				.OnScrubPositionChanged_Lambda([this](double NewScrubTime, bool bIsScrubbing)
+					{
+						if (bIsScrubbing)
+						{
+							OnSliderScrubPositionChanged.ExecuteIfBound(NewScrubTime, bIsScrubbing);
+						}
+					})
+			]
+		]
+	];
+}
+
+/////////////////////////////////////////////////
+// class FInteractionAssetEditor
+const FName PoseSearchInteractionAssetEditorAppName = FName(TEXT("PoseSearchInteractionAssetEditorApp"));
+
+// Tab identifiers
+struct FInteractionAssetEditorTabs
+{
+	static const FName AssetDetailsID;
+	static const FName ViewportID;
+};
+const FName FInteractionAssetEditorTabs::AssetDetailsID(TEXT("PoseSearchInteractionAssetEditorAssetDetailsTabID"));
+const FName FInteractionAssetEditorTabs::ViewportID(TEXT("PoseSearchInteractionAssetEditorViewportTabID"));
+
+void FInteractionAssetEditor::InitAssetEditor(const EToolkitMode::Type Mode, const TSharedPtr<IToolkitHost>& InitToolkitHost, UPoseSearchInteractionAsset* InteractionAsset)
+{
+	// Create Preview Scene
+	if (!PreviewScene.IsValid())
+	{
+		PreviewScene = MakeShareable(
+				new FInteractionAssetPreviewScene(
+					FPreviewScene::ConstructionValues()
+					.SetCreatePhysicsScene(false)
+					.SetTransactional(false)
+					.ForceUseMovementComponentInNonGameWorld(true),
+					StaticCastSharedRef<FInteractionAssetEditor>(AsShared())));
+
+		//Temporary fix for missing attached assets - MDW (Copied from FPersonaToolkit::CreatePreviewScene)
+		PreviewScene->GetWorld()->GetWorldSettings()->SetIsTemporarilyHiddenInEditor(false);
+	}
+
+	// Create view model
+	ViewModel = MakeShared<FInteractionAssetViewModel>();
+	ViewModel->Initialize(InteractionAsset, PreviewScene.ToSharedRef());
+
+	// Create viewport widget
+	{
+		FInteractionAssetPreviewRequiredArgs PreviewArgs(
+			StaticCastSharedRef<FInteractionAssetEditor>(AsShared()),
+			PreviewScene.ToSharedRef());
+			
+		PreviewWidget = SNew(SInteractionAssetPreview, PreviewArgs)
+			.SliderColor(FLinearColor::Red)
+			.SliderScrubTime_Lambda([this]()
+				{
+					return ViewModel->GetPlayTime();
+				})
+			.SliderViewRange_Lambda([this]() 
+				{ 
+					return ViewModel->GetPreviewPlayRange();
+				})
+			.OnSliderScrubPositionChanged_Lambda([this](float NewScrubPosition, bool bScrubbing)
+				{
+					ViewModel->SetPlayTime(NewScrubPosition, !bScrubbing);
+				})
+			.OnBackwardEnd_Raw(this, &FInteractionAssetEditor::PreviewBackwardEnd)
+			.OnBackwardStep_Raw(this, &FInteractionAssetEditor::PreviewBackwardStep)
+			.OnBackward_Raw(this, &FInteractionAssetEditor::PreviewBackward)
+			.OnPause_Raw(this, &FInteractionAssetEditor::PreviewPause)
+			.OnForward_Raw(this, &FInteractionAssetEditor::PreviewForward)
+			.OnForwardStep_Raw(this, &FInteractionAssetEditor::PreviewForwardStep)
+			.OnForwardEnd_Raw(this, &FInteractionAssetEditor::PreviewForwardEnd);
+	}
+
+	// asset details widget
+	FDetailsViewArgs AssetDetailsArgs;
+	AssetDetailsArgs.bHideSelectionTip = true;
+	AssetDetailsArgs.NotifyHook = this;
+
+	FPropertyEditorModule& PropertyModule = FModuleManager::LoadModuleChecked<FPropertyEditorModule>("PropertyEditor");
+	EditingAssetWidget = PropertyModule.CreateDetailView(AssetDetailsArgs);
+	EditingAssetWidget->SetObject(InteractionAsset);
+
+	// Define Editor Layout
+	const TSharedRef<FTabManager::FLayout> StandaloneDefaultLayout =
+	FTabManager::NewLayout("Standalone_PoseSearchInteractionAssetDatabaseEditor_Layout_v0.01")
+		->AddArea
+		(
+			// Main application area
+			FTabManager::NewPrimaryArea()->SetOrientation(Orient_Horizontal)
+			->Split
+			(
+				FTabManager::NewStack()
+					->SetSizeCoefficient(0.4f)
+					->AddTab(FInteractionAssetEditorTabs::AssetDetailsID, ETabState::OpenedTab)
+					->SetHideTabWell(false)
+			)
+			->Split
+			(
+				FTabManager::NewStack()
+					->SetSizeCoefficient(0.6f)
+					->AddTab(FInteractionAssetEditorTabs::ViewportID, ETabState::OpenedTab)
+					->SetHideTabWell(false)
+			)
+		);
+
+	FAssetEditorToolkit::InitAssetEditor(
+		Mode,
+		InitToolkitHost,
+		PoseSearchInteractionAssetEditorAppName,
+		StandaloneDefaultLayout,
+		true,
+		true,
+		InteractionAsset,
+		false);
+
+	RegenerateMenusAndToolbars();
+}
+
+void FInteractionAssetEditor::RegisterTabSpawners(const TSharedRef<FTabManager>& InTabManager)
+{
+	WorkspaceMenuCategory = InTabManager->AddLocalWorkspaceMenuCategory(
+		LOCTEXT("WorkspaceMenu_PoseSearchInteractionAssetEditor", "Pose Search Interaction Asset Editor"));
+	TSharedRef<FWorkspaceItem> WorkspaceMenuCategoryRef = WorkspaceMenuCategory.ToSharedRef();
+
+	FAssetEditorToolkit::RegisterTabSpawners(InTabManager);
+
+	InTabManager->RegisterTabSpawner(
+		FInteractionAssetEditorTabs::ViewportID,
+		FOnSpawnTab::CreateSP(this, &FInteractionAssetEditor::SpawnTab_Viewport))
+		.SetDisplayName(LOCTEXT("ViewportTab", "Viewport"))
+		.SetGroup(WorkspaceMenuCategoryRef)
+		.SetIcon(FSlateIcon(FAppStyle::GetAppStyleSetName(), "GraphEditor.EventGraph_16x"));
+			
+	InTabManager->RegisterTabSpawner(
+		FInteractionAssetEditorTabs::AssetDetailsID,
+		FOnSpawnTab::CreateSP(this, &FInteractionAssetEditor::SpawnTab_AssetDetails))
+		.SetDisplayName(LOCTEXT("PoseSearchInteractionAssetDetailsTab", "Pose Search Interaction Asset Details"))
+		.SetGroup(WorkspaceMenuCategoryRef)
+		.SetIcon(FSlateIcon(FAppStyle::GetAppStyleSetName(), "LevelEditor.Tabs.Details"));
+}
+
+void FInteractionAssetEditor::UnregisterTabSpawners(const TSharedRef<FTabManager>& InTabManager)
+{
+	FAssetEditorToolkit::UnregisterTabSpawners(InTabManager);
+
+	InTabManager->UnregisterTabSpawner(FInteractionAssetEditorTabs::ViewportID);
+	InTabManager->UnregisterTabSpawner(FInteractionAssetEditorTabs::AssetDetailsID);
+}
+
+FName FInteractionAssetEditor::GetToolkitFName() const
+{
+	return FName("PoseSearchInteractionAssetEditor");
+}
+
+FText FInteractionAssetEditor::GetBaseToolkitName() const
+{
+	return LOCTEXT("PoseSearchInteractionAssetEditorAppLabel", "Pose Search Interaction Asset Editor");
+}
+
+FText FInteractionAssetEditor::GetToolkitName() const
+{
+	FFormatNamedArguments Args;
+	Args.Add(TEXT("AssetName"), FText::FromString(GetNameSafe(GetInteractionAsset())));
+	return FText::Format(LOCTEXT("PoseSearchInteractionAssetEditorToolkitName", "{AssetName}"), Args);
+}
+
+FLinearColor FInteractionAssetEditor::GetWorldCentricTabColorScale() const
+{
+	return FLinearColor::White;
+}
+
+FString FInteractionAssetEditor::GetWorldCentricTabPrefix() const
+{
+	return TEXT("PoseSearchInteractionAssetEditor");
+}
+
+void FInteractionAssetEditor::SetPreviewProperties(float AnimAssetTime, const FVector& AnimAssetBlendParameters, bool bAnimAssetPlaying)
+{
+	ViewModel->SetPreviewProperties(AnimAssetTime, AnimAssetBlendParameters, bAnimAssetPlaying);
+}
+
+TSharedRef<SDockTab> FInteractionAssetEditor::SpawnTab_Viewport(const FSpawnTabArgs& Args)
+{
+	check(Args.GetTabId() == FInteractionAssetEditorTabs::ViewportID);
+
+	TSharedRef<SDockTab> SpawnedTab = SNew(SDockTab).Label(LOCTEXT("ViewportTab_Title", "Viewport"));
+
+	if (PreviewWidget.IsValid())
+	{
+		SpawnedTab->SetContent(PreviewWidget.ToSharedRef());
+	}
+
+	return SpawnedTab;
+}
+
+TSharedRef<SDockTab> FInteractionAssetEditor::SpawnTab_AssetDetails(const FSpawnTabArgs& Args)
+{
+	check(Args.GetTabId() == FInteractionAssetEditorTabs::AssetDetailsID);
+
+	return SNew(SDockTab)
+		.Label(LOCTEXT("PoseSearchInteractionAsset_Details_Title", "Pose Search Interaction Asset Details"))
+		[
+			EditingAssetWidget.ToSharedRef()
+		];
+}
+
+} // namespace UE::PoseSearch
+
+#undef LOCTEXT_NAMESPACE
diff --git a/Engine/Plugins/Animation/PoseSearch/Source/Editor/Private/PoseSearchInteractionAssetFactory.cpp b/Engine/Plugins/Animation/PoseSearch/Source/Editor/Private/PoseSearchInteractionAssetFactory.cpp
new file mode 100644
--- /dev/null
+++ b/Engine/Plugins/Animation/PoseSearch/Source/Editor/Private/PoseSearchInteractionAssetFactory.cpp
@@ -0,0 +1,26 @@
+// Copyright Epic Games, Inc. All Rights Reserved.
+
+#include "PoseSearchInteractionAssetFactory.h"
+#include "PoseSearch/PoseSearchInteractionAsset.h"
+
+#define LOCTEXT_NAMESPACE "PoseSearchInteractionAssetFactory"
+
+UPoseSearchInteractionAssetFactory::UPoseSearchInteractionAssetFactory(const FObjectInitializer& ObjectInitializer)
+	: Super(ObjectInitializer)
+{
+	bCreateNew = true;
+	bEditAfterNew = true;
+	SupportedClass = UPoseSearchInteractionAsset::StaticClass();
+}
+
+UObject* UPoseSearchInteractionAssetFactory::FactoryCreateNew(UClass* Class, UObject* InParent, FName Name, EObjectFlags Flags, UObject* Context, FFeedbackContext* Warn)
+{
+	return NewObject<UPoseSearchInteractionAsset>(InParent, Class, Name, Flags);
+}
+
+FString UPoseSearchInteractionAssetFactory::GetDefaultNewAssetName() const
+{
+	return FString(TEXT("NewPoseSearchInteractionAsset"));
+}
+
+#undef LOCTEXT_NAMESPACE
\ No newline at end of file
diff --git a/Engine/Plugins/Animation/PoseSearch/Source/Editor/Public/AnimGraphNode_MotionMatchingInteraction.h b/Engine/Plugins/Animation/PoseSearch/Source/Editor/Public/AnimGraphNode_MotionMatchingInteraction.h
new file mode 100644
--- /dev/null
+++ b/Engine/Plugins/Animation/PoseSearch/Source/Editor/Public/AnimGraphNode_MotionMatchingInteraction.h
@@ -0,0 +1,24 @@
+// Copyright Epic Games, Inc. All Rights Reserved.
+
+#pragma once
+
+#include "AnimGraphNode_BlendStack.h"
+#include "PoseSearch/AnimNode_MotionMatchingInteraction.h"
+#include "AnimGraphNode_MotionMatchingInteraction.generated.h"
+
+UCLASS(Experimental)
+class UAnimGraphNode_MotionMatchingInteraction : public UAnimGraphNode_BlendStack_Base
+{
+	GENERATED_UCLASS_BODY()
+
+	UPROPERTY(EditAnywhere, Category = Settings)
+	FAnimNode_MotionMatchingInteraction Node;
+
+public:
+	virtual FText GetNodeTitle(ENodeTitleType::Type TitleType) const override;
+	virtual FLinearColor GetNodeTitleColor() const override;
+	virtual FText GetTooltipText() const override;
+
+protected:
+	virtual FAnimNode_BlendStack_Standalone* GetBlendStackNode() const override { return (FAnimNode_BlendStack_Standalone*)(&Node); }
+};
diff --git a/Engine/Plugins/Animation/PoseSearch/Source/Editor/Public/PoseSearchDatabaseViewModel.h b/Engine/Plugins/Animation/PoseSearch/Source/Editor/Public/PoseSearchDatabaseViewModel.h
--- a/Engine/Plugins/Animation/PoseSearch/Source/Editor/Public/PoseSearchDatabaseViewModel.h
+++ b/Engine/Plugins/Animation/PoseSearch/Source/Editor/Public/PoseSearchDatabaseViewModel.h
@@ -146,6 +146,7 @@ namespace UE::PoseSearch
 
 		float PlayTime = 0.f;
 		float DeltaTimeMultiplier = 1.f;
+		float StepDeltaTime = 1.f / 30.f;
 
 		/** Scene asset being viewed and edited by this view model. */
 		TWeakObjectPtr<UPoseSearchDatabase> PoseSearchDatabasePtr;
diff --git a/Engine/Plugins/Animation/PoseSearch/Source/Editor/Public/PoseSearchInteractionAssetEditor.h b/Engine/Plugins/Animation/PoseSearch/Source/Editor/Public/PoseSearchInteractionAssetEditor.h
new file mode 100644
--- /dev/null
+++ b/Engine/Plugins/Animation/PoseSearch/Source/Editor/Public/PoseSearchInteractionAssetEditor.h
@@ -0,0 +1,253 @@
+// Copyright Epic Games, Inc. All Rights Reserved.
+
+#pragma once
+
+#include "AdvancedPreviewScene.h"
+#include "CoreMinimal.h"
+#include "EditorViewportClient.h"
+#include "EdMode.h"
+#include "IMultiAnimAssetEditor.h"
+#include "Misc/NotifyHook.h"
+#include "PoseSearch/PoseSearchAssetSampler.h"
+#include "PoseSearch/PoseSearchRole.h"
+#include "SCommonEditorViewportToolbarBase.h"
+#include "SEditorViewport.h"
+
+class UAnimPreviewInstance;
+class UPoseSearchInteractionAsset;
+class UDebugSkelMeshComponent;
+
+namespace UE::PoseSearch
+{
+
+class FInteractionAssetPreviewScene;
+class FInteractionAssetEditor;
+class SInteractionAssetPreview;
+class SInteractionAssetViewport;
+
+// Experimental, this feature might be removed without warning, not for production use
+struct FInteractionAssetPreviewActor
+{
+public:
+	bool SpawnPreviewActor(UWorld* World, const UPoseSearchInteractionAsset* InteractionAsset, const UE::PoseSearch::FRole& Role);
+	void UpdatePreviewActor(const UPoseSearchInteractionAsset* InteractionAsset, float PlayTime);
+	void Destroy();
+	UAnimPreviewInstance* GetAnimPreviewInstance();
+	const UE::PoseSearch::FRole& GetRole() const { return ActorRole; }
+	FTransform GetDebugActorTransformFromSampler() const;
+	void ForceDebugActorTransform(const FTransform& ActorTransform);
+
+private:
+	enum { PreviewActor, DebugActor, NumActors };
+	TStaticArray<TWeakObjectPtr<AActor>, NumActors> ActorPtrs;
+	TStaticArray<UE::PoseSearch::FAnimationAssetSampler, NumActors> Samplers; // used to keep track of the animation position / orientation
+
+	float CurrentTime = 0.f;
+	UE::PoseSearch::FRole ActorRole = UE::PoseSearch::DefaultRole;
+	
+	// @todo: add support for blend spaces
+	FVector BlendParameters = FVector::ZeroVector;
+};
+
+// Experimental, this feature might be removed without warning, not for production use
+class FInteractionAssetViewModel : public TSharedFromThis<FInteractionAssetViewModel>, public FGCObject
+{
+public:
+	virtual void AddReferencedObjects(FReferenceCollector& Collector) override;
+	virtual FString GetReferencerName() const override { return TEXT("FInteractionAssetViewModel"); }
+	void Initialize(UPoseSearchInteractionAsset* InteractionAsset, const TSharedRef<FInteractionAssetPreviewScene>& PreviewScene);
+	void RemovePreviewActors();
+	void PreviewBackwardEnd();
+	void PreviewBackwardStep();
+	void PreviewBackward();
+	void PreviewPause();
+	void PreviewForward();
+	void PreviewForwardStep();
+	void PreviewForwardEnd();
+	const UPoseSearchInteractionAsset* GetInteractionAsset() const { return InteractionAssetPtr.Get(); }
+	void Tick(float DeltaSeconds);
+	const TArray<FInteractionAssetPreviewActor>& GetPreviewActors() const { return PreviewActors; }
+	TArray<FInteractionAssetPreviewActor>& GetPreviewActors() { return PreviewActors; }
+	TRange<double> GetPreviewPlayRange() const;
+	void SetPlayTime(float NewPlayTime, bool bInTickPlayTime);
+	float GetPlayTime() const { return PlayTime; }
+	void SetPreviewProperties(float AnimAssetTime, const FVector& AnimAssetBlendParameters, bool bAnimAssetPlaying);
+
+private:
+	UWorld* GetWorld();
+
+	float PlayTime = 0.f;
+	float DeltaTimeMultiplier = 1.f;
+	float StepDeltaTime = 1.f / 30.f;
+
+	/** asset being viewed and edited by this view model. */
+	TWeakObjectPtr<UPoseSearchInteractionAsset> InteractionAssetPtr;
+	/** Weak pointer to the PreviewScene */
+	TWeakPtr<FInteractionAssetPreviewScene> PreviewScenePtr;
+	/** Actors to be displayed in the preview viewport */
+	TArray<FInteractionAssetPreviewActor> PreviewActors;
+};
+
+// Experimental, this feature might be removed without warning, not for production use
+class FInteractionAssetEdMode : public FEdMode
+{
+public:
+	const static FEditorModeID EdModeId;
+
+	virtual void Tick(FEditorViewportClient* ViewportClient, float DeltaTime) override;
+	virtual void Render(const FSceneView* View, FViewport* Viewport, FPrimitiveDrawInterface* PDI) override;
+	virtual bool AllowWidgetMove() override;
+	virtual bool ShouldDrawWidget() const override;
+	virtual bool GetCustomInputCoordinateSystem(FMatrix& InMatrix, void* InData) override;
+
+private:
+	FInteractionAssetViewModel* ViewModel = nullptr;
+};
+
+// Experimental, this feature might be removed without warning, not for production use
+class FInteractionAssetViewportClient : public FEditorViewportClient
+{
+public:
+	FInteractionAssetViewportClient(const TSharedRef<FInteractionAssetPreviewScene>& InPreviewScene, const TSharedRef<SInteractionAssetViewport>& InViewport, const TSharedRef<FInteractionAssetEditor>& InAssetEditor);
+	virtual void Draw(const FSceneView* View, FPrimitiveDrawInterface* PDI) override;
+	virtual void TrackingStarted(const struct FInputEventState& InInputState, bool bIsDragging, bool bNudge) override;
+	virtual void TrackingStopped() override;
+
+	/** Get the preview scene we are viewing */
+	TSharedRef<FInteractionAssetPreviewScene> GetPreviewScene() const { return PreviewScenePtr.Pin().ToSharedRef(); }
+	TSharedRef<FInteractionAssetEditor> GetAssetEditor() const { return AssetEditorPtr.Pin().ToSharedRef();	}
+
+private:
+	/** Preview scene we are viewing */
+	TWeakPtr<FInteractionAssetPreviewScene> PreviewScenePtr;
+	/** Asset editor we are embedded in */
+	TWeakPtr<FInteractionAssetEditor> AssetEditorPtr;
+};
+
+// Experimental, this feature might be removed without warning, not for production use
+class FInteractionAssetPreviewScene : public FAdvancedPreviewScene
+{
+public:
+	FInteractionAssetPreviewScene(ConstructionValues CVs, const TSharedRef<FInteractionAssetEditor>& Editor);
+	virtual void Tick(float InDeltaTime) override;
+	TSharedRef<FInteractionAssetEditor> GetEditor() const { return EditorPtr.Pin().ToSharedRef(); }
+
+private:
+	/** The asset editor we are embedded in */
+	TWeakPtr<FInteractionAssetEditor> EditorPtr;
+};
+
+// Experimental, this feature might be removed without warning, not for production use
+struct FInteractionAssetPreviewRequiredArgs
+{
+	FInteractionAssetPreviewRequiredArgs(const TSharedRef<FInteractionAssetEditor>& InAssetEditor, const TSharedRef<FInteractionAssetPreviewScene>& InPreviewScene)
+		: AssetEditor(InAssetEditor)
+		, PreviewScene(InPreviewScene)
+	{
+	}
+
+	TSharedRef<FInteractionAssetEditor> AssetEditor;
+	TSharedRef<FInteractionAssetPreviewScene> PreviewScene;
+};
+
+// Experimental, this feature might be removed without warning, not for production use
+class SInteractionAssetViewport : public SEditorViewport, public ICommonEditorViewportToolbarInfoProvider
+{
+public:
+	SLATE_BEGIN_ARGS(SInteractionAssetViewport) {}
+	SLATE_END_ARGS();
+
+	void Construct(const FArguments& InArgs, const FInteractionAssetPreviewRequiredArgs& InRequiredArgs);
+	virtual TSharedRef<class SEditorViewport> GetViewportWidget() override;
+	virtual TSharedPtr<FExtender> GetExtenders() const override;
+	virtual void OnFloatingButtonClicked() override;
+
+protected:
+	virtual void BindCommands() override;
+	virtual TSharedRef<FEditorViewportClient> MakeEditorViewportClient() override;
+	virtual TSharedPtr<SWidget> MakeViewportToolbar() override;
+
+	/** The viewport toolbar */
+	TSharedPtr<SCommonEditorViewportToolbarBase> ViewportToolbar;
+	/** Viewport client */
+	TSharedPtr<FInteractionAssetViewportClient> ViewportClient;
+	/** The preview scene that we are viewing */
+	TWeakPtr<FInteractionAssetPreviewScene> PreviewScenePtr;
+	/** Asset editor we are embedded in */
+	TWeakPtr<FInteractionAssetEditor> AssetEditorPtr;
+};
+
+// Experimental, this feature might be removed without warning, not for production use
+class SInteractionAssetPreview : public SCompoundWidget
+{
+public:
+	DECLARE_DELEGATE_TwoParams(FOnScrubPositionChanged, double, bool)
+	DECLARE_DELEGATE(FOnButtonClickedEvent)
+
+	SLATE_BEGIN_ARGS(SInteractionAssetPreview) {}
+		SLATE_ATTRIBUTE(FLinearColor, SliderColor);
+		SLATE_ATTRIBUTE(double, SliderScrubTime);
+		SLATE_ATTRIBUTE(TRange<double>, SliderViewRange);
+		SLATE_EVENT(FOnScrubPositionChanged, OnSliderScrubPositionChanged);
+		SLATE_EVENT(FOnButtonClickedEvent, OnBackwardEnd);
+		SLATE_EVENT(FOnButtonClickedEvent, OnBackwardStep);
+		SLATE_EVENT(FOnButtonClickedEvent, OnBackward);
+		SLATE_EVENT(FOnButtonClickedEvent, OnPause);
+		SLATE_EVENT(FOnButtonClickedEvent, OnForward);
+		SLATE_EVENT(FOnButtonClickedEvent, OnForwardStep);
+		SLATE_EVENT(FOnButtonClickedEvent, OnForwardEnd);
+	SLATE_END_ARGS();
+
+	void Construct(const FArguments& InArgs, const FInteractionAssetPreviewRequiredArgs& InRequiredArgs);
+
+protected:
+	TAttribute<FLinearColor> SliderColor;
+	TAttribute<double> SliderScrubTime;
+	TAttribute<TRange<double>> SliderViewRange = TRange<double>(0.0, 1.0);
+	FOnScrubPositionChanged OnSliderScrubPositionChanged;
+	FOnButtonClickedEvent OnBackwardEnd;
+	FOnButtonClickedEvent OnBackwardStep;
+	FOnButtonClickedEvent OnBackward;
+	FOnButtonClickedEvent OnPause;
+	FOnButtonClickedEvent OnForward;
+	FOnButtonClickedEvent OnForwardStep;
+	FOnButtonClickedEvent OnForwardEnd;
+};
+
+// Experimental, this feature might be removed without warning, not for production use
+class FInteractionAssetEditor : public IMultiAnimAssetEditor, public FNotifyHook
+{
+public:
+	void InitAssetEditor(const EToolkitMode::Type Mode, const TSharedPtr<IToolkitHost>& InitToolkitHost, UPoseSearchInteractionAsset* InteractionAsset);
+
+	virtual void RegisterTabSpawners(const TSharedRef<FTabManager>& TabManager) override;
+	virtual void UnregisterTabSpawners(const TSharedRef<FTabManager>& TabManager) override;
+	virtual FName GetToolkitFName() const override;
+	virtual FText GetBaseToolkitName() const override;
+	virtual FText GetToolkitName() const override;
+	virtual FLinearColor GetWorldCentricTabColorScale() const override;
+	virtual FString GetWorldCentricTabPrefix() const override;
+	virtual void SetPreviewProperties(float AnimAssetTime, const FVector& AnimAssetBlendParameters, bool bAnimAssetPlaying) override;
+
+	const UPoseSearchInteractionAsset* GetInteractionAsset() const { return ViewModel.IsValid() ? ViewModel->GetInteractionAsset() : nullptr; }
+	FInteractionAssetViewModel* GetViewModel() const { return ViewModel.Get(); }
+	void PreviewBackwardEnd() { ViewModel->PreviewBackwardEnd(); }
+	void PreviewBackwardStep() { ViewModel->PreviewBackwardStep();	}
+	void PreviewBackward() { ViewModel->PreviewBackward(); }
+	void PreviewPause() { ViewModel->PreviewPause(); }
+	void PreviewForward() { ViewModel->PreviewForward(); }
+	void PreviewForwardStep() { ViewModel->PreviewForwardStep(); }
+	void PreviewForwardEnd() { ViewModel->PreviewForwardEnd(); }
+
+private:
+	TSharedRef<SDockTab> SpawnTab_Viewport(const FSpawnTabArgs& Args);
+	TSharedRef<SDockTab> SpawnTab_AssetDetails(const FSpawnTabArgs& Args);
+
+	TSharedPtr<SInteractionAssetPreview> PreviewWidget;
+	TSharedPtr<IDetailsView> EditingAssetWidget;
+	TSharedPtr<FInteractionAssetPreviewScene> PreviewScene;
+	TSharedPtr<FInteractionAssetViewModel> ViewModel;
+};
+
+} // UE::PoseSearch
+
diff --git a/Engine/Plugins/Animation/PoseSearch/Source/Editor/Public/PoseSearchInteractionAssetFactory.h b/Engine/Plugins/Animation/PoseSearch/Source/Editor/Public/PoseSearchInteractionAssetFactory.h
new file mode 100644
--- /dev/null
+++ b/Engine/Plugins/Animation/PoseSearch/Source/Editor/Public/PoseSearchInteractionAssetFactory.h
@@ -0,0 +1,19 @@
+// Copyright Epic Games, Inc. All Rights Reserved.
+
+#pragma once
+
+#include "CoreMinimal.h"
+#include "UObject/ObjectMacros.h"
+#include "Factories/Factory.h"
+#include "PoseSearchInteractionAssetFactory.generated.h"
+
+UCLASS(Experimental)
+class UPoseSearchInteractionAssetFactory : public UFactory
+{
+	GENERATED_UCLASS_BODY()
+
+	// UFactory interface
+	virtual UObject* FactoryCreateNew(UClass* Class, UObject* InParent, FName Name, EObjectFlags Flags, UObject* Context, FFeedbackContext* Warn) override;
+	virtual FString GetDefaultNewAssetName() const override;
+	// End of UFactory interface
+};
diff --git a/Engine/Plugins/Animation/PoseSearch/Source/Runtime/Private/AnimNode_MotionMatchingInteraction.cpp b/Engine/Plugins/Animation/PoseSearch/Source/Runtime/Private/AnimNode_MotionMatchingInteraction.cpp
new file mode 100644
--- /dev/null
+++ b/Engine/Plugins/Animation/PoseSearch/Source/Runtime/Private/AnimNode_MotionMatchingInteraction.cpp
@@ -0,0 +1,245 @@
+// Copyright Epic Games, Inc. All Rights Reserved.
+
+#include "PoseSearch/AnimNode_MotionMatchingInteraction.h"
+#include "Animation/AnimInertializationSyncScope.h"
+#include "Animation/AnimInstanceProxy.h"
+#include "PoseSearch/PoseSearchInteractionAsset.h"
+#include "PoseSearch/PoseSearchInteractionLibrary.h"
+#include "PoseSearch/PoseHistoryProvider.h"
+#include "PoseSearch/PoseSearchDatabase.h"
+#include "PoseSearch/PoseSearchLibrary.h"
+#include "PoseSearch/PoseSearchSchema.h"
+#include "HAL/IConsoleManager.h"
+
+#if ENABLE_ANIM_DEBUG
+TAutoConsoleVariable<bool> CVarAnimNodeMotionMatchingInteractionDebug(TEXT("a.AnimNode.MotionMatchingInteraction.Debug"), false, TEXT("Turn on visualization debugging for AnimNode Motion Matching Interaction"));
+#endif // ENABLE_ANIM_DEBUG
+
+void FAnimNode_MotionMatchingInteraction::GatherDebugData(FNodeDebugData& DebugData)
+{
+	Super::GatherDebugData(DebugData);
+
+	FString DebugLine = DebugData.GetNodeName(this);
+#if ENABLE_ANIM_DEBUG
+	static bool bActive = true;
+	DebugLine += FString::Printf(TEXT("\n - Active: (%s)"), bActive);
+#endif
+	DebugData.AddDebugItem(DebugLine);
+}
+
+void FAnimNode_MotionMatchingInteraction::Initialize_AnyThread(const FAnimationInitializeContext& Context)
+{
+	Super::Initialize_AnyThread(Context);
+	Source.Initialize(Context);
+}
+
+void FAnimNode_MotionMatchingInteraction::CacheBones_AnyThread(const FAnimationCacheBonesContext& Context)
+{
+	Super::CacheBones_AnyThread(Context);
+	Source.CacheBones(Context);
+}
+
+void FAnimNode_MotionMatchingInteraction::Reset()
+{
+	Super::Reset();
+	BlendLerp = 0.f;
+	TranslationWarpLerp = 0.f;
+	RotationWarpLerp = 0.f;
+	InteractingRolesNum = 0;
+}
+
+void FAnimNode_MotionMatchingInteraction::UpdateAssetPlayer(const FAnimationUpdateContext& Context)
+{
+	using namespace UE::PoseSearch;
+
+	if (NeedsReset(Context))
+	{
+		Reset();
+	}
+
+	UpdateCounter.SynchronizeWith(Context.AnimInstanceProxy->GetUpdateCounter());
+
+	GetEvaluateGraphExposedInputs().Execute(Context);
+
+	bool bBlendToExecuted = ConditionalBlendTo(Context);
+
+	int32 NewInteractingRolesNum = 0;
+	const float DeltaTime = Context.GetDeltaTime();
+	FPoseHistoryProvider* PoseHistoryProvider = Context.GetMessage<FPoseHistoryProvider>();
+	if (!PoseHistoryProvider)
+	{
+		UE_LOG(LogPoseSearch, Error, TEXT("FAnimNode_MotionMatchingInteraction::Update_AnyThread couldn't find the FPoseHistoryProvider"));
+	}
+	else
+	{
+		// we're NOT providing the UMultiAnimAsset here, relying on the interaction subsystem to resolve the ContinuingProperties properly
+		FPoseSearchContinuingProperties ContinuingProperties;
+		ContinuingProperties.PlayingAsset = GetAnimAsset();
+		ContinuingProperties.PlayingAssetAccumulatedTime = GetAccumulatedTime();
+
+		check(Context.AnimInstanceProxy);
+		const FPoseSearchInteractionBlueprintResult Result = UPoseSearchInteractionLibrary::MotionMatchInteraction(Availabilities, Context.AnimInstanceProxy->GetAnimInstanceObject(), ContinuingProperties, PoseHistoryProvider->GetHistoryCollector(), bValidateResultAgainstAvailabilities);
+
+		UAnimationAsset* RoledAnimAsset = nullptr;
+		if (const UMultiAnimAsset* MultiAnimAsset = Cast<UMultiAnimAsset>(Result.SelectedAnimation))
+		{
+			check(Result.SelectedDatabase != nullptr);
+
+			RoledAnimAsset = MultiAnimAsset->GetAnimationAsset(Result.Role);
+			NewInteractingRolesNum = MultiAnimAsset->GetNumRoles();
+		}
+		else if (UAnimationAsset* SingleAnimAsset = Cast<UAnimationAsset>(Result.SelectedAnimation))
+		{
+			RoledAnimAsset = SingleAnimAsset;
+			NewInteractingRolesNum = 1;
+		}
+
+		if (RoledAnimAsset)
+		{
+			bool bExecuteBlendTo = false;
+			bool bUpdatePropertiesFromResult = false;
+			if (InteractingRolesNum == 0 || AnimPlayers.IsEmpty())
+			{
+				bExecuteBlendTo = true;
+				bUpdatePropertiesFromResult = true;
+			}
+			else if (EvaluationMode == EMotionMatchingInteractionEvaluationMode::ContinuousReselection)
+			{
+				const FBlendStackAnimPlayer& MainAnimPlayer = AnimPlayers[0];
+				const UAnimationAsset* PlayingAnimAsset = MainAnimPlayer.GetAnimationAsset();
+
+				if (RoledAnimAsset != PlayingAnimAsset ||
+					Result.bIsMirrored != MainAnimPlayer.GetMirror() ||
+					Result.BlendParameters != MainAnimPlayer.GetBlendParameters() ||
+					!Result.bIsContinuingPoseSearch)
+				{
+					bExecuteBlendTo = true;
+				}
+
+				bUpdatePropertiesFromResult = true;
+			}
+			else if (RoledAnimAsset == AnimPlayers[0].GetAnimationAsset() && Result.bIsContinuingPoseSearch)
+			{
+				// we don't update FullAlignedActorRootBoneTransform since we're not planning to blend into the newly selected animation here
+				bUpdatePropertiesFromResult = true;
+			}
+
+			if (bUpdatePropertiesFromResult)
+			{
+				FullAlignedActorRootBoneTransform = Result.FullAlignedActorRootBoneTransform;
+				WantedPlayRate = Result.WantedPlayRate;
+				BlendParameters = Result.BlendParameters;
+			}
+
+			if (bExecuteBlendTo)
+			{
+				const FPoseSearchRoledSkeleton* RoledSkeleton = Result.SelectedDatabase->Schema->GetRoledSkeleton(Result.Role);
+				check(RoledSkeleton);
+
+				BlendTo(Context, RoledAnimAsset, Result.SelectedTime, Result.bLoop, Result.bIsMirrored, RoledSkeleton->MirrorDataTable.Get(),
+					BlendTime, BlendProfile, BlendOption, bUseInertialBlend, BlendParameters, WantedPlayRate);
+
+				bBlendToExecuted = true;
+			}
+		}
+	}
+
+	const bool bDidBlendToRequestAnInertialBlend = bBlendToExecuted && bUseInertialBlend;
+	UE::Anim::TOptionalScopedGraphMessage<UE::Anim::FAnimInertializationSyncScope> InertializationSync(bDidBlendToRequestAnInertialBlend, Context);
+	
+	UpdatePlayRate(WantedPlayRate);
+	UpdateBlendspaceParameters(BlendspaceUpdateMode, BlendParameters);
+
+	// calculating the translation and rotation warp lerps, used to warp the root transform towards the last computed FullAlignedActorRootBoneTransform
+	if (BlendTime > UE_KINDA_SMALL_NUMBER)
+	{
+		const float BlendSign = NewInteractingRolesNum > 0 ? 1.f : -1.f;
+		BlendLerp = FMath::Clamp(BlendLerp + (BlendSign * DeltaTime / BlendTime), 0.f, 1.f);
+	}
+	
+	const float WarpSign = NewInteractingRolesNum > 1 ? 1.f : -1.f;
+	if (InitialTranslationWarpTime > UE_KINDA_SMALL_NUMBER)
+	{
+		TranslationWarpLerp = FMath::Clamp(TranslationWarpLerp + (WarpSign * DeltaTime / InitialTranslationWarpTime), 0.f, 1.f);
+	}
+	if (InitialRotationWarpTime > UE_KINDA_SMALL_NUMBER)
+	{
+		RotationWarpLerp = FMath::Clamp(RotationWarpLerp + (WarpSign * DeltaTime / InitialRotationWarpTime), 0.f, 1.f);
+	}
+
+	// updating source
+	const float SourceBlendLerp = 1.f - BlendLerp;
+	FAnimationUpdateContext SourceContext = Context.FractionalWeightAndRootMotion(SourceBlendLerp, SourceBlendLerp);
+	if (BlendLerp > UE_KINDA_SMALL_NUMBER)
+	{
+		SourceContext = SourceContext.AsInactive();
+	}
+	Source.Update(SourceContext);
+	
+	if (Source.GetLinkNode())
+	{
+		// updating blend stack
+		if (BlendLerp < UE_KINDA_SMALL_NUMBER && DeltaTime > UE_KINDA_SMALL_NUMBER)
+		{
+			// resetting the blendstack if there's no BlendLerp weight to it
+			Reset();
+		}
+
+		FAnimationUpdateContext BlendStackContext = Context.FractionalWeightAndRootMotion(BlendLerp, BlendLerp);
+
+		// bypassing FAnimNode_BlendStack::UpdateAssetPlayer, since we overridden its behaviour
+		FAnimNode_BlendStack_Standalone::UpdateAssetPlayer(BlendStackContext);
+	}
+	else
+	{
+		// bypassing FAnimNode_BlendStack::UpdateAssetPlayer, since we overridden its behaviour
+		FAnimNode_BlendStack_Standalone::UpdateAssetPlayer(Context);
+	}
+
+#if ENABLE_ANIM_DEBUG
+	if (UE_TRACE_CHANNELEXPR_IS_ENABLED(AnimationChannel))
+	{
+		TRACE_ANIM_NODE_VALUE(Context, *FString("InteractingRolesNum"), InteractingRolesNum);
+		TRACE_ANIM_NODE_VALUE(Context, *FString("NewInteractingRolesNum"), NewInteractingRolesNum);
+		TRACE_ANIM_NODE_VALUE(Context, *FString("BlendToExecuted"), bBlendToExecuted);
+		TRACE_ANIM_NODE_VALUE(Context, *FString("BlendLerp"), BlendLerp);
+		TRACE_ANIM_NODE_VALUE(Context, *FString("TranslationWarpLerp"), TranslationWarpLerp);
+		TRACE_ANIM_NODE_VALUE(Context, *FString("RotationWarpLerp"), RotationWarpLerp);
+	}
+#endif // ENABLE_ANIM_DEBUG
+
+	InteractingRolesNum = NewInteractingRolesNum;
+}
+
+void FAnimNode_MotionMatchingInteraction::Evaluate_AnyThread(FPoseContext& Output)
+{
+	check(Output.AnimInstanceProxy);
+
+	// evaluating Source to get the base pose 
+	Source.Evaluate(Output);
+
+	// @todo perform a proper blend between Source and blend stack. In the current WIP implementation blend stack will stomp over Source 
+	if (!AnimPlayers.IsEmpty())
+	{
+		Super::Evaluate_AnyThread(Output);
+	}
+
+	if (TranslationWarpLerp > UE_KINDA_SMALL_NUMBER || RotationWarpLerp > UE_KINDA_SMALL_NUMBER)
+	{
+		const FTransform ComponentTransform = Output.AnimInstanceProxy->GetComponentTransform();
+		const FTransform FullAlignedActorRootBoneLocalTransform = FullAlignedActorRootBoneTransform.GetRelativeTransform(ComponentTransform);
+
+#if ENABLE_ANIM_DEBUG
+		if (CVarAnimNodeMotionMatchingInteractionDebug.GetValueOnAnyThread())
+		{
+			Output.AnimInstanceProxy->AnimDrawDebugCoordinateSystem(FullAlignedActorRootBoneTransform.GetLocation(), FullAlignedActorRootBoneTransform.Rotator(), 15.f, false, 0.f, 0.f, SDPG_Foreground);
+			Output.AnimInstanceProxy->AnimDrawDebugCoordinateSystem(ComponentTransform.GetLocation(), ComponentTransform.Rotator(), 5.f, false, 0.f, 0.f, SDPG_Foreground);
+		}
+#endif // ENABLE_ANIM_DEBUG
+
+		const FCompactPoseBoneIndex RootBoneIndex(0);
+		FTransform& RootBoneTransform = Output.Pose[RootBoneIndex];
+		RootBoneTransform.SetTranslation(FMath::Lerp(RootBoneTransform.GetTranslation(), FullAlignedActorRootBoneLocalTransform.GetTranslation(), TranslationWarpLerp));
+		RootBoneTransform.SetRotation(FQuat::Slerp(RootBoneTransform.GetRotation(), FullAlignedActorRootBoneLocalTransform.GetRotation(), RotationWarpLerp));
+	}
+}
diff --git a/Engine/Plugins/Animation/PoseSearch/Source/Runtime/Private/MotionMatchingInteractionAnimNodeLibrary.cpp b/Engine/Plugins/Animation/PoseSearch/Source/Runtime/Private/MotionMatchingInteractionAnimNodeLibrary.cpp
new file mode 100644
--- /dev/null
+++ b/Engine/Plugins/Animation/PoseSearch/Source/Runtime/Private/MotionMatchingInteractionAnimNodeLibrary.cpp
@@ -0,0 +1,21 @@
+// Copyright Epic Games, Inc. All Rights Reserved.
+
+#include "PoseSearch/MotionMatchingInteractionAnimNodeLibrary.h"
+#include "PoseSearch/AnimNode_MotionMatchingInteraction.h"
+
+FMotionMatchingInteractionAnimNodeReference UMotionMatchingInteractionAnimNodeLibrary::ConvertToMotionMatchingInteractionNode(const FAnimNodeReference& Node, EAnimNodeReferenceConversionResult& Result)
+{
+	return FAnimNodeReference::ConvertToType<FMotionMatchingInteractionAnimNodeReference>(Node, Result);
+}
+
+void UMotionMatchingInteractionAnimNodeLibrary::SetAvailabilities(const FMotionMatchingInteractionAnimNodeReference& MotionMatchingInteractionNode, const TArray<FPoseSearchInteractionAvailability>& Availabilities)
+{
+	if (FAnimNode_MotionMatchingInteraction* MotionMatchingInteractionNodePtr = MotionMatchingInteractionNode.GetAnimNodePtr<FAnimNode_MotionMatchingInteraction>())
+	{
+		MotionMatchingInteractionNodePtr->Availabilities = Availabilities;
+	}
+	else
+	{
+		UE_LOG(LogPoseSearch, Warning, TEXT("UMotionMatchingInteractionAnimNodeLibrary::SetAvailabilities called on an invalid context or with an invalid type"));
+	}
+}
\ No newline at end of file
diff --git a/Engine/Plugins/Animation/PoseSearch/Source/Runtime/Private/PoseSearchInteractionAsset.cpp b/Engine/Plugins/Animation/PoseSearch/Source/Runtime/Private/PoseSearchInteractionAsset.cpp
new file mode 100644
--- /dev/null
+++ b/Engine/Plugins/Animation/PoseSearch/Source/Runtime/Private/PoseSearchInteractionAsset.cpp
@@ -0,0 +1,345 @@
+// Copyright Epic Games, Inc. All Rights Reserved.
+
+#include "PoseSearch/PoseSearchInteractionAsset.h"
+#include "Animation/AnimSequence.h"
+#include "Animation/BlendSpace.h"
+#include "PoseSearch/PoseSearchAssetSampler.h"
+
+bool UPoseSearchInteractionAsset::IsLooping() const
+{
+	float CommonPlayLength = -1.f;
+	for (const FPoseSearchInteractionAssetItem& Item : Items)
+	{
+		if (const UAnimationAsset* AnimationAsset = Item.Animation.Get())
+		{
+			if (const UAnimSequenceBase* SequenceBase = Cast<UAnimSequenceBase>(AnimationAsset))
+			{
+				if (!SequenceBase->bLoop)
+				{
+					return false;
+				}
+			}
+			else if (const UBlendSpace* BlendSpace = Cast<UBlendSpace>(AnimationAsset))
+			{
+				if (!BlendSpace->bLoop)
+				{
+					return false;
+				}
+			}
+			else
+			{
+				unimplemented();
+			}
+
+			if (CommonPlayLength < 0.f)
+			{
+				CommonPlayLength = AnimationAsset->GetPlayLength();
+			}
+			else if (!FMath::IsNearlyEqual(CommonPlayLength, AnimationAsset->GetPlayLength()))
+			{
+				return false;
+			}
+		}
+	}
+	return true;
+}
+
+bool UPoseSearchInteractionAsset::HasRootMotion() const
+{
+	bool bHasAtLeastOneValidItem = false;
+	bool bHasRootMotion = true;
+
+	for (const FPoseSearchInteractionAssetItem& Item : Items)
+	{
+		if (const UAnimationAsset* AnimationAsset = Item.Animation.Get())
+		{
+			if (const UAnimSequenceBase* SequenceBase = Cast<UAnimSequenceBase>(AnimationAsset))
+			{
+				bHasRootMotion &= SequenceBase->HasRootMotion();
+			}
+			else if (const UBlendSpace* BlendSpace = Cast<UBlendSpace>(AnimationAsset))
+			{
+				BlendSpace->ForEachImmutableSample([&bHasRootMotion](const FBlendSample& Sample)
+				{
+					if (const UAnimSequence* Sequence = Sample.Animation.Get())
+					{
+						bHasRootMotion &= Sequence->HasRootMotion();
+					}
+				});
+			}
+			else
+			{
+				unimplemented();
+			}
+			bHasAtLeastOneValidItem = true;
+		}
+	}
+
+	return bHasAtLeastOneValidItem && bHasRootMotion;
+}
+
+float UPoseSearchInteractionAsset::GetPlayLength() const
+{
+	float MaxPlayLength = 0.f;
+	for (const FPoseSearchInteractionAssetItem& Item : Items)
+	{
+		if (const UAnimationAsset* AnimationAsset = Item.Animation.Get())
+		{
+			MaxPlayLength = FMath::Max(MaxPlayLength, AnimationAsset->GetPlayLength());
+		}
+	}
+	return MaxPlayLength;
+}
+
+#if WITH_EDITOR
+int32 UPoseSearchInteractionAsset::GetFrameAtTime(float Time) const
+{
+	const UAnimationAsset* MaxPlayLengthAnim = nullptr;
+	float MaxPlayLength = -1.f;
+	for (const FPoseSearchInteractionAssetItem& Item : Items)
+	{
+		if (const UAnimationAsset* AnimationAsset = Item.Animation.Get())
+		{
+			const float PlayLength = AnimationAsset->GetPlayLength();
+			if (PlayLength > MaxPlayLength)
+			{
+				MaxPlayLength = PlayLength;
+				MaxPlayLengthAnim = AnimationAsset;
+			}
+		}
+	}
+	
+	if (MaxPlayLengthAnim)
+	{
+		if (const UAnimSequenceBase* SequenceBase = Cast<UAnimSequenceBase>(MaxPlayLengthAnim))
+		{
+			return SequenceBase->GetFrameAtTime(Time);
+		}
+
+		if (const UBlendSpace* BlendSpace = Cast<UBlendSpace>(MaxPlayLengthAnim))
+		{
+			// returning the percentage of time as value to diplay in the pose search debugger (NoTe: BlendSpace->GetPlayLength() is one)
+			return FMath::RoundToInt(Time * 100.f);
+		}
+
+		unimplemented();
+	}
+
+	return 0;
+}
+#endif // WITH_EDITOR
+
+FQuat UPoseSearchInteractionAsset::FindReferenceOrientation(const TArrayView<const FTransform> Transforms, const TArrayView<int32> SortedByWarpingWeightRotationItemIndex) const
+{
+	const int32 ItemsNum = Items.Num();
+
+	check(ItemsNum > 0);
+	check(Transforms.Num() == ItemsNum);
+	check(Transforms.Num() == SortedByWarpingWeightRotationItemIndex.Num());
+
+	const int32 LastItemIndex = ItemsNum - 1;
+	if (ItemsNum > 1)
+	{
+		FVector OtherItemsPositionsSum = Transforms[SortedByWarpingWeightRotationItemIndex[0]].GetTranslation();
+		for (int32 ItemIndex = 1; ItemIndex < LastItemIndex; ++ItemIndex)
+		{
+			OtherItemsPositionsSum += Transforms[SortedByWarpingWeightRotationItemIndex[ItemIndex]].GetTranslation();
+		}
+		
+		const FVector OtherItemsPositionAverage = OtherItemsPositionsSum / LastItemIndex;
+		const FVector DeltaPosition = OtherItemsPositionAverage - Transforms[SortedByWarpingWeightRotationItemIndex[LastItemIndex]].GetTranslation();
+
+		if (!DeltaPosition.IsNearlyZero())
+		{
+			return DeltaPosition.ToOrientationQuat();
+		}
+	}
+
+	return Transforms[SortedByWarpingWeightRotationItemIndex[LastItemIndex]].GetRotation();
+}
+
+FVector UPoseSearchInteractionAsset::FindReferencePosition(const TArrayView<const FTransform> Transforms, const TArrayView<float> NormalizedWarpingWeightTranslation) const
+{
+	const int32 ItemsNum = Items.Num();
+	
+	check(ItemsNum > 0);
+	check(Transforms.Num() == ItemsNum);
+	check(Transforms.Num() == NormalizedWarpingWeightTranslation.Num());
+
+	FVector PositionsSum = FVector::ZeroVector;
+	for (int32 ItemIndex = 0; ItemIndex < ItemsNum; ++ItemIndex)
+	{
+		PositionsSum += Transforms[ItemIndex].GetTranslation() * NormalizedWarpingWeightTranslation[ItemIndex];
+	}
+
+	return PositionsSum;
+}
+
+UAnimationAsset* UPoseSearchInteractionAsset::GetAnimationAsset(const UE::PoseSearch::FRole& Role) const
+{
+	for (const FPoseSearchInteractionAssetItem& Item : Items)
+	{
+		if (Item.Role == Role)
+		{
+			return Item.Animation;
+		}
+	}
+	return nullptr;
+}
+
+FTransform UPoseSearchInteractionAsset::GetOrigin(const UE::PoseSearch::FRole& Role) const
+{
+	for (const FPoseSearchInteractionAssetItem& Item : Items)
+	{
+		if (Item.Role == Role)
+		{
+			return Item.Origin;
+		}
+	}
+	return FTransform::Identity;
+}
+
+#if WITH_EDITOR
+FTransform UPoseSearchInteractionAsset::GetDebugWarpOrigin(const UE::PoseSearch::FRole& Role, bool bComposeWithDebugWarpOffset) const
+{
+	for (int32 ItemIndex = 0; ItemIndex < Items.Num(); ++ItemIndex)
+	{
+		const FPoseSearchInteractionAssetItem& Item = Items[ItemIndex];
+		if (Item.Role == Role)
+		{
+#if WITH_EDITORONLY_DATA
+			if (bComposeWithDebugWarpOffset && bEnableDebugWarp && DebugWarpOffsets.IsValidIndex(ItemIndex))
+			{
+				return DebugWarpOffsets[ItemIndex] * Item.Origin;
+			}
+#endif // WITH_EDITORONLY_DATA
+
+			return Item.Origin;
+		}
+	}
+	return FTransform::Identity;
+}
+#endif // WITH_EDITOR
+
+void UPoseSearchInteractionAsset::CalculateWarpTransforms(float Time, const TArrayView<const FTransform> ActorRootBoneTransforms, TArrayView<FTransform> FullAlignedActorRootBoneTransforms) const
+{
+	check(ActorRootBoneTransforms.Num() == GetNumRoles());
+	check(FullAlignedActorRootBoneTransforms.Num() == GetNumRoles());
+
+	const int32 ItemsNum = Items.Num();
+	if (ItemsNum == 0)
+	{
+		return;
+	}
+
+	TArray<FTransform, TInlineAllocator<UE::PoseSearch::PreallocatedRolesNum>> AssetRootBoneTransforms;
+	AssetRootBoneTransforms.SetNum(ItemsNum);
+
+	// ItemIndex is the RoleIndex and Role = Item.Role
+	for (int32 ItemIndex = 0; ItemIndex < ItemsNum; ++ItemIndex)
+	{
+		const FPoseSearchInteractionAssetItem& Item = Items[ItemIndex];
+
+		// sampling the AnimationAsset to extract the current time transform and the initial (time of 0) transform
+		const UE::PoseSearch::FAnimationAssetSampler Sampler(Item.Animation, Item.Origin);
+		AssetRootBoneTransforms[ItemIndex] = Sampler.ExtractRootTransform(Time);
+
+#if DO_CHECK
+		// array containing the bone index of the root bone (0)
+		TArray<uint16, TInlineAllocator<1>> BoneIndices;
+		BoneIndices.SetNumZeroed(1);
+
+		// extracting the pose, containing only the root bone from the Sampler 
+		FMemMark Mark(FMemStack::Get());
+		FCompactPose Pose;
+		FBoneContainer BoneContainer;
+		BoneContainer.InitializeTo(BoneIndices, UE::Anim::FCurveFilterSettings(UE::Anim::ECurveFilterMode::DisallowAll), *Items[ItemIndex].Animation->GetSkeleton());
+		Pose.SetBoneContainer(&BoneContainer);
+		Sampler.ExtractPose(Time, Pose);
+
+		// making sure the animation root bone transform is Identity, so we can confuse the root with the root BONE transform and preserve performances!
+		const FTransform& RootBoneTransform = Pose.GetBones()[0];
+		check(RootBoneTransform.Equals(FTransform::Identity));
+#endif // DO_CHECK
+	}
+
+	TArray<int32, TInlineAllocator<UE::PoseSearch::PreallocatedRolesNum>> SortedByWarpingWeightRotationItemIndex;
+	TArray<float, TInlineAllocator<UE::PoseSearch::PreallocatedRolesNum>> NormalizedWarpingWeightTranslation;
+	SortedByWarpingWeightRotationItemIndex.SetNum(ItemsNum);
+	NormalizedWarpingWeightTranslation.SetNum(ItemsNum);
+
+	float WarpingWeightTranslationSum = 0.f;
+	float WarpingWeightRotationSum = 0.f;
+	for (int32 ItemIndex = 0; ItemIndex < ItemsNum; ++ItemIndex)
+	{
+		SortedByWarpingWeightRotationItemIndex[ItemIndex] = ItemIndex;
+		WarpingWeightTranslationSum += Items[ItemIndex].WarpingWeightTranslation;
+		WarpingWeightRotationSum += Items[ItemIndex].WarpingWeightRotation;
+	}
+
+	const float NormalizedHomogeneousWeight = 1.f / ItemsNum;
+	if (WarpingWeightTranslationSum > UE_KINDA_SMALL_NUMBER)
+	{
+		for (int32 ItemIndex = 0; ItemIndex < ItemsNum; ++ItemIndex)
+		{
+			NormalizedWarpingWeightTranslation[ItemIndex] = Items[ItemIndex].WarpingWeightTranslation / WarpingWeightTranslationSum;
+		}
+	}
+	else
+	{
+		for (int32 ItemIndex = 0; ItemIndex < ItemsNum; ++ItemIndex)
+		{
+			NormalizedWarpingWeightTranslation[ItemIndex] = NormalizedHomogeneousWeight;
+		}
+	}
+
+	if (WarpingWeightRotationSum > UE_KINDA_SMALL_NUMBER)
+	{
+		SortedByWarpingWeightRotationItemIndex.Sort([this](const int32 A, const int32 B)
+			{
+				return Items[A].WarpingWeightRotation < Items[B].WarpingWeightRotation;
+			});
+	}
+
+	const FQuat AssetReferenceOrientation = FindReferenceOrientation(AssetRootBoneTransforms, SortedByWarpingWeightRotationItemIndex);
+	const FQuat ActorsReferenceOrientation = FindReferenceOrientation(ActorRootBoneTransforms, SortedByWarpingWeightRotationItemIndex);
+	
+	FQuat WeightedActorsReferenceOrientation = ActorsReferenceOrientation;
+	if (WarpingWeightRotationSum > UE_KINDA_SMALL_NUMBER)
+	{
+		// ItemIndex are in order of WarpingWeightRotation. the last one is the one with the highest most WarpingWeightRotation, the most "important"
+		for (int32 ItemIndex : SortedByWarpingWeightRotationItemIndex)
+		{
+			const FPoseSearchInteractionAssetItem& Item = Items[ItemIndex];
+			const float NormalizedWarpingWeightRotation = Item.WarpingWeightRotation / WarpingWeightRotationSum;
+			if (NormalizedWarpingWeightRotation > NormalizedHomogeneousWeight)
+			{
+				// NormalizedHomogeneousWeight is one only if ItemsNum is one, 
+				// BUT NormalizedWarpingWeightRotation > NormalizedHomogeneousWeight should always be false
+				check(!FMath::IsNearlyEqual(NormalizedHomogeneousWeight, 1.f));
+
+				// how much this item wants to reorient the ReferenceOrientation from the homogeneous "fair" value
+				const float SlerpParam = (NormalizedWarpingWeightRotation - NormalizedHomogeneousWeight) / (1.f - NormalizedHomogeneousWeight);
+
+				// calculating the reference orientation relative to the character
+				// AssetReferenceOrientation in actor world orientation
+				const FQuat ActorAssetReferenceOrientation = ActorRootBoneTransforms[ItemIndex].GetRotation() * (AssetRootBoneTransforms[ItemIndex].GetRotation().Inverse() * AssetReferenceOrientation);
+
+				WeightedActorsReferenceOrientation = FQuat::Slerp(WeightedActorsReferenceOrientation, ActorAssetReferenceOrientation, SlerpParam);
+			}
+		}
+	}
+
+	const FVector AssetReferencePosition = FindReferencePosition(AssetRootBoneTransforms, NormalizedWarpingWeightTranslation);
+	const FVector ActorsReferencePosition = FindReferencePosition(ActorRootBoneTransforms, NormalizedWarpingWeightTranslation);
+
+	// aligning all the actors to ActorsReferencePosition, WeightedActorsReferenceOrientation
+	const FTransform AssetReferenceTransform(AssetReferenceOrientation, AssetReferencePosition);
+	const FTransform ActorsReferenceTransform(WeightedActorsReferenceOrientation, ActorsReferencePosition);
+	const FTransform AssetReferenceInverseTransform = AssetReferenceTransform.Inverse();
+	
+	for (int32 ItemIndex = 0; ItemIndex < ItemsNum; ++ItemIndex)
+	{
+		FullAlignedActorRootBoneTransforms[ItemIndex] = (AssetRootBoneTransforms[ItemIndex] * AssetReferenceInverseTransform) * ActorsReferenceTransform;
+	}
+}
\ No newline at end of file
diff --git a/Engine/Plugins/Animation/PoseSearch/Source/Runtime/Private/PoseSearchInteractionIsland.cpp b/Engine/Plugins/Animation/PoseSearch/Source/Runtime/Private/PoseSearchInteractionIsland.cpp
new file mode 100644
--- /dev/null
+++ b/Engine/Plugins/Animation/PoseSearch/Source/Runtime/Private/PoseSearchInteractionIsland.cpp
@@ -0,0 +1,667 @@
+// Copyright Epic Games, Inc. All Rights Reserved.
+
+#include "PoseSearch/PoseSearchInteractionIsland.h"
+#include "Animation/AnimClassInterface.h"
+#include "Animation/AnimInstance.h"
+#include "Components/SkeletalMeshComponent.h"
+#include "PoseSearch/PoseSearchInteractionLibrary.h"
+#include "DrawDebugHelpers.h"
+#include "GameFramework/CharacterMovementComponent.h"
+#include "PoseSearch/AnimNode_PoseSearchHistoryCollector.h"
+#include "PoseSearch/MultiAnimAsset.h"
+#include "PoseSearch/PoseSearchLibrary.h"
+#include "PoseSearch/PoseSearchDatabase.h"
+#include "PoseSearch/PoseSearchSchema.h"
+
+namespace UE::PoseSearch
+{
+
+#if ENABLE_ANIM_DEBUG
+static TAutoConsoleVariable<bool> CVarPoseSearchInteractionShowIslands(TEXT("a.PoseSearchInteraction.ShowIslands"), false, TEXT("Show Pose Search Interaction Islands"));
+#endif
+
+// Utils functions
+///////////////////////////////////////////////////////////
+static FInteractionSearchResult InitSearchResult(const FSearchResult& SearchResult, int32 SearchIndex, const FInteractionSearchContext& SearchContext)
+{
+	FInteractionSearchResult InteractionSearchResult;
+	static_cast<FSearchResult&>(InteractionSearchResult) = SearchResult;
+	InteractionSearchResult.SearchIndex = SearchIndex;
+
+	if (const UE::PoseSearch::FSearchIndexAsset* SearchIndexAsset = SearchResult.GetSearchIndexAsset())
+	{
+		if (const FPoseSearchDatabaseAnimationAssetBase* DatabaseAnimationAssetBase = SearchResult.Database->GetDatabaseAnimationAsset<FPoseSearchDatabaseAnimationAssetBase>(*SearchIndexAsset))
+		{
+			if (const UMultiAnimAsset* MultiAnimAsset = Cast<UMultiAnimAsset>(DatabaseAnimationAssetBase->GetAnimationAsset()))
+			{
+				check(MultiAnimAsset->GetNumRoles() == SearchContext.AnimInstances.Num());
+				check(MultiAnimAsset->GetNumRoles() == SearchContext.HistoryCollectors.Num());
+				check(MultiAnimAsset->GetNumRoles() == SearchContext.Roles.Num());
+
+				UE::PoseSearch::FRoleToIndex SearchContextRoleToIndex;
+				SearchContextRoleToIndex.Reserve(SearchContext.Roles.Num());
+				for (int32 SearchContextRoleIndex = 0; SearchContextRoleIndex < SearchContext.Roles.Num(); ++SearchContextRoleIndex)
+				{
+					SearchContextRoleToIndex.Add(SearchContext.Roles[SearchContextRoleIndex]) = SearchContextRoleIndex;
+				}
+
+				// mapping MultiAnimAsset Roles to SearchContext.* indexes
+				TArray<FTransform, TInlineAllocator<UE::PoseSearch::PreallocatedRolesNum>> ActorRootBoneTransforms;
+				for (int32 MultiAnimAssetRoleIndex = 0; MultiAnimAssetRoleIndex < MultiAnimAsset->GetNumRoles(); ++MultiAnimAssetRoleIndex)
+				{
+					const UE::PoseSearch::FRole& InteractionAssetRole = MultiAnimAsset->GetRole(MultiAnimAssetRoleIndex);
+					const int32 SearchContextIndex = SearchContextRoleToIndex[InteractionAssetRole];
+
+					if (UAnimInstance* AnimInstance = SearchContext.AnimInstances[SearchContextIndex].Get())
+					{
+						const FTransform RootBoneTransform = AnimInstance->GetSkelMeshComponent()->GetBoneTransform(0);
+						ActorRootBoneTransforms.Add(RootBoneTransform);
+					}
+				}
+
+				if (ActorRootBoneTransforms.Num() == MultiAnimAsset->GetNumRoles())
+				{
+					// FullAlignedActorRootBoneTransforms is mapped to the MultiAnimAsset roles:
+					// FullAlignedActorRootBoneTransforms[0] is for MultiAnimAsset->GetRole(0)
+					TArray<FTransform, TInlineAllocator<UE::PoseSearch::PreallocatedRolesNum>> FullAlignedActorRootBoneTransforms;
+					FullAlignedActorRootBoneTransforms.SetNum(MultiAnimAsset->GetNumRoles());
+
+					// @todo: should it be SearchResult.AssetTime + DeltaTime?
+					MultiAnimAsset->CalculateWarpTransforms(SearchResult.AssetTime, ActorRootBoneTransforms, FullAlignedActorRootBoneTransforms);
+					InteractionSearchResult.FullAlignedActorRootBoneTransforms.SetNum(MultiAnimAsset->GetNumRoles(), EAllowShrinking::No);
+
+					for (int32 InteractionAssetRoleIndex = 0; InteractionAssetRoleIndex < MultiAnimAsset->GetNumRoles(); ++InteractionAssetRoleIndex)
+					{
+						const UE::PoseSearch::FRole& InteractionAssetRole = MultiAnimAsset->GetRole(InteractionAssetRoleIndex);
+						const int32 SearchContextIndex = SearchContextRoleToIndex[InteractionAssetRole];
+
+						InteractionSearchResult.FullAlignedActorRootBoneTransforms[SearchContextIndex] = FullAlignedActorRootBoneTransforms[InteractionAssetRoleIndex];
+					}
+				}
+			}
+			else
+			{
+				// support for non UMultiAnimAsset to be backward compatible with regular motion matching searches
+				check(SearchContext.AnimInstances.Num() == 1);
+
+				if (UAnimInstance* AnimInstance = SearchContext.AnimInstances[0].Get())
+				{
+					const FTransform RootBoneTransform = AnimInstance->GetSkelMeshComponent()->GetBoneTransform(0);
+					InteractionSearchResult.FullAlignedActorRootBoneTransforms.SetNum(1, EAllowShrinking::No);
+					// @todo: should it be RootBoneTransform + root motion transform?
+					InteractionSearchResult.FullAlignedActorRootBoneTransforms[0] = RootBoneTransform;
+				}
+			}
+		}
+	}
+
+	return InteractionSearchResult;
+}
+
+static void InitSearchResults(TArray<FInteractionSearchResult>& SearchResults, TConstArrayView<FSearchResult> PoseSearchResults, TConstArrayView<FInteractionSearchContext> SearchContexts)
+{
+	// WIP!
+	// @todo: figure out multiple policies to use the most characters? right now only the best search is "valid" with the most characters
+	int32 BestSearchIndex = INDEX_NONE;
+	for (int32 SearchIndex = 0; SearchIndex < PoseSearchResults.Num(); ++SearchIndex)
+	{
+		if (PoseSearchResults[SearchIndex].IsValid())
+		{
+			if (BestSearchIndex == INDEX_NONE)
+			{
+				BestSearchIndex = SearchIndex;
+			}
+			else if (SearchContexts[SearchIndex].Roles.Num() > SearchContexts[BestSearchIndex].Roles.Num())
+			{
+				BestSearchIndex = SearchIndex;
+			}
+			else if (SearchContexts[SearchIndex].Roles.Num() == SearchContexts[BestSearchIndex].Roles.Num() &&
+				PoseSearchResults[SearchIndex].PoseCost < PoseSearchResults[BestSearchIndex].PoseCost)
+			{
+				BestSearchIndex = SearchIndex;
+			}
+		}
+	}
+
+	if (BestSearchIndex != INDEX_NONE)
+	{
+		SearchResults.SetNum(1, EAllowShrinking::No);
+		SearchResults[0] = InitSearchResult(PoseSearchResults[BestSearchIndex], BestSearchIndex, SearchContexts[BestSearchIndex]);
+	}
+}
+
+// FInteractionSearchContext
+///////////////////////////////////////////////////////////
+bool FInteractionSearchContext::IsValid() const
+{
+	if (Database == nullptr)
+	{
+		return false;
+	}
+
+	const int32 Num = AnimInstances.Num();
+	if (Num < 1)
+	{
+		return false;
+	}
+
+	if (Num != HistoryCollectors.Num())
+	{
+		return false;
+	}
+
+	if (Num != Roles.Num())
+	{
+		return false;
+	}
+
+	for (int32 IndexA = 0; IndexA < Num; ++IndexA)
+	{
+		if (AnimInstances[IndexA] == nullptr)
+		{
+			return false;
+		}
+
+		for (int32 IndexB = IndexA + 1; IndexB < Num; ++IndexB)
+		{
+			if (AnimInstances[IndexA] == AnimInstances[IndexB])
+			{
+				return false;
+			}
+		}
+	}
+
+	for (int32 IndexA = 0; IndexA < Num; ++IndexA)
+	{
+		for (int32 IndexB = IndexA + 1; IndexB < Num; ++IndexB)
+		{
+			if (Roles[IndexA] == Roles[IndexB])
+			{
+				return false;
+			}
+		}
+	}
+
+	for (int32 IndexA = 0; IndexA < Num; ++IndexA)
+	{
+		if (HistoryCollectors[IndexA] == nullptr)
+		{
+			return false;
+		}
+
+		for (int32 IndexB = IndexA + 1; IndexB < Num; ++IndexB)
+		{
+			if (HistoryCollectors[IndexA] == HistoryCollectors[IndexB])
+			{
+				return false;
+			}
+		}
+	}
+
+	return true;
+}
+
+bool FInteractionSearchContext::IsEquivalent(const FInteractionSearchContext& Other) const
+{
+	if (Database != Other.Database)
+	{
+		return false;
+	}
+
+	const int32 Num = AnimInstances.Num();
+	if (Num != Other.AnimInstances.Num())
+	{
+		return false;
+	}
+
+	int32 CommonRoledAnimInstances = 0;
+	for (int32 IndexThis = 0; IndexThis < Num; ++IndexThis)
+	{
+		for (int32 IndexOther = 0; IndexOther < Num; ++IndexOther)
+		{
+			if (AnimInstances[IndexThis] == Other.AnimInstances[IndexOther] &&
+				HistoryCollectors[IndexThis] == Other.HistoryCollectors[IndexOther] &&
+				Roles[IndexThis] == Other.Roles[IndexOther])
+			{
+				++CommonRoledAnimInstances;
+			}
+		}
+	}
+
+	// using >= in case there are duplicated animinstances in this or Other (this->IsValid() should be false!)
+	if (CommonRoledAnimInstances >= Num)
+	{
+		return true;
+	}
+
+	return false;
+}
+
+// FInteractionSearchResult
+///////////////////////////////////////////////////////////
+bool FInteractionSearchResult::operator==(const FInteractionSearchResult& Other) const
+{
+	return static_cast<const FSearchResult&>(*this) == static_cast<const FSearchResult&>(Other) &&
+				FullAlignedActorRootBoneTransforms.Num() == Other.FullAlignedActorRootBoneTransforms.Num() &&
+				0 == FMemory::Memcmp(FullAlignedActorRootBoneTransforms.GetData(), Other.FullAlignedActorRootBoneTransforms.GetData(), FullAlignedActorRootBoneTransforms.Num() * sizeof(FTransform));
+}
+
+// FIslandPreTickFunction
+///////////////////////////////////////////////////////////
+void FInteractionIsland::FPreTickFunction::ExecuteTick(float DeltaTime, enum ELevelTick TickType, ENamedThreads::Type CurrentThread, const FGraphEventRef& MyCompletionGraphEvent)
+{
+	// Called before any skeletal mesh component tick, when there aren't animation jobs flying. No need to FScopeLock Lock(&Mutex);
+	// generating trajectories before running any of the skeletal mesh component ticks
+	for (const FInteractionSearchContext& SearchContext : Island->SearchContexts)
+	{
+		for (int32 Index = 0; Index < SearchContext.AnimInstances.Num(); ++Index)
+		{
+			if (const UAnimInstance* AnimInstance = SearchContext.AnimInstances[Index].Get())
+			{
+				// since FInteractionIsland has a tick dependency with the USkeletalMeshComponent it's safe modify the FAnimNode_PoseSearchHistoryCollector_Base
+				FAnimNode_PoseSearchHistoryCollector_Base* HistoryCollector = const_cast<FAnimNode_PoseSearchHistoryCollector_Base*>(SearchContext.HistoryCollectors[Index]);
+				check(HistoryCollector);
+				HistoryCollector->GenerateTrajectory(AnimInstance);
+			}
+		}
+	}
+}
+
+// FPostTickFunction
+///////////////////////////////////////////////////////////
+void FInteractionIsland::FPostTickFunction::ExecuteTick(float DeltaTime, enum ELevelTick TickType, ENamedThreads::Type CurrentThread, const FGraphEventRef& MyCompletionGraphEvent)
+{
+	// do nothing
+}
+
+// FInteractionIsland
+///////////////////////////////////////////////////////////
+FInteractionIsland::FInteractionIsland(ULevel* Level)
+{
+	PreTickFunction.bCanEverTick = true;
+	PreTickFunction.bStartWithTickEnabled = true;
+	PreTickFunction.TickGroup = ETickingGroup::TG_PrePhysics;
+	PreTickFunction.Island = this;
+	PreTickFunction.SetTickFunctionEnable(true);
+	PreTickFunction.RegisterTickFunction(Level);
+
+	PostTickFunction.bCanEverTick = true;
+	PostTickFunction.bStartWithTickEnabled = true;
+	PostTickFunction.TickGroup = ETickingGroup::TG_PrePhysics;
+	PostTickFunction.SetTickFunctionEnable(true);
+	PostTickFunction.RegisterTickFunction(Level);
+}
+
+FInteractionIsland::~FInteractionIsland()
+{
+	PreTickFunction.UnRegisterTickFunction();
+	PostTickFunction.UnRegisterTickFunction();
+
+	Uninject();
+}
+
+void FInteractionIsland::DebugDraw(const FColor& Color) const
+{
+	// called only by UPoseSearchInteractionSubsystem::Tick so no need to lock SearchResultsMutex to protect the read of SearchContexts
+#if ENABLE_DRAW_DEBUG
+	check(IsInGameThread());
+
+	if (CVarPoseSearchInteractionShowIslands.GetValueOnAnyThread())
+	{
+		for (const FInteractionSearchContext& SearchContext : SearchContexts)
+		{
+			for (int32 Index = 0; Index < SearchContext.AnimInstances.Num(); ++Index)
+			{
+				if (const UAnimInstance* AnimInstance = SearchContext.AnimInstances[Index].Get())
+				{
+					const FVector Position = AnimInstance->GetSkelMeshComponent()->GetComponentLocation();
+					const float BroadPhaseRadius = SearchContext.BroadPhaseRadiuses[Index];
+					DrawDebugCircle(AnimInstance->GetWorld(), Position, BroadPhaseRadius, 40, Color, false, 0.f, SDPG_Foreground, 0.f, FVector::XAxisVector, FVector::YAxisVector, false);
+				}
+			}
+		}
+	}
+#endif // ENABLE_DRAW_DEBUG
+}
+
+void FInteractionIsland::InjectToActor(AActor* Actor)
+{
+	check(IsInGameThread());
+
+	// Called by UPoseSearchInteractionSubsystem::Tick when there aren't animation jobs flying. No need to FScopeLock Lock(&Mutex);
+	if (Actor)
+	{
+		if (UCharacterMovementComponent* CharacterMovementComponent = Actor->GetComponentByClass<UCharacterMovementComponent>())
+		{
+			if (USkeletalMeshComponent* SkeletalMeshComponent = Actor->GetComponentByClass<USkeletalMeshComponent>())
+			{
+				const bool bIsFirstInjectedActor = IsUninjected();
+
+				//	tick order: 
+				//		CharacterMovementComponent(s) ->
+				//			Island.PreTickFunction ->
+				//				first injected actor SkeletalMeshComponent ->
+				//					Island.PostTickFunction ->
+				//						other SkeletalMeshComponent(s)
+				CharacterMovementComponents.AddUnique(CharacterMovementComponent);
+				SkeletalMeshComponents.AddUnique(SkeletalMeshComponent);
+	
+				// making sure that if we add a unique CharacterMovementComponent, we add as well a unique SkeletalMeshComponent
+				// (so we can remove them later on in a consistent fashion)
+				check(CharacterMovementComponents.Num() == SkeletalMeshComponents.Num());
+
+				PreTickFunction.AddPrerequisite(Actor, CharacterMovementComponent->PrimaryComponentTick);
+				SkeletalMeshComponent->PrimaryComponentTick.AddPrerequisite(Actor, PreTickFunction);
+
+				if (bIsFirstInjectedActor)
+				{
+					PostTickFunction.AddPrerequisite(Actor, SkeletalMeshComponent->PrimaryComponentTick);
+				}
+				else 
+				{
+					SkeletalMeshComponent->PrimaryComponentTick.AddPrerequisite(Actor, PostTickFunction);
+				}
+			}
+			else
+			{
+				UE_LOG(LogPoseSearch, Error, TEXT("FInteractionIsland::InjectToActor requires AActor %s to have a USkeletalMeshComponent to work!"), *Actor->GetName());
+			}
+		}
+		else
+		{
+			UE_LOG(LogPoseSearch, Error, TEXT("FInteractionIsland::InjectToActor requires AActor %s to have a UCharacterMovementComponent to work!"), *Actor->GetName());
+		}
+	}
+}
+
+void FInteractionIsland::AddSearchContext(const FInteractionSearchContext& SearchContext)
+{
+#if DO_CHECK
+	check(SearchContext.IsValid());
+	check(IsInGameThread());
+
+	for (const FInteractionSearchContext& ContainedSearchContext : SearchContexts)
+	{
+		check(!ContainedSearchContext.IsEquivalent(SearchContext));
+	}
+#endif // DO_CHECK
+	SearchContexts.Add(SearchContext);
+}
+
+void FInteractionIsland::Uninject()
+{
+	// Called by UPoseSearchInteractionSubsystem::Tick when there aren't animation jobs flying. No need to FScopeLock Lock(&Mutex);
+	check(IsInGameThread());
+
+	check(CharacterMovementComponents.Num() == SkeletalMeshComponents.Num());
+
+	for (int32 ActorIndex = 0; ActorIndex < CharacterMovementComponents.Num(); ++ActorIndex)
+	{
+		UCharacterMovementComponent* CharacterMovementComponent = CharacterMovementComponents[ActorIndex].Get();
+		USkeletalMeshComponent* SkeletalMeshComponent = SkeletalMeshComponents[ActorIndex].Get();
+
+		if (CharacterMovementComponent)
+		{
+			check(SkeletalMeshComponent);
+
+			AActor* Actor = CharacterMovementComponent->GetOwner();
+			check(Actor);
+
+			PreTickFunction.RemovePrerequisite(Actor, CharacterMovementComponent->PrimaryComponentTick);
+			SkeletalMeshComponent->PrimaryComponentTick.RemovePrerequisite(Actor, PreTickFunction);
+
+			const bool bIsFirstInjectedActor = ActorIndex == 0;
+			if (bIsFirstInjectedActor)
+			{
+				PostTickFunction.RemovePrerequisite(Actor, SkeletalMeshComponent->PrimaryComponentTick);
+			}
+			else
+			{
+				SkeletalMeshComponent->PrimaryComponentTick.RemovePrerequisite(Actor, PostTickFunction);
+			}
+		}
+		else
+		{
+			check(!SkeletalMeshComponent);
+		}
+	}
+
+	CharacterMovementComponents.Reset();
+	SkeletalMeshComponents.Reset();
+
+	SearchContexts.Reset();
+	SearchResults.Reset();
+	bSearchPerfomed = false;
+}
+
+bool FInteractionIsland::IsUninjected() const
+{
+	return SkeletalMeshComponents.IsEmpty();
+}
+
+bool FInteractionIsland::DoSearch_AnyThread(UObject* AnimInstance, const FPoseSearchContinuingProperties& ContinuingProperties, FPoseSearchInteractionBlueprintResult& Result)
+{
+	bool bDoPerfomSearch = false;
+	{
+		// thread safety note!
+		// goal:	avoiding deadlock between SearchResultsMutex lock and waiting for UAnimInstance::HandleExistingParallelEvaluationTask.
+		// why:		UPoseSearchLibrary::MotionMatch could call via AnimInstance GetProxyOnAnyThread<FAnimInstanceProxy>() that, if on GameThread, 
+		//			could call UAnimInstance::HandleExistingParallelEvaluationTask 
+		// fix:		avoid UPoseSearchLibrary::MotionMatch calls wrapped by any lock, at the cost of eventually (by design should be NEVER) performing the searches twice
+		//			By design we should inject ticks dependencies (added by FInteractionIsland::InjectToActor via AddTickPrerequisiteComponent), so concurrently 
+		//			fly of FInteractionIsland within the same island that requires searches is forbidden
+		FScopeLock Lock(&SearchResultsMutex);
+		bDoPerfomSearch = !bSearchPerfomed;
+	}
+
+	if (bDoPerfomSearch)
+	{
+		FMemMark Mark(FMemStack::Get());
+
+		TArray<UAnimInstance*, TInlineAllocator<UE::PoseSearch::PreallocatedRolesNum, TMemStackAllocator<>>> AnimInstances;
+		TArray<const UE::PoseSearch::IPoseHistory*, TInlineAllocator<UE::PoseSearch::PreallocatedRolesNum, TMemStackAllocator<>>> PoseHistories;
+		TArray<UE::PoseSearch::FSearchResult, TMemStackAllocator<>> PoseSearchResults;
+
+		// SearchContexts are modified only by UPoseSearchInteractionSubsystem::Tick and constant otherwise, so it's safe to access them in a threaded enviroment without locks
+		PoseSearchResults.SetNum(SearchContexts.Num());
+
+		for (int32 SearchIndex = 0; SearchIndex < SearchContexts.Num(); ++SearchIndex)
+		{
+			const FInteractionSearchContext& SearchContext = SearchContexts[SearchIndex];
+			const UPoseSearchDatabase* Database = SearchContext.Database.Get();
+			if (!Database)
+			{
+				UE_LOG(LogPoseSearch, Error, TEXT("FInteractionIsland::DoSearch_AnyThread invalid context database"));
+				return false;
+			}
+
+			if (!Database->Schema)
+			{
+				UE_LOG(LogPoseSearch, Error, TEXT("FInteractionIsland::DoSearch_AnyThread invalid schema for context database %s"), *Database->GetName());
+				return false;
+			}
+
+			AnimInstances.Reset();
+			for (const TWeakObjectPtr<UAnimInstance>& AnimInstancePtr : SearchContext.AnimInstances)
+			{
+				UAnimInstance* SearchContextAnimInstance = AnimInstancePtr.Get();
+				if (!SearchContextAnimInstance)
+				{
+					UE_LOG(LogPoseSearch, Error, TEXT("FInteractionIsland::DoSearch_AnyThread null anim instance"));
+					return false;
+				}
+
+				AnimInstances.Add(SearchContextAnimInstance);
+			}
+
+			PoseHistories.Reset();
+			for (const FAnimNode_PoseSearchHistoryCollector_Base* HistoryCollector : SearchContext.HistoryCollectors)
+			{
+				PoseHistories.Add(&HistoryCollector->GetPoseHistory());
+			}
+
+			const UObject* AssetsToSearch[] = { Database };
+			FPoseSearchFutureProperties PoseSearchFutureProperties;
+
+			FPoseSearchContinuingProperties ContinuingPropertiesToUse = SearchContext.ContinuingProperties;
+
+			// @todo: implemennt this carefully, since it'll cause thread safety concerns depending on which actors is calling the query first 
+			// and if their animations get integrated by different play rates
+			
+			//if (ContinuingProperties.PlayingAsset)
+			//{
+			//	if (const UMultiAnimAsset* MultiAnimAsset = Cast<UMultiAnimAsset>(ContinuingProperties.PlayingAsset))
+			//	{
+			//		// @todo: if they are compatible... etcetc MultiAnimAsset->GetNumRoles
+			//		ContinuingPropertiesToUse = ContinuingProperties;
+			//	}
+			//	else if (const UMultiAnimAsset* SearchContextMultiAnimAsset = Cast<UMultiAnimAsset>(SearchContext.ContinuingProperties.PlayingAsset))
+			//	{
+			//		for (int32 RoleIndex = 0; RoleIndex < SearchContextMultiAnimAsset->GetNumRoles(); ++RoleIndex)
+			//		{
+			//			if (SearchContextMultiAnimAsset->GetAnimationAsset(SearchContextMultiAnimAsset->GetRole(RoleIndex)) == ContinuingProperties.PlayingAsset)
+			//			{
+			//				ContinuingPropertiesToUse.PlayingAssetAccumulatedTime = ContinuingProperties.PlayingAssetAccumulatedTime;
+			//				break;
+			//			}
+			//		}
+			//	}
+			//	else if (!SearchContext.ContinuingProperties.PlayingAsset)
+			//	{
+			//		// @todo: should we scan all the database iterating over the UMultiAnimAsset to look for ContinuingProperties.PlayingAsset?
+			//	}
+			//	else
+			//	{
+			//	}
+			//}
+
+			// @todo: we could perform multiple UPoseSearchLibrary::MotionMatch in parallel!
+			const UE::PoseSearch::FSearchResult PoseSearchResult = UPoseSearchLibrary::MotionMatch(AnimInstances, SearchContext.Roles,
+				PoseHistories, AssetsToSearch, ContinuingPropertiesToUse, PoseSearchFutureProperties);
+
+			if (PoseSearchResult.PoseCost.GetTotalCost() < SearchContext.MaxCost)
+			{
+				PoseSearchResults[SearchIndex] = PoseSearchResult;
+			}
+		}
+
+		if (!PoseSearchResults.IsEmpty())
+		{
+			// locking to update SearchResults and bSearchPerfomed
+			FScopeLock Lock(&SearchResultsMutex);
+
+			if (!bSearchPerfomed)
+			{
+				InitSearchResults(SearchResults, PoseSearchResults, SearchContexts);
+
+				bSearchPerfomed = true;
+			}
+			else
+			{
+				UE_LOG(LogPoseSearch, Warning, TEXT("FInteractionIsland::DoSearch_AnyThread performance warning: performed duplicated search :("));
+				
+#if DO_CHECK
+				TArray<FInteractionSearchResult> CompareSearchResults;
+				InitSearchResults(CompareSearchResults, PoseSearchResults, SearchContexts);
+
+				if (CompareSearchResults != SearchResults)
+				{
+					UE_LOG(LogPoseSearch, Error, TEXT("FInteractionIsland::DoSearch_AnyThread duplicated search differs from original one. Searches are NOT deterministic!"));
+				}
+#endif // DO_CHECK
+			}
+	
+			// calling this funtion within this scope since we already locked SearchResultsMutex
+			return GetResult_AnyThread(AnimInstance, Result);
+		}
+	}
+
+	return GetResult_AnyThread(AnimInstance, Result);
+}
+
+bool FInteractionIsland::GetResult_AnyThread(UObject* AnimInstance, FPoseSearchInteractionBlueprintResult& Result)
+{
+	// locking to read SearchResults
+	FScopeLock Lock(&SearchResultsMutex);
+
+	// looking for AnimInstance in SearchResults to fill up Result
+	for (const FInteractionSearchResult& SearchResult : SearchResults)
+	{
+		const FInteractionSearchContext& SearchContext = SearchContexts[SearchResult.SearchIndex];
+		for (int32 AnimInstanceIndex = 0; AnimInstanceIndex < SearchContext.AnimInstances.Num(); ++AnimInstanceIndex)
+		{
+			if (SearchContext.AnimInstances[AnimInstanceIndex].Get() == AnimInstance)
+			{
+				const UPoseSearchDatabase* Database = SearchResult.Database.Get();
+				check(Database);
+
+				const UE::PoseSearch::FSearchIndexAsset* SearchIndexAsset = SearchResult.GetSearchIndexAsset();
+				check(SearchIndexAsset);
+
+				const FPoseSearchDatabaseAnimationAssetBase* DatabaseAnimationAssetBase = Database->GetDatabaseAnimationAsset<FPoseSearchDatabaseAnimationAssetBase>(*SearchIndexAsset);
+				check(DatabaseAnimationAssetBase);
+
+				const UE::PoseSearch::FRole Role = SearchContext.Roles[AnimInstanceIndex];
+
+				Result.SelectedAnimation = DatabaseAnimationAssetBase->GetAnimationAsset();
+				Result.SelectedTime = SearchResult.AssetTime;
+				Result.bIsContinuingPoseSearch = SearchResult.bIsContinuingPoseSearch;
+				Result.bLoop = SearchIndexAsset->IsLooping();
+				Result.bIsMirrored = SearchIndexAsset->IsMirrored();
+				Result.BlendParameters = SearchIndexAsset->GetBlendParameters();
+				Result.SelectedDatabase = Database;
+				Result.SearchCost = SearchResult.PoseCost.GetTotalCost();
+				Result.Role = Role;
+
+				if (SearchResult.FullAlignedActorRootBoneTransforms.IsValidIndex(AnimInstanceIndex))
+				{
+					Result.FullAlignedActorRootBoneTransform = SearchResult.FullAlignedActorRootBoneTransforms[AnimInstanceIndex];
+				}
+
+				//Result.bIsFromContinuingPlaying = SearchResult.bIsContinuingPoseSearch;
+
+				// figuring out the WantedPlayRate
+				Result.WantedPlayRate = 1.f;
+				//if (Future.Animation && Future.IntervalTime > 0.f)
+				//{
+				//	if (const UPoseSearchFeatureChannel_PermutationTime* PermutationTimeChannel = Database->Schema->FindFirstChannelOfType<UPoseSearchFeatureChannel_PermutationTime>())
+				//	{
+				//		const FSearchIndex& SearchIndex = Database->GetSearchIndex();
+				//		if (!SearchIndex.IsValuesEmpty())
+				//		{
+				//			TConstArrayView<float> ResultData = Database->GetSearchIndex().GetPoseValues(SearchResult.PoseIdx);
+				//			const float ActualIntervalTime = PermutationTimeChannel->GetPermutationTime(ResultData);
+				//			ProviderResult.WantedPlayRate = ActualIntervalTime / Future.IntervalTime;
+				//		}
+				//	}
+				//}
+
+				// we found our AnimInstance: we can stop searching
+				return true;
+			}
+		}
+	}
+
+	return false;
+}
+
+const FInteractionSearchResult* FInteractionIsland::FindSearchResult(const FInteractionSearchContext& SearchContext) const
+{
+	// called only by UPoseSearchInteractionSubsystem::Tick via UPoseSearchInteractionSubsystem::PopulateContinuingProperties so no need to lock SearchResultsMutex to protect the read of SearchResults
+	check(IsInGameThread());
+
+	// searching for InSearchContext in all the SearchContexts referenced by valid active SearchResults
+	for (const FInteractionSearchResult& SearchResult : SearchResults)
+	{
+		const FInteractionSearchContext& LocalSearchContext = SearchContexts[SearchResult.SearchIndex];
+		if (LocalSearchContext.Database == SearchContext.Database &&
+			LocalSearchContext.AnimInstances == SearchContext.AnimInstances &&
+			LocalSearchContext.Roles == SearchContext.Roles)
+		{
+			return &SearchResult;
+		}
+	}
+	return nullptr;
+}
+
+} // namespace UE::PoseSearch
diff --git a/Engine/Plugins/Animation/PoseSearch/Source/Runtime/Private/PoseSearchInteractionLibrary.cpp b/Engine/Plugins/Animation/PoseSearch/Source/Runtime/Private/PoseSearchInteractionLibrary.cpp
new file mode 100644
--- /dev/null
+++ b/Engine/Plugins/Animation/PoseSearch/Source/Runtime/Private/PoseSearchInteractionLibrary.cpp
@@ -0,0 +1,42 @@
+// Copyright Epic Games, Inc. All Rights Reserved.
+
+#include "PoseSearch/PoseSearchInteractionLibrary.h"
+#include "Animation/AnimMontage.h"
+#include "PoseSearch/PoseSearchInteractionSubsystem.h"
+#include "Engine/World.h"
+
+FPoseSearchInteractionBlueprintResult UPoseSearchInteractionLibrary::MotionMatchInteraction_Pure(TArray<FPoseSearchInteractionAvailability> Availabilities, UObject* AnimInstance, FPoseSearchContinuingProperties ContinuingProperties, FName PoseHistoryName, bool bValidateResultAgainstAvailabilities)
+{
+	FPoseSearchInteractionBlueprintResult Result;
+	if (UPoseSearchInteractionSubsystem* InteractionSubsystem = UPoseSearchInteractionSubsystem::GetSubsystem_AnyThread(AnimInstance))
+	{
+		InteractionSubsystem->Query_AnyThread(Availabilities, AnimInstance, ContinuingProperties, Result, PoseHistoryName, nullptr, bValidateResultAgainstAvailabilities);
+	}
+	return Result;
+}
+
+FPoseSearchInteractionBlueprintResult UPoseSearchInteractionLibrary::MotionMatchInteraction(TArray<FPoseSearchInteractionAvailability> Availabilities, UObject* AnimInstance, FPoseSearchContinuingProperties ContinuingProperties, FName PoseHistoryName, bool bValidateResultAgainstAvailabilities)
+{
+	return MotionMatchInteraction_Pure(Availabilities, AnimInstance, ContinuingProperties, PoseHistoryName, bValidateResultAgainstAvailabilities);
+}
+
+FPoseSearchInteractionBlueprintResult UPoseSearchInteractionLibrary::MotionMatchInteraction(const TArrayView<const FPoseSearchInteractionAvailability> Availabilities, UObject* AnimInstance, const FPoseSearchContinuingProperties& ContinuingProperties, const FAnimNode_PoseSearchHistoryCollector_Base* HistoryCollector, bool bValidateResultAgainstAvailabilities)
+{
+	FPoseSearchInteractionBlueprintResult Result;
+	if (UPoseSearchInteractionSubsystem* InteractionSubsystem = UPoseSearchInteractionSubsystem::GetSubsystem_AnyThread(AnimInstance))
+	{
+		InteractionSubsystem->Query_AnyThread(Availabilities, AnimInstance, ContinuingProperties, Result, FName(), HistoryCollector, bValidateResultAgainstAvailabilities);
+	}
+	return Result;
+}
+
+FPoseSearchContinuingProperties UPoseSearchInteractionLibrary::GetMontageContinuingProperties(UAnimInstance* AnimInstance)
+{
+	FPoseSearchContinuingProperties ContinuingProperties;
+	if (const FAnimMontageInstance* AnimMontageInstance = AnimInstance->GetActiveMontageInstance())
+	{
+		ContinuingProperties.PlayingAsset = AnimMontageInstance->Montage;
+		ContinuingProperties.PlayingAssetAccumulatedTime = AnimMontageInstance->DeltaTimeRecord.GetPrevious();
+	}
+	return ContinuingProperties;
+}
\ No newline at end of file
diff --git a/Engine/Plugins/Animation/PoseSearch/Source/Runtime/Private/PoseSearchInteractionSubsystem.cpp b/Engine/Plugins/Animation/PoseSearch/Source/Runtime/Private/PoseSearchInteractionSubsystem.cpp
new file mode 100644
--- /dev/null
+++ b/Engine/Plugins/Animation/PoseSearch/Source/Runtime/Private/PoseSearchInteractionSubsystem.cpp
@@ -0,0 +1,786 @@
+// Copyright Epic Games, Inc. All Rights Reserved.
+
+#include "PoseSearch/PoseSearchInteractionSubsystem.h"
+#include "PoseSearch/PoseSearchInteractionIsland.h"
+#include "PoseSearch/PoseSearchDatabase.h"
+#include "PoseSearch/PoseSearchRole.h"
+#include "PoseSearch/PoseSearchSchema.h"
+#include "VisualLogger/VisualLogger.h"
+
+namespace UE::PoseSearch
+{
+	FString FInteractionAvailabilityEx::GetPoseHistoryName() const
+	{
+		if (HistoryCollector)
+		{
+			return "HistoryProvider";
+		}
+		return PoseHistoryName.ToString();
+	}
+	
+	const FAnimNode_PoseSearchHistoryCollector_Base* FInteractionAvailabilityEx::GetHistoryCollector(const UAnimInstance* AnimInstance) const
+	{
+		if (HistoryCollector)
+		{
+			return HistoryCollector;
+		}
+		return UPoseSearchLibrary::FindPoseHistoryNode(PoseHistoryName, AnimInstance);
+	}
+
+	struct FInteractionSearchContextGroup
+	{
+		bool Contains(const FInteractionSearchContext& SearchContext) const
+		{
+			for (const TWeakObjectPtr<UAnimInstance>& AnimInstance : SearchContext.AnimInstances)
+			{
+				if (AnimInstances.Find(AnimInstance.Get()))
+				{
+					return true;
+				}
+			}
+			return false;
+		}
+
+		void Add(const FInteractionSearchContext& SearchContext, int32 SearchContextIndex)
+		{
+			for (const TWeakObjectPtr<UAnimInstance>& AnimInstance : SearchContext.AnimInstances)
+			{
+				if (AnimInstance.Get())
+				{
+					AnimInstances.Add(AnimInstance.Get());
+				}
+			}
+
+			SearchContextsIndices.Add(SearchContextIndex);
+		}
+
+		void Merge(const FInteractionSearchContextGroup& SearchContextGroup)
+		{
+			for (const UAnimInstance* AnimInstance : SearchContextGroup.AnimInstances)
+			{
+				check(AnimInstance);
+				AnimInstances.Add(AnimInstance);
+			}
+
+			for (int32 SearchContextsIndex : SearchContextGroup.SearchContextsIndices)
+			{
+				SearchContextsIndices.Add(SearchContextsIndex);
+			}
+		}
+
+		TSet<const UAnimInstance*, DefaultKeyFuncs<const UAnimInstance*>, TInlineSetAllocator<UE::PoseSearch::PreallocatedRolesNum>> AnimInstances;
+		TArray<int32, TInlineAllocator<PreallocatedSearchesNum>> SearchContextsIndices;
+	};
+
+	struct FRoledAnimInstance
+	{
+		UAnimInstance* AnimInstance = nullptr;
+		UE::PoseSearch::FRole Role = UE::PoseSearch::DefaultRole;
+		const FAnimNode_PoseSearchHistoryCollector_Base* HistoryCollector = nullptr;
+		float BroadPhaseRadius = 0.f;
+		float MaxCost = MAX_flt;
+	};
+
+	// @todo: inherit from a TSet if it gets too slow 
+	// sorted array of FRoledAnimInstance(s). Sorted by UAnimInstance to prevent "alternating" tuples between different frames like
+	// ([AnimInstanceA, RoleA], [AnimInstanceB, RoleB]) vs ([AnimInstanceB, RoleB], [AnimInstanceA, RoleA]) 
+	// that ultimately represent the same search, but endign up with different SearchIDs
+	struct FRoledAnimInstances
+	{
+		void AddRoledAnimInstace(UAnimInstance* AnimInstance, UE::PoseSearch::FRole Role, const FAnimNode_PoseSearchHistoryCollector_Base* HistoryCollector, float BroadPhaseRadius, float MaxCost)
+		{
+			check(AnimInstance && HistoryCollector);
+
+			for (FRoledAnimInstance& RoledAnimInstance : RoledAnimInstances)
+			{
+				if (RoledAnimInstance.AnimInstance == AnimInstance && RoledAnimInstance.HistoryCollector == HistoryCollector && RoledAnimInstance.Role == Role)
+				{
+					// we already found the FRoledAnimInstances. Updating the BroadPhaseRadius and MaxCost accordingly
+					RoledAnimInstance.BroadPhaseRadius = FMath::Max(RoledAnimInstance.BroadPhaseRadius, BroadPhaseRadius);
+					RoledAnimInstance.MaxCost = FMath::Min(RoledAnimInstance.MaxCost, MaxCost);
+					return;
+				}
+			}
+
+			// we couldn't find the FRoledAnimInstances. Let's add a new one
+			FRoledAnimInstance& NewRoledAnimInstance = RoledAnimInstances.AddDefaulted_GetRef();
+			NewRoledAnimInstance.AnimInstance = AnimInstance;
+			NewRoledAnimInstance.HistoryCollector = HistoryCollector;
+			NewRoledAnimInstance.Role = Role;
+			NewRoledAnimInstance.BroadPhaseRadius = BroadPhaseRadius;
+			NewRoledAnimInstance.MaxCost = MaxCost;
+
+			// @todo: use a better datastructure to keep things FRoledAnimInstance(s) sorted? Or use insertion sort?
+			RoledAnimInstances.Sort([](const FRoledAnimInstance& A, const FRoledAnimInstance& B)
+			{
+				// we're sorting by AnimInstance address in memory, since searches neeed to be consistent 
+				// between frames on the same machine, not between replicated machines
+				return B.AnimInstance < A.AnimInstance;
+			});
+		}
+
+		TArrayView<const FRoledAnimInstance> GetDataView() const
+		{
+			return MakeArrayView(RoledAnimInstances.GetData(), RoledAnimInstances.Num());
+		}
+
+	private:
+		TArray<FRoledAnimInstance, TInlineAllocator<UE::PoseSearch::PreallocatedRolesNum, TMemStackAllocator<>>> RoledAnimInstances;
+	};
+
+	typedef TMap<const UPoseSearchDatabase*, FRoledAnimInstances, TInlineSetAllocator<PreallocatedSearchesNum, UE::PoseSearch::TMemStackSetAllocator<>>> FDatabaseToRoledAnimInstances;
+	typedef TPair<const UPoseSearchDatabase*, FRoledAnimInstances> FDatabaseToRoledAnimInstancesPair;
+
+	static void ProcessAvailabilityRequestsMap(const FAvailabilityRequestsMap& AvailabilityRequestsMap, UWorld* SubsystemWorld, FDatabaseToRoledAnimInstances& DatabaseToRoledAnimInstances)
+	{
+		for (FAvailabilityRequestsMap::TConstIterator Iter = AvailabilityRequestsMap.CreateConstIterator(); Iter; ++Iter)
+		{
+			// making sure the weak pointer is still valid
+			UAnimInstance* AnimInstance = Cast<UAnimInstance>(Iter->Key.Get());
+			if (!AnimInstance)
+			{
+				UE_LOG(LogPoseSearch, Log, TEXT("ProcessAvailabilityRequestsMap null anim instance. The associated character got removed from the world"));
+				continue;
+			}
+
+			if (!AnimInstance->GetWorld())
+			{
+				UE_LOG(LogPoseSearch, Error, TEXT("ProcessAvailabilityRequestsMap AnimInstance %s is not in a world!"), *AnimInstance->GetName());
+				continue;
+			}
+
+			if (SubsystemWorld != AnimInstance->GetWorld())
+			{
+				UE_LOG(LogPoseSearch, Error, TEXT("ProcessAvailabilityRequestsMap AnimInstance %s is from World %s, and supposed to be in World @s!"), *AnimInstance->GetName(), *AnimInstance->GetWorld()->GetName(), *SubsystemWorld->GetName());
+				continue;
+			}
+
+			for (const FInteractionAvailabilityEx& AvailabilityRequest : Iter->Value)
+			{
+				if (!AvailabilityRequest.Database)
+				{
+					UE_LOG(LogPoseSearch, Log, TEXT("ProcessAvailabilityRequestsMap null AvailabilityRequest.Database"));
+					continue;
+				}
+
+				if (!AvailabilityRequest.Database->Schema)
+				{
+					UE_LOG(LogPoseSearch, Error, TEXT("ProcessAvailabilityRequestsMap null Schema for Database %s"), *AvailabilityRequest.Database->GetName());
+					continue;
+				}
+
+				const FAnimNode_PoseSearchHistoryCollector_Base* HistoryCollector = AvailabilityRequest.GetHistoryCollector(AnimInstance);
+				if (!HistoryCollector)
+				{
+					UE_LOG(LogPoseSearch, Error, TEXT("ProcessAvailabilityRequestsMap couldn't find PoseHistory %s for AnimInstance %s"), *AvailabilityRequest.GetPoseHistoryName(), *AnimInstance->GetName());
+					continue;
+				}
+
+				FRoledAnimInstances& RoledAnimInstances = DatabaseToRoledAnimInstances.FindOrAdd(AvailabilityRequest.Database);
+				if (AvailabilityRequest.RolesFilter.IsEmpty())
+				{
+					// adding ALL the possible roles from the database:
+					for (const FPoseSearchRoledSkeleton& RoledSkeleton : AvailabilityRequest.Database->Schema->GetRoledSkeletons())
+					{
+						RoledAnimInstances.AddRoledAnimInstace(AnimInstance, RoledSkeleton.Role, HistoryCollector, AvailabilityRequest.BroadPhaseRadius, AvailabilityRequest.MaxCost);
+					}
+				}
+				else
+				{
+					for (const UE::PoseSearch::FRole& Role : AvailabilityRequest.RolesFilter)
+					{
+						if (AvailabilityRequest.Database->Schema->GetRoledSkeleton(Role))
+						{
+							RoledAnimInstances.AddRoledAnimInstace(AnimInstance, Role, HistoryCollector, AvailabilityRequest.BroadPhaseRadius, AvailabilityRequest.MaxCost);
+						}
+						else
+						{
+							UE_LOG(LogPoseSearch, Warning, TEXT("ProcessAvailabilityRequestsMap unsupported Role %s for Database %s"), *Role.ToString(), *AvailabilityRequest.Database->GetName());
+						}
+					}
+				}
+			}
+		}
+	}
+
+	template <typename DataType, typename EvaluateCombinationType>
+	static void GeneratePermutationsRecursive(const TArrayView<DataType> Data, int32 DataIndex, TArrayView<int32> Combination, int32 CombinationIndex, EvaluateCombinationType EvaluateCombination)
+	{
+		if (CombinationIndex == Combination.Num())
+		{ 
+			EvaluateCombination(Data, Combination);
+		}
+		else if (DataIndex < Data.Num())
+		{
+			Combination[CombinationIndex] = DataIndex;
+			GeneratePermutationsRecursive(Data, DataIndex + 1, Combination, CombinationIndex + 1, EvaluateCombination);
+			GeneratePermutationsRecursive(Data, DataIndex + 1, Combination, CombinationIndex, EvaluateCombination);
+		}
+	} 
+
+	template <typename DataType, typename EvaluateCombinationType>
+	static void GeneratePermutations(const TArrayView<DataType> Data, int32 CombinationCardinality, EvaluateCombinationType EvaluateCombination)
+	{
+		TArray<int32, TInlineAllocator<UE::PoseSearch::PreallocatedRolesNum>> Combination;
+		Combination.SetNum(CombinationCardinality);
+		GeneratePermutationsRecursive(Data, 0, Combination, 0, EvaluateCombination);
+	}
+}
+
+FCriticalSection UPoseSearchInteractionSubsystem::RetrieveSubsystemMutex;
+
+UE::PoseSearch::FInteractionIsland& UPoseSearchInteractionSubsystem::CreateIsland()
+{
+	return *Islands.Add_GetRef(new UE::PoseSearch::FInteractionIsland(ToRawPtr(GetWorld()->PersistentLevel)));
+}
+
+void UPoseSearchInteractionSubsystem::DestroyIsland(int32 Index)
+{
+	delete Islands[Index];
+	Islands.RemoveAt(Index);
+}
+
+UE::PoseSearch::FInteractionIsland& UPoseSearchInteractionSubsystem::GetAvailableIsland()
+{
+	using namespace UE::PoseSearch;
+
+	for (FInteractionIsland* Island : Islands)
+	{
+		if (Island->IsUninjected())
+		{
+			return *Island;
+		}
+	}
+
+	return CreateIsland();
+}
+
+void UPoseSearchInteractionSubsystem::DestroyAllIslands()
+{
+	for (int32 IslandIndex = Islands.Num() - 1; IslandIndex >= 0; --IslandIndex)
+	{
+		DestroyIsland(IslandIndex);
+	}
+}
+
+void UPoseSearchInteractionSubsystem::UninjectAllIslands()
+{
+	using namespace UE::PoseSearch;
+
+	check(IsInGameThread());
+
+	for (FInteractionIsland* Island : Islands)
+	{
+		Island->Uninject();
+	}
+}
+
+bool UPoseSearchInteractionSubsystem::ValidateAllIslands() const
+{
+#if DO_CHECK
+	using namespace UE::PoseSearch;
+
+	TSet<TWeakObjectPtr<UCharacterMovementComponent>> CharacterMovementComponents;
+	TSet<TWeakObjectPtr<USkeletalMeshComponent>> SkeletalMeshComponents;
+
+	bool bAlreadyInSet = false;
+	for (const FInteractionIsland* Island : Islands)
+	{
+		for (const TWeakObjectPtr<UCharacterMovementComponent>& CharacterMovementComponent : Island->GetCharacterMovementComponents())
+		{
+			CharacterMovementComponents.Add(CharacterMovementComponent, &bAlreadyInSet);
+			if (bAlreadyInSet)
+			{
+				return false;
+			}
+		}
+
+		for (const TWeakObjectPtr<USkeletalMeshComponent>& SkeletalMeshComponent : Island->GetSkeletalMeshComponents())
+		{
+			SkeletalMeshComponents.Add(SkeletalMeshComponent, &bAlreadyInSet);
+			if (bAlreadyInSet)
+			{
+				return false;
+			}
+		}
+	}
+#endif // DO_CHECK
+	return true;
+}
+
+void UPoseSearchInteractionSubsystem::PopulateContinuingProperties(UE::PoseSearch::FInteractionSearchContext& SearchContext, float DeltaSeconds) const
+{
+	using namespace UE::PoseSearch;
+
+	check(IsInGameThread());
+
+	// searching this SearchContext in all the islands to initialize its continuing pose
+	for (const FInteractionIsland* Island : Islands)
+	{
+		if (const FSearchResult* SearchResult = Island->FindSearchResult(SearchContext))
+		{
+			// is still valid...
+			if (SearchResult->IsValid())
+			{
+				if (const UE::PoseSearch::FSearchIndexAsset* SearchIndexAsset = SearchResult->GetSearchIndexAsset())
+				{
+					if (const FPoseSearchDatabaseAnimationAssetBase* DatabaseAsset = SearchResult->Database->GetDatabaseAnimationAsset<FPoseSearchDatabaseAnimationAssetBase>(*SearchIndexAsset))
+					{
+						SearchContext.ContinuingProperties.PlayingAsset = DatabaseAsset->GetAnimationAsset();
+						SearchContext.ContinuingProperties.PlayingAssetAccumulatedTime = SearchResult->AssetTime + DeltaSeconds;
+					}
+				}
+			}
+			break;
+		}
+	}
+}
+
+UE::PoseSearch::FInteractionIsland* UPoseSearchInteractionSubsystem::FindIsland(UObject* InAnimInstance)
+{
+	using namespace UE::PoseSearch;
+
+	if (InAnimInstance)
+	{
+		if (UAnimInstance* AnimInstance = Cast<UAnimInstance>(InAnimInstance))
+		{
+			if (AActor* Actor = AnimInstance->GetOwningActor())
+			{
+				if (USkeletalMeshComponent* SkeletalMeshComponent = Actor->GetComponentByClass<USkeletalMeshComponent>())
+				{
+					for (FInteractionIsland* Island : Islands)
+					{
+						for (const TWeakObjectPtr<USkeletalMeshComponent>& IslandSkeletalMeshComponent : Island->GetSkeletalMeshComponents())
+						{
+							if (IslandSkeletalMeshComponent.Get() == SkeletalMeshComponent)
+							{
+								return Island;
+							}
+						}
+					}
+				}
+			}
+		}
+	}
+	return nullptr;
+}
+
+UPoseSearchInteractionSubsystem* UPoseSearchInteractionSubsystem::GetSubsystem_AnyThread(UObject* AnimInstance)
+{
+	if (AnimInstance)
+	{
+		if (UWorld* World = AnimInstance->GetWorld())
+		{
+			FScopeLock Lock(&RetrieveSubsystemMutex);
+
+			// don't create a subsystem from any thread
+			if (World->HasSubsystem<UPoseSearchInteractionSubsystem>())
+			{
+				return World->GetSubsystem<UPoseSearchInteractionSubsystem>();
+			}
+		}
+	}
+	return nullptr;
+}
+
+void UPoseSearchInteractionSubsystem::Initialize(FSubsystemCollectionBase& Collection)
+{
+	Super::Initialize(Collection);
+}
+
+void UPoseSearchInteractionSubsystem::Deinitialize()
+{
+	DestroyAllIslands();
+
+	Super::Deinitialize();
+}
+
+void UPoseSearchInteractionSubsystem::Tick(float DeltaSeconds)
+{
+	using namespace UE::PoseSearch;
+
+	Super::Tick(DeltaSeconds);
+
+	if (AvailabilityRequestsMap.IsEmpty() && Islands.IsEmpty())
+	{
+		// nothing to do. early out
+		return;
+	}
+
+	check(IsInGameThread());
+	check(GetWorld());
+
+	// FScopeLock Lock(&AvailabilityRequestsMapMutex); is not necessary since UPoseSearchInteractionSubsystem gets ticked outside the parallel animation jobs
+	FMemMark Mark(FMemStack::Get());
+	
+	// processing the AvailabilityRequestsMap to fill up a map of Databases pointing to an array of all the AnimInstances with related roles
+	FDatabaseToRoledAnimInstances DatabaseToRoledAnimInstances;
+	ProcessAvailabilityRequestsMap(AvailabilityRequestsMap, GetWorld(), DatabaseToRoledAnimInstances);
+
+	// for each database now we try to create all the possible permutations of the roled anim instances
+	// for example, given a database set up with assets for 2 characters interactions with roles RoleA and RoleB
+	// and 2 anim instances, all of them willing to partecipate in the 2 characters interaction with both roles RoleA and RoleB:
+	// CharA could be taking RoleA and RoleB,
+	// CharB could be taking RoleA and RoleB,
+	// we generate all the permutations from the array of options:
+	// CharA/RoleA, CharA/RoleB, CharB/RoleA, CharB/RoleB
+	//
+	// and we prune the invalid tuples:
+	//
+	// CharA/RoleA - CharA/RoleB -> invalid because of same CharA
+	// CharA/RoleA - CharB/RoleA -> invalid because of same RoleA
+	// CharA/RoleA - CharB/RoleB -> VALID!
+	//
+	// CharA/RoleB - CharB/RoleA -> VALID!
+	// CharA/RoleB - CharB/RoleB -> invalid because of same RoleB
+	//
+	// CharB/RoleA - CharB/RoleB -> invalid because of same CharB
+
+	TSet<const UAnimInstance*, DefaultKeyFuncs<const UAnimInstance*>, TInlineSetAllocator<UE::PoseSearch::PreallocatedRolesNum>> AnimInstances;
+	TSet<UE::PoseSearch::FRole, DefaultKeyFuncs<UE::PoseSearch::FRole>, TInlineSetAllocator<UE::PoseSearch::PreallocatedRolesNum>> CoveredRoles;
+	TArray<FInteractionSearchContext, TInlineAllocator<PreallocatedSearchesNum, TMemStackAllocator<>>> SearchContexts;
+	for (const FDatabaseToRoledAnimInstancesPair& DatabaseToRoledAnimInstancesPair : DatabaseToRoledAnimInstances)
+	{
+		const UPoseSearchDatabase* Database = DatabaseToRoledAnimInstancesPair.Key;
+		const TArray<FPoseSearchRoledSkeleton>& RoledSkeletons = Database->Schema->GetRoledSkeletons();
+		const int32 CombinationCardinality = RoledSkeletons.Num();
+		
+		// @todo: naive implementation: optimize this code!
+		GeneratePermutations(DatabaseToRoledAnimInstancesPair.Value.GetDataView(), CombinationCardinality, 
+			[&AnimInstances, &CoveredRoles, Database, &RoledSkeletons, &SearchContexts](const TArrayView<const FRoledAnimInstance> RoledAnimInstances, const TArrayView<int32> Combination)
+			{
+				// checking if this combination is valid for the Database:
+				AnimInstances.Reset();
+				CoveredRoles.Reset();
+
+				for (int32 CombinationIndex = 0; CombinationIndex < Combination.Num(); ++CombinationIndex)
+				{
+					const FRoledAnimInstance& RoledAnimInstance = RoledAnimInstances[Combination[CombinationIndex]];
+
+					bool bIsAlreadyInSet = false;
+					AnimInstances.Add(RoledAnimInstance.AnimInstance, &bIsAlreadyInSet);
+
+					// is there any duplicated AnimInstance?
+					if (bIsAlreadyInSet)
+					{
+						return false;
+					}
+
+					CoveredRoles.Add(RoledAnimInstance.Role);
+				}
+
+				// does it cover all the roles?
+				for (const FPoseSearchRoledSkeleton& RoledSkeleton : RoledSkeletons)
+				{
+					if (!CoveredRoles.Find(RoledSkeleton.Role))
+					{
+						return false;
+					}
+				}
+
+				// @todo: naive implementation: optimize this code!
+				// checking if all the actors are within BroadPhaseRadius to each other
+				for (int32 CombinationIndexA = 0; CombinationIndexA < Combination.Num(); ++CombinationIndexA)
+				{
+					const FRoledAnimInstance& RoledAnimInstanceA = RoledAnimInstances[Combination[CombinationIndexA]];
+					const FVector ActorLocationA = RoledAnimInstanceA.AnimInstance->GetOwningActor()->GetActorLocation();
+
+					for (int32 CombinationIndexB = 0; CombinationIndexB < Combination.Num(); ++CombinationIndexB)
+					{
+						if (CombinationIndexA != CombinationIndexB)
+						{
+							const FRoledAnimInstance& RoledAnimInstanceB = RoledAnimInstances[Combination[CombinationIndexB]];
+							const FVector ActorLocationB = RoledAnimInstanceB.AnimInstance->GetOwningActor()->GetActorLocation();
+
+							const FVector DeltaLocation = ActorLocationA - ActorLocationB;
+							if (DeltaLocation.Length() > RoledAnimInstanceA.BroadPhaseRadius)
+							{
+								return false;
+							}
+						}
+					}
+				}
+
+				FInteractionSearchContext& SearchContext = SearchContexts.AddDefaulted_GetRef();
+				SearchContext.Database = Database;
+				SearchContext.AnimInstances.Reserve(Combination.Num());
+				SearchContext.HistoryCollectors.Reserve(Combination.Num());
+				SearchContext.Roles.Reserve(Combination.Num());
+				
+#if ENABLE_DRAW_DEBUG
+				SearchContext.BroadPhaseRadiuses.Reserve(Combination.Num());
+#endif //ENABLE_DRAW_DEBUG
+
+				for (int32 CombinationIndex = 0; CombinationIndex < Combination.Num(); ++CombinationIndex)
+				{
+					const FRoledAnimInstance& RoledAnimInstance = RoledAnimInstances[Combination[CombinationIndex]];
+					SearchContext.AnimInstances.Add(RoledAnimInstance.AnimInstance);
+					SearchContext.HistoryCollectors.Add(RoledAnimInstance.HistoryCollector);
+					SearchContext.Roles.Add(RoledAnimInstance.Role);
+#if ENABLE_DRAW_DEBUG
+					SearchContext.BroadPhaseRadiuses.Add(RoledAnimInstance.BroadPhaseRadius);
+#endif //ENABLE_DRAW_DEBUG
+
+					// search is valid only if cost result is lower than any AvailabilityRequest.MaxCost (copied in RoledAnimInstance.MaxCost)
+					SearchContext.MaxCost = FMath::Min(SearchContext.MaxCost, RoledAnimInstance.MaxCost);
+				}
+
+				return true;
+			});
+	}
+	
+	// for each valid SearchContexts we try to figure out the continuing pose properties from the current Islands
+	for (FInteractionSearchContext& SearchContext : SearchContexts)
+	{
+		PopulateContinuingProperties(SearchContext, DeltaSeconds);
+	}
+
+	// grouping SearchContexts in different Islands
+	TArray<FInteractionSearchContextGroup, TInlineAllocator<PreallocatedSearchesNum, TMemStackAllocator<>>> SearchContextGroups;
+	for (int32 SearchContextIndex = 0; SearchContextIndex < SearchContexts.Num(); ++SearchContextIndex)
+	{
+		const FInteractionSearchContext& SearchContext = SearchContexts[SearchContextIndex];
+
+		int32 MainSearchContextGroupIndex = INDEX_NONE;
+		for (int32 SearchContextGroupIndex = 0; SearchContextGroupIndex < SearchContextGroups.Num();)
+		{
+			if (SearchContextGroups[SearchContextGroupIndex].Contains(SearchContext))
+			{
+				if (MainSearchContextGroupIndex == INDEX_NONE)
+				{
+					MainSearchContextGroupIndex = SearchContextGroupIndex;
+					SearchContextGroups[MainSearchContextGroupIndex].Add(SearchContext, SearchContextIndex);
+					++SearchContextGroupIndex;
+				}
+				else
+				{
+					SearchContextGroups[MainSearchContextGroupIndex].Merge(SearchContextGroups[SearchContextGroupIndex]);
+					SearchContextGroups.RemoveAt(SearchContextGroupIndex);
+				}
+			}
+			else
+			{
+				++SearchContextGroupIndex;
+			}
+		}
+		if (MainSearchContextGroupIndex == INDEX_NONE)
+		{
+			SearchContextGroups.AddDefaulted_GetRef().Add(SearchContext, SearchContextIndex);
+		}
+	}
+
+#if ENABLE_DRAW_DEBUG
+	DebugDraw();
+#endif //ENABLE_DRAW_DEBUG
+
+	UninjectAllIslands();
+
+	for (const FInteractionSearchContextGroup& SearchContextGroup : SearchContextGroups)
+	{
+		UE::PoseSearch::FInteractionIsland& Island = GetAvailableIsland();
+
+		for (const UAnimInstance* AnimInstance : SearchContextGroup.AnimInstances)
+		{
+			Island.InjectToActor(AnimInstance->GetOwningActor());
+		}
+
+		for (int32 SearchContextsIndex : SearchContextGroup.SearchContextsIndices)
+		{
+			Island.AddSearchContext(SearchContexts[SearchContextsIndex]);
+		}
+	}
+
+	AvailabilityRequestsMap.Reset();
+
+	check(ValidateAllIslands());
+}
+
+TStatId UPoseSearchInteractionSubsystem::GetStatId() const
+{
+	RETURN_QUICK_DECLARE_CYCLE_STAT(UPoseSearchInteractionSubsystem, STATGROUP_Tickables);
+}
+
+void UPoseSearchInteractionSubsystem::Query_AnyThread(const TArrayView<const FPoseSearchInteractionAvailability> Availabilities, UObject* AnimInstance, 
+	const FPoseSearchContinuingProperties& ContinuingProperties, FPoseSearchInteractionBlueprintResult& Result,
+	FName PoseHistoryName, const FAnimNode_PoseSearchHistoryCollector_Base* HistoryCollector, bool bValidateResultAgainstAvailabilities)
+{
+	using namespace UE::PoseSearch;
+
+	Result = FPoseSearchInteractionBlueprintResult();
+
+	if (AnimInstance)
+	{
+		// if we find AnimInstance in an island, we perform ALL the Island motion matching searches.
+		// thread safety is ensured by the lock on Island->SearchResultsMutex in DoSearch_AnyThread
+		if (FInteractionIsland* Island = FindIsland(AnimInstance))
+		{
+			Island->DoSearch_AnyThread(AnimInstance, ContinuingProperties, Result);
+
+			if (bValidateResultAgainstAvailabilities && Result.SelectedAnimation)
+			{
+				bool bResultValidated = false;
+
+				for (const FPoseSearchInteractionAvailability& Availability : Availabilities)
+				{
+					if (Availability.Database == Result.SelectedDatabase &&
+						(Availability.RolesFilter.IsEmpty() || Availability.RolesFilter.Contains(Result.Role)))
+					{
+						bResultValidated = true;
+						break;
+					}
+				}
+
+				if (!bResultValidated)
+				{
+					Result = FPoseSearchInteractionBlueprintResult();
+				}
+			}
+		}
+
+		// queuing the availabilities for the next frame Query_AnyThread
+		if (!Availabilities.IsEmpty())
+		{
+			FScopeLock Lock(&AvailabilityRequestsMapMutex);
+
+			TArray<FInteractionAvailabilityEx>& AvailabilityRequests = AvailabilityRequestsMap.FindOrAdd(AnimInstance);
+			for (const FPoseSearchInteractionAvailability& Availability : Availabilities)
+			{
+				if (Availability.Database)
+				{
+					AvailabilityRequests.Emplace(Availability, PoseHistoryName, HistoryCollector);
+				}
+			}
+		}
+	}
+}
+
+void UPoseSearchInteractionSubsystem::DebugDraw() const
+{
+	using namespace UE::PoseSearch;
+
+#if ENABLE_DRAW_DEBUG
+	static const FColor Colors[] = { FColor::Red, FColor::Green, FColor::Blue, FColor::Yellow, FColor::Black };
+	static const int32 NumColors = sizeof(Colors) / sizeof(Colors[0]);
+	int32 CurrentColorIndex = 0;
+
+	for (const FInteractionIsland* Island : Islands)
+	{
+		if (!Island->IsUninjected())
+		{
+			Island->DebugDraw(Colors[CurrentColorIndex]);
+			CurrentColorIndex = (CurrentColorIndex + 1) % NumColors;
+		}
+	}
+
+#endif //ENABLE_DRAW_DEBUG
+
+#if ENABLE_VISUAL_LOG
+	TStringBuilder<512> StringBuilder;
+
+	for (FAvailabilityRequestsMap::TConstIterator Iter = AvailabilityRequestsMap.CreateConstIterator(); Iter; ++Iter)
+	{
+		// making sure the weak pointer is still valid
+		if (UAnimInstance* AnimInstance = Cast<UAnimInstance>(Iter->Key.Get()))
+		{
+			// looking for valid results for this AnimInstance
+			const UPoseSearchDatabase* ResultDatabase = nullptr;
+			UE::PoseSearch::FRole ResultRole;
+			for (const FInteractionIsland* Island : Islands)
+			{
+				if (!Island->IsUninjected())
+				{
+					for (const FInteractionSearchResult& SearchResult : Island->GetSearchResults())
+					{
+						const FInteractionSearchContext& SearchContext = Island->GetSearchContexts()[SearchResult.SearchIndex];
+						for (int32 AnimInstanceIndex = 0; AnimInstanceIndex < SearchContext.AnimInstances.Num(); ++AnimInstanceIndex)
+						{
+							if (SearchContext.AnimInstances[AnimInstanceIndex].Get() == AnimInstance)
+							{
+								ResultDatabase = SearchResult.Database.Get();
+								check(ResultDatabase);
+								ResultRole = SearchContext.Roles[AnimInstanceIndex];
+								break;
+							}
+						}
+
+						if (ResultDatabase)
+						{
+							break;
+						}
+					}
+				}
+
+				if (ResultDatabase)
+				{
+					break;
+				}
+			}
+
+			StringBuilder.Reset();
+			for (const FInteractionAvailabilityEx& AvailabilityRequest : Iter->Value)
+			{
+				StringBuilder.Append(GetNameSafe(AvailabilityRequest.Database));
+				StringBuilder.Append(" / ");
+				StringBuilder.Append(AvailabilityRequest.GetPoseHistoryName());
+				StringBuilder.Append(" [");
+				bool bAddComma = false;
+
+				if (AvailabilityRequest.RolesFilter.IsEmpty())
+				{
+					// adding ALL the possible roles from the database:
+					if (AvailabilityRequest.Database && AvailabilityRequest.Database->Schema)
+					{
+						for (const FPoseSearchRoledSkeleton& RoledSkeleton : AvailabilityRequest.Database->Schema->GetRoledSkeletons())
+						{
+							if (bAddComma)
+							{
+								StringBuilder.Append(",");
+							}
+							else
+							{
+								bAddComma = true;
+							}
+
+							if (AvailabilityRequest.Database == ResultDatabase && RoledSkeleton.Role == ResultRole)
+							{
+								StringBuilder.Append("(*)");
+							}
+
+							StringBuilder.Append(RoledSkeleton.Role.ToString());
+						}
+					}
+				}
+				else
+				{
+					for (const FName& Role : AvailabilityRequest.RolesFilter)
+					{
+						if (bAddComma)
+						{
+							StringBuilder.Append(",");
+						}
+						else
+						{
+							bAddComma = true;
+						}
+
+						if (AvailabilityRequest.Database == ResultDatabase && Role == ResultRole)
+						{
+							StringBuilder.Append("(*)");
+						}
+
+						StringBuilder.Append(Role.ToString());
+					}
+				}
+
+				StringBuilder.Append("]");
+				StringBuilder.Append("\n");
+			}
+			
+			// @todo: waiting for visual logger fix to display strings :/
+			const FVector StringOffset(0.f, 0.f, 0.001f);
+			const FVector ActorLocation = AnimInstance->GetOwningActor()->GetActorLocation();
+			UE_VLOG_SEGMENT(AnimInstance, "PoseSearchInteraction", Display, ActorLocation, ActorLocation + StringOffset, FColor::Transparent, TEXT("%s"), StringBuilder.ToString());
+		}
+	}
+#endif // ENABLE_VISUAL_LOG
+}
\ No newline at end of file
diff --git a/Engine/Plugins/Animation/PoseSearch/Source/Runtime/Public/PoseSearch/AnimNode_MotionMatchingInteraction.h b/Engine/Plugins/Animation/PoseSearch/Source/Runtime/Public/PoseSearch/AnimNode_MotionMatchingInteraction.h
new file mode 100644
--- /dev/null
+++ b/Engine/Plugins/Animation/PoseSearch/Source/Runtime/Public/PoseSearch/AnimNode_MotionMatchingInteraction.h
@@ -0,0 +1,78 @@
+// Copyright Epic Games, Inc. All Rights Reserved.
+
+#pragma once
+
+#include "BlendStack/AnimNode_BlendStack.h"
+#include "PoseSearch/PoseSearchInteractionLibrary.h"
+#include "AnimNode_MotionMatchingInteraction.generated.h"
+
+struct FAnimationInitializeContext;
+struct FNodeDebugData;
+
+UENUM(Experimental)
+enum class EMotionMatchingInteractionEvaluationMode
+{
+	// Node will continuously provide its availabilities and eventually blend to newly selected animations
+	ContinuousReselection,
+
+	// Node will continuously provide its availabilities to keep the interaction alive, but will play only the first selected animation
+	// the idea is to let the animation play until the end and allow the evenutual state machine playing this node to be able to perform an automatic transition
+	SingleSelection,
+
+	// @todo: is this needed?
+	// Node will stop providing its availabilities and consequently kill the interaction, when the first selected animation stop playing as continuing pose
+	// the idea is to let the animation play until its valid as continuing pose and allow the evenutual state machine playing this node to be able to perform an automatic transition
+	// UntilContinuingPoseIsValid,
+};
+
+USTRUCT(Experimental, BlueprintInternalUseOnly)
+struct POSESEARCH_API FAnimNode_MotionMatchingInteraction : public FAnimNode_BlendStack
+{
+	GENERATED_BODY();
+
+	UPROPERTY(EditAnywhere, BlueprintReadWrite, Category = Links)
+	FPoseLink Source;
+
+	UPROPERTY(EditAnywhere, Category = Settings)
+	EMotionMatchingInteractionEvaluationMode EvaluationMode = EMotionMatchingInteractionEvaluationMode::ContinuousReselection;
+
+	UPROPERTY(EditAnywhere, Category = Settings, meta = (PinHiddenByDefault))
+	TArray<FPoseSearchInteractionAvailability> Availabilities;
+
+	UPROPERTY(EditAnywhere, Category = Settings, meta = (PinHiddenByDefault))
+	bool bValidateResultAgainstAvailabilities = true;
+
+	// time from the beginning of the interaction to warp to full translation alignment
+	UPROPERTY(EditAnywhere, BlueprintReadWrite, Category = Warping, meta = (PinHiddenByDefault, ClampMin = "0"))
+	float InitialTranslationWarpTime = 0.2f;
+
+	// time from the beginning of the interaction to warp to full rotation alignment
+	UPROPERTY(EditAnywhere, BlueprintReadWrite, Category = Warping, meta = (PinHiddenByDefault, ClampMin = "0"))
+	float InitialRotationWarpTime = 0.2f;
+
+	// FAnimNode_Base interface
+	virtual void GatherDebugData(FNodeDebugData& DebugData) override;
+	virtual void Initialize_AnyThread(const FAnimationInitializeContext& Context) override;
+	virtual void Evaluate_AnyThread(FPoseContext& Output) override;
+	virtual void CacheBones_AnyThread(const FAnimationCacheBonesContext& Context) override;
+
+	virtual void UpdateAssetPlayer(const FAnimationUpdateContext& Context) override;
+	// End of FAnimNode_Base interface
+
+	virtual void Reset() override;
+
+protected:
+	float BlendLerp = 0.f;
+	float TranslationWarpLerp = 0.f;
+	float RotationWarpLerp = 0.f;
+	// if a search is successful InteractingRolesNum > 0. 
+	// if InteractingRolesNum == 1 it means that the search is a regular single character motion matching search
+	// if InteractingRolesNum > 1 it means this node is interacting with other actors via the Interaction Subsystem
+	int32 InteractingRolesNum = 0;
+
+	// Update Counter for detecting being relevant
+	FGraphTraversalCounter UpdateCounter;
+
+	// wanted world transform for full aligment interaction
+	FTransform FullAlignedActorRootBoneTransform = FTransform::Identity;
+};
diff --git a/Engine/Plugins/Animation/PoseSearch/Source/Runtime/Public/PoseSearch/MotionMatchingInteractionAnimNodeLibrary.h b/Engine/Plugins/Animation/PoseSearch/Source/Runtime/Public/PoseSearch/MotionMatchingInteractionAnimNodeLibrary.h
new file mode 100644
--- /dev/null
+++ b/Engine/Plugins/Animation/PoseSearch/Source/Runtime/Public/PoseSearch/MotionMatchingInteractionAnimNodeLibrary.h
@@ -0,0 +1,39 @@
+// Copyright Epic Games, Inc. All Rights Reserved.
+
+#pragma once
+
+#include "Animation/AnimNodeReference.h"
+#include "PoseSearch/PoseSearchInteractionLibrary.h"
+#include "Kismet/BlueprintFunctionLibrary.h"
+#include "MotionMatchingInteractionAnimNodeLibrary.generated.h"
+
+struct FAnimNode_MotionMatchingInteraction;
+
+USTRUCT(Experimental, BlueprintType)
+struct FMotionMatchingInteractionAnimNodeReference : public FAnimNodeReference
+{
+	GENERATED_BODY()
+
+	typedef FAnimNode_MotionMatchingInteraction FInternalNodeType;
+};
+
+UCLASS(Experimental)
+class POSESEARCH_API UMotionMatchingInteractionAnimNodeLibrary : public UBlueprintFunctionLibrary
+{
+	GENERATED_BODY()
+
+public:
+	UFUNCTION(BlueprintCallable, Category = "Animation|Pose Search", meta = (BlueprintThreadSafe, ExpandEnumAsExecs = "Result"))
+	static FMotionMatchingInteractionAnimNodeReference ConvertToMotionMatchingInteractionNode(const FAnimNodeReference& Node, EAnimNodeReferenceConversionResult& Result);
+
+	UFUNCTION(BlueprintPure, Category = "Animation|Pose Search", meta = (BlueprintThreadSafe, DisplayName = "Convert to Motion Matching Interaction Node"))
+	static void ConvertToMotionMatchingInteractionNodePure(const FAnimNodeReference& Node, FMotionMatchingInteractionAnimNodeReference& MotionMatchingInteractionNode, bool& Result)
+	{
+		EAnimNodeReferenceConversionResult ConversionResult;
+		MotionMatchingInteractionNode = ConvertToMotionMatchingInteractionNode(Node, ConversionResult);
+		Result = (ConversionResult == EAnimNodeReferenceConversionResult::Succeeded);
+	}
+
+	UFUNCTION(BlueprintCallable, Category = "Animation|Pose Search", meta = (BlueprintThreadSafe))
+	static void SetAvailabilities(const FMotionMatchingInteractionAnimNodeReference& MotionMatchingInteractionNode, const TArray<FPoseSearchInteractionAvailability>& Availabilities);
+};
\ No newline at end of file
diff --git a/Engine/Plugins/Animation/PoseSearch/Source/Runtime/Public/PoseSearch/PoseSearchDefines.h b/Engine/Plugins/Animation/PoseSearch/Source/Runtime/Public/PoseSearch/PoseSearchDefines.h
--- a/Engine/Plugins/Animation/PoseSearch/Source/Runtime/Public/PoseSearch/PoseSearchDefines.h
+++ b/Engine/Plugins/Animation/PoseSearch/Source/Runtime/Public/PoseSearch/PoseSearchDefines.h
@@ -31,6 +31,7 @@ static constexpr FBoneIndexType WorldSpaceIndexType = FBoneIndexType(-2);
 static constexpr int32 PreallocatedRolesNum = 4;
 static constexpr int32 PreallocatedCachedQueriesNum = 8;
 static constexpr int32 PreallocatedCachedChannelDataNum = 64;
+static constexpr int32 PreallocatedSearchesNum = 4;
 static constexpr float FiniteDelta = 1 / 60.0f; // Time delta used for computing pose derivatives
 static constexpr int32 MaxNumberOfCollectedPoseCandidatesPerDatabase = 256;
 static constexpr int32 InvalidSearchId = INDEX_NONE; // @todo: it should be 0, since the search id is an hash!
diff --git a/Engine/Plugins/Animation/PoseSearch/Source/Runtime/Public/PoseSearch/PoseSearchInteractionAsset.h b/Engine/Plugins/Animation/PoseSearch/Source/Runtime/Public/PoseSearch/PoseSearchInteractionAsset.h
new file mode 100644
--- /dev/null
+++ b/Engine/Plugins/Animation/PoseSearch/Source/Runtime/Public/PoseSearch/PoseSearchInteractionAsset.h
@@ -0,0 +1,89 @@
+// Copyright Epic Games, Inc. All Rights Reserved.
+
+#pragma once
+
+#include "Engine/DataAsset.h"
+#include "PoseSearch/MultiAnimAsset.h"
+#include "PoseSearch/PoseSearchRole.h"
+#include "PoseSearchInteractionAsset.generated.h"
+
+USTRUCT(Experimental)
+struct POSESEARCH_API FPoseSearchInteractionAssetItem
+{
+	GENERATED_BODY()
+
+	// associated aniamtion for this FPoseSearchInteractionAssetItem
+	UPROPERTY(EditAnywhere, Category = "Settings")
+	TObjectPtr<UAnimationAsset> Animation;
+
+	// associated role for this FPoseSearchInteractionAssetItem
+	UPROPERTY(EditAnywhere, Category = "Settings")
+	FName Role;
+
+	// relative weight to the other FPoseSearchInteractionAssetItem::WarpingWeightRotation(s) defining which character will be rotated while warping
+	// 0 - the associated character to this item will move fully to compensate the warping errors
+	// > 0 && all the other FPoseSearchInteractionAssetItem::WarpingWeightTranslation as zero, and the associated character will not move
+	UPROPERTY(EditAnywhere, Category = "Warping", meta = (ClampMin = "0", ClampMax = "1"))
+	float WarpingWeightRotation = 0.5f;
+
+	// relative weight to the other FPoseSearchInteractionAssetItem::WarpingWeightTranslation(s) defining which character will be translated while warping
+	// 0 - the associated character to this item will move fully to compensate the warping errors
+	// > 0 && all the other FPoseSearchInteractionAssetItem::WarpingWeightTranslation as zero, and the associated character will not move
+	UPROPERTY(EditAnywhere, Category = "Warping", meta = (ClampMin = "0", ClampMax = "1"))
+	float WarpingWeightTranslation = 0.5f;
+
+	// offset from the origin
+	UPROPERTY(EditAnywhere, Category = "Settings")
+	FTransform Origin = FTransform::Identity;
+};
+
+UCLASS(Experimental, BlueprintType, Category = "Animation|Pose Search")
+class POSESEARCH_API UPoseSearchInteractionAsset : public UMultiAnimAsset
+{
+	GENERATED_BODY()
+
+	UPROPERTY(EditAnywhere, Category = "Settings")
+	TArray<FPoseSearchInteractionAssetItem> Items;
+
+public:
+#if WITH_EDITORONLY_DATA
+
+	UPROPERTY(Transient, EditAnywhere, Category = "Debug", meta = (EditCondition=bEnableDebugWarp, EditConditionHides))
+	TArray<FTransform> DebugWarpOffsets;
+
+	// used to test warping: 0 - no warping applied, 1 - full warping/aligment applied
+	// test warping actors will be offsetted by Items::DebugWarpOffset transforms from the original
+	// UMultiAnimAsset::GetOrigin() definition and warped accordingly with CalculateWarpTransforms
+	// following the rotation and translation weights defined in Items::WarpingWeightRotation and
+	// Items::WarpingWeightTranslation as relative weights between the Items (they'll be normalized at runtime)
+	UPROPERTY(Transient, EditAnywhere, Category = "Debug", meta = (ClampMin = "0", ClampMax = "1", EditCondition=bEnableDebugWarp, EditConditionHides))
+	float DebugWarpAmount = 0.f;
+
+	UPROPERTY(Transient, EditAnywhere, Category = "Debug")
+	bool bEnableDebugWarp = false;
+#endif // WITH_EDITORONLY_DATA
+
+	virtual bool IsLooping() const override;
+	virtual bool HasRootMotion() const override;
+	virtual float GetPlayLength() const override;
+
+#if WITH_EDITOR
+	virtual int32 GetFrameAtTime(float Time) const override;
+#endif // WITH_EDITOR
+
+	virtual int32 GetNumRoles() const override { return Items.Num(); }
+	virtual UE::PoseSearch::FRole GetRole(int32 RoleIndex) const override { return Items[RoleIndex].Role; }
+	
+	virtual UAnimationAsset* GetAnimationAsset(const FName /*UE::PoseSearch::FRole*/& Role) const override;
+
+	virtual FTransform GetOrigin(const UE::PoseSearch::FRole& Role) const override;
+
+#if WITH_EDITOR
+	FTransform GetDebugWarpOrigin(const UE::PoseSearch::FRole& Role, bool bComposeWithDebugWarpOffset) const;
+#endif // WITH_EDITOR
+
+	virtual void CalculateWarpTransforms(float Time, const TArrayView<const FTransform> ActorRootBoneTransforms, TArrayView<FTransform> FullAlignedActorRootBoneTransforms) const override;
+
+	FQuat FindReferenceOrientation(const TArrayView<const FTransform> Transforms, const TArrayView<int32> SortedByWarpingWeightRotationItemIndex) const;
+	FVector FindReferencePosition(const TArrayView<const FTransform> Transforms, const TArrayView<float> NormalizedWarpingWeightTranslation) const;
+};
\ No newline at end of file
diff --git a/Engine/Plugins/Animation/PoseSearch/Source/Runtime/Public/PoseSearch/PoseSearchInteractionIsland.h b/Engine/Plugins/Animation/PoseSearch/Source/Runtime/Public/PoseSearch/PoseSearchInteractionIsland.h
new file mode 100644
--- /dev/null
+++ b/Engine/Plugins/Animation/PoseSearch/Source/Runtime/Public/PoseSearch/PoseSearchInteractionIsland.h
@@ -0,0 +1,105 @@
+// Copyright Epic Games, Inc. All Rights Reserved.
+
+#pragma once
+
+#include "PoseSearch/PoseSearchLibrary.h"
+
+struct FPoseSearchInteractionBlueprintResult;
+struct FPoseSearchContinuingProperties;
+class UCharacterMovementComponent;
+
+namespace UE::PoseSearch
+{
+struct FInteractionIsland;
+
+// Experimental, this feature might be removed without warning, not for production use
+struct FInteractionSearchContext
+{
+	// @todo: since AnimInstances, HistoryCollectors, Roles, and BroadPhaseRadiuses have the same cardinality,
+	// add an array of structs instead of many arrays
+	TArray<TWeakObjectPtr<UAnimInstance>> AnimInstances;
+	// @todo: perhaps use the animnode of the pose history collector instead
+	TArray<const FAnimNode_PoseSearchHistoryCollector_Base*> HistoryCollectors;
+	TArray<UE::PoseSearch::FRole> Roles;
+
+#if ENABLE_DRAW_DEBUG
+	TArray<float> BroadPhaseRadiuses;
+#endif // ENABLE_DRAW_DEBUG
+
+	float MaxCost = MAX_flt;
+
+	TWeakObjectPtr<const UPoseSearchDatabase> Database;
+	FPoseSearchContinuingProperties ContinuingProperties;
+
+	bool IsValid() const;
+	bool IsEquivalent(const FInteractionSearchContext& Other) const;
+};
+
+// Experimental, this feature might be removed without warning, not for production use
+struct FInteractionSearchResult : public FSearchResult
+{
+	int32 SearchIndex = INDEX_NONE;
+	TArray<FTransform, TInlineAllocator<PreallocatedRolesNum>> FullAlignedActorRootBoneTransforms;
+
+	bool operator==(const FInteractionSearchResult& Other) const;
+};
+
+// Experimental, this feature might be removed without warning, not for production use
+// FInteractionIsland contains ticks functions injected between the interacting actors UCharacterMovementComponent and USkeletalMeshComponent
+// to create a execution threading fence to be able to perform motion matching searches between the involved characters in a thread safe manner.
+// Look at UPoseSearchInteractionSubsystem "Execution model and threading details" for additional information
+struct FInteractionIsland
+{
+	UE_NONCOPYABLE(FInteractionIsland);
+	
+	FInteractionIsland(ULevel* Level);
+	~FInteractionIsland();
+
+	bool DoSearch_AnyThread(UObject* AnimInstance, const FPoseSearchContinuingProperties& ContinuingProperties, FPoseSearchInteractionBlueprintResult& Result);
+
+	const TArray<TWeakObjectPtr<UCharacterMovementComponent>>& GetCharacterMovementComponents() const { return CharacterMovementComponents; }
+	const TArray<TWeakObjectPtr<USkeletalMeshComponent>>& GetSkeletalMeshComponents() const { return SkeletalMeshComponents; }
+	const TArray<FInteractionSearchContext>& GetSearchContexts() const { return SearchContexts; }
+	const TArray<FInteractionSearchResult>& GetSearchResults() const { return SearchResults; }
+
+	bool IsUninjected() const;
+	void InjectToActor(AActor* Actor);
+	void Uninject();
+
+	const FInteractionSearchResult* FindSearchResult(const FInteractionSearchContext& SearchContext) const;
+	void AddSearchContext(const FInteractionSearchContext& SearchContext);
+
+	void DebugDraw(const FColor& Color = FColor::Red) const;
+
+private:
+	bool GetResult_AnyThread(UObject* AnimInstance, FPoseSearchInteractionBlueprintResult& Result);
+
+	struct FPreTickFunction : public FTickFunction
+	{
+		virtual void ExecuteTick(float DeltaTime, enum ELevelTick TickType, ENamedThreads::Type CurrentThread, const FGraphEventRef& MyCompletionGraphEvent) override;
+		virtual FString DiagnosticMessage() override { return TEXT("FPreTickFunction"); }
+		FInteractionIsland* Island = nullptr;
+	};
+	FPreTickFunction PreTickFunction;
+
+	struct FPostTickFunction : public FTickFunction
+	{
+		virtual void ExecuteTick(float DeltaTime, enum ELevelTick TickType, ENamedThreads::Type CurrentThread, const FGraphEventRef& MyCompletionGraphEvent) override;
+		virtual FString DiagnosticMessage() override { return TEXT("FPostTickFunction"); }
+	};
+	FPostTickFunction PostTickFunction;
+
+	TArray<TWeakObjectPtr<UCharacterMovementComponent>> CharacterMovementComponents;
+	TArray<TWeakObjectPtr<USkeletalMeshComponent>> SkeletalMeshComponents;
+	
+	// there's one FSearchContext for each search we need to perform (including all the possible roles permutations). Added by UPoseSearchInteractionSubsystem::Tick
+	TArray<FInteractionSearchContext> SearchContexts;
+
+	// SearchResults contains only the best results, and it has not necessarly the same cardinality as SearchContexts. usually SearchResults.Num() < SearchContexts.Num()
+	TArray<FInteractionSearchResult> SearchResults;
+	bool bSearchPerfomed = false;
+
+	FCriticalSection SearchResultsMutex;
+};
+
+} // namespace UE::PoseSearch
diff --git a/Engine/Plugins/Animation/PoseSearch/Source/Runtime/Public/PoseSearch/PoseSearchInteractionLibrary.h b/Engine/Plugins/Animation/PoseSearch/Source/Runtime/Public/PoseSearch/PoseSearchInteractionLibrary.h
new file mode 100644
--- /dev/null
+++ b/Engine/Plugins/Animation/PoseSearch/Source/Runtime/Public/PoseSearch/PoseSearchInteractionLibrary.h
@@ -0,0 +1,122 @@
+// Copyright Epic Games, Inc. All Rights Reserved.
+
+#pragma once
+
+#include "Kismet/BlueprintFunctionLibrary.h"
+#include "PoseSearch/PoseSearchLibrary.h"
+#include "PoseSearchInteractionLibrary.generated.h"
+
+// subsystem in charge of coordinating Characters availability for motion matched interactions,
+// scheduling motion matching searches and synchronizing animations playback
+
+
+// input for MotionMatchInteraction_Pure: it declares that the associated character (AnimInstance) is willing to partecipate in an interction 
+// described by a UMultiAnimAsset (derived by UPoseSearchInteractionAsset) contained in the UPoseSearchDatabase Database
+// with one of the roles in RolesFilter (if empty ANY of the Database roles can be taken)
+// the MotionMatchInteraction_Pure will ultimately setup a motion matching query using looking for the pose history "PoseHistoryName" 
+// to gather bone and trajectory positions for this character
+// for an interaction to be valid, the query needs to find all the other interacting characters within BroadPhaseRadius,
+// and reach a maximum cost of MaxCost
+USTRUCT(Experimental, BlueprintType, Category = "Animation|Pose Search")
+struct POSESEARCH_API FPoseSearchInteractionAvailability
+{
+	GENERATED_BODY()
+
+	// Database describing the interaction. It'll contains multi character UMultiAnimAsset and a schema with multiple skeletons with associated roles
+	UPROPERTY(EditAnywhere, BlueprintReadWrite, Category=Settings)
+	TObjectPtr<UPoseSearchDatabase> Database;
+
+	// roles the character is willing to take to partecipate in this interaction. If empty ANY of the Database roles can be taken
+	UPROPERTY(EditAnywhere, BlueprintReadWrite, Category=Settings)
+	TArray<FName> RolesFilter;
+
+	// the associated character to this FPoseSearchInteractionAvailability will partecipate in an interaction only if all the necessary roles gest assigned to character within BroadPhaseRadius centimeters
+	UPROPERTY(EditAnywhere, BlueprintReadWrite, Category=Settings)
+	float BroadPhaseRadius = 500.f;
+
+	// the associated character to this FPoseSearchInteractionAvailability will not partecipate in an interaction if the motion matching search cost result is higher than MaxCost
+	UPROPERTY(EditAnywhere, BlueprintReadWrite, Category=Settings)
+	float MaxCost = 1000.f;
+};
+
+USTRUCT(Experimental, BlueprintType, Category="Animation|Pose Search")
+struct POSESEARCH_API FPoseSearchInteractionBlueprintResult
+{
+	GENERATED_BODY()
+public:
+
+	// animation assigned to this character to partecipate in the interaction
+	UPROPERTY(Transient, VisibleAnywhere, BlueprintReadOnly, Category=State)
+	TObjectPtr<UObject> SelectedAnimation = nullptr;
+	
+	// SelectedAnimation associated time
+	UPROPERTY(Transient, VisibleAnywhere, BlueprintReadOnly, Category=State)
+	float SelectedTime = 0.f;
+	
+	// SelectedAnimation at SelectedTime is from the continuing pose search
+	UPROPERTY(Transient, VisibleAnywhere, BlueprintReadOnly, Category=State)
+	bool bIsContinuingPoseSearch = false;
+
+	// SelectedAnimation associated play rate
+	UPROPERTY(Transient, VisibleAnywhere, BlueprintReadOnly, Category=State)
+	float WantedPlayRate = 0.f;
+
+	// SelectedAnimation associated looping state
+	UPROPERTY(Transient, VisibleAnywhere, BlueprintReadOnly, Category=State)
+	bool bLoop = false;
+	
+	// SelectedAnimation associated mirror state
+	UPROPERTY(Transient, VisibleAnywhere, BlueprintReadOnly, Category=State)
+	bool bIsMirrored = false;
+	
+	// SelectedAnimation associated BlendParameters (if SelectedAnimation is a UBlendSpace)
+	UPROPERTY(Transient, VisibleAnywhere, BlueprintReadOnly, Category=State)
+	FVector BlendParameters = FVector::ZeroVector;
+
+	// selected SelectedDatabase for this character interaction
+	UPROPERTY(Transient, VisibleAnywhere, BlueprintReadOnly, Category=State)
+	TWeakObjectPtr<const UPoseSearchDatabase> SelectedDatabase = nullptr;
+
+	// associated motion matching search cost for this result
+	UPROPERTY(Transient, VisibleAnywhere, BlueprintReadOnly, Category=State)
+	float SearchCost = MAX_flt;
+
+	// assigned role to this character (AnimInstance)
+	UPROPERTY(Transient, VisibleAnywhere, BlueprintReadOnly, Category=State)
+	FName Role;
+
+	// root bone transform for the character at full aligment
+	UPROPERTY(Transient, VisibleAnywhere, BlueprintReadOnly, Category=State)
+	FTransform FullAlignedActorRootBoneTransform = FTransform::Identity;	
+};
+
+UCLASS(Experimental)
+class UPoseSearchInteractionLibrary : public UBlueprintFunctionLibrary
+{
+	GENERATED_BODY()
+
+public:
+
+	// function publishing this character (via its AnimInstance) FPoseSearchInteractionAvailability to the UPoseSearchInteractionSubsystem,
+	// FPoseSearchInteractionAvailability represents the character availability to partecipate in an interaction with other characters for the next frame.
+	// that means there will always be one frame delay between publiching availabilities and getting a result back from MotionMatchInteraction_Pure!
+	// 
+	// if FPoseSearchInteractionBlueprintResult has a valid SelectedAnimation, this will be the animation assigned to this character to partecipate in this interaction.
+	// additional interaction properties, like assigned role, SelectedAnimation time, SearchCost, etc can be found within the result
+	// ContinuingProperties are used to figure out the continuing pose and bias it accordingly. ContinuingProperties can reference directly the UMultiAnimAsset
+	// or any of the roled UMultiAnimAsset::GetAnimationAsset, and the UPoseSearchInteractionSubsystem will figure out the related UMultiAnimAsset
+	// PoseHistoryName is the name of the pose history node used for the associated motion matching search
+	// if bValidateResultAgainstAvailabilities is true, the result will be invalidated if doesn't respect the new availabilities
+	UFUNCTION(Experimental, BlueprintPure, Category = "Animation|Pose Search", meta = (BlueprintThreadSafe))
+	static FPoseSearchInteractionBlueprintResult MotionMatchInteraction_Pure(TArray<FPoseSearchInteractionAvailability> Availabilities, UObject* AnimInstance, FPoseSearchContinuingProperties ContinuingProperties, FName PoseHistoryName, bool bValidateResultAgainstAvailabilities);
+
+	// BlueprintCallable version of MotionMatchInteraction_Pure
+	UFUNCTION(Experimental, BlueprintCallable, Category = "Animation|Pose Search", meta = (BlueprintThreadSafe))
+	static FPoseSearchInteractionBlueprintResult MotionMatchInteraction(TArray<FPoseSearchInteractionAvailability> Availabilities, UObject* AnimInstance, FPoseSearchContinuingProperties ContinuingProperties, FName PoseHistoryName, bool bValidateResultAgainstAvailabilities);
+
+	// version of MotionMatchInteraction_Pure referencing directly the HistoryCollector rather than looking for it by name
+	static FPoseSearchInteractionBlueprintResult MotionMatchInteraction(const TArrayView<const FPoseSearchInteractionAvailability> Availabilities, UObject* AnimInstance, const FPoseSearchContinuingProperties& ContinuingProperties, const FAnimNode_PoseSearchHistoryCollector_Base* HistoryCollector, bool bValidateResultAgainstAvailabilities);
+
+	UFUNCTION(Experimental, BlueprintPure, Category = "Animation|Pose Search", meta = (BlueprintThreadSafe))
+	static FPoseSearchContinuingProperties GetMontageContinuingProperties(UAnimInstance* AnimInstance);
+};
\ No newline at end of file
diff --git a/Engine/Plugins/Animation/PoseSearch/Source/Runtime/Public/PoseSearch/PoseSearchInteractionSubsystem.h b/Engine/Plugins/Animation/PoseSearch/Source/Runtime/Public/PoseSearch/PoseSearchInteractionSubsystem.h
new file mode 100644
--- /dev/null
+++ b/Engine/Plugins/Animation/PoseSearch/Source/Runtime/Public/PoseSearch/PoseSearchInteractionSubsystem.h
@@ -0,0 +1,126 @@
+// Copyright Epic Games, Inc. All Rights Reserved.
+
+#pragma once
+
+#include "Subsystems/WorldSubsystem.h"
+#include "PoseSearch/PoseSearchInteractionLibrary.h"
+#include "PoseSearchInteractionSubsystem.generated.h"
+
+struct FPoseSearchContinuingProperties;
+class UAnimInstance;
+class UPoseSearchDatabase;
+
+namespace UE::PoseSearch
+{
+
+struct FInteractionIsland;
+struct FInteractionSearchContext;
+
+// Experimental, this feature might be removed without warning, not for production use
+struct FInteractionAvailabilityEx : public FPoseSearchInteractionAvailability
+{
+	FInteractionAvailabilityEx(const FPoseSearchInteractionAvailability& InAvailability, FName InPoseHistoryName, const FAnimNode_PoseSearchHistoryCollector_Base* InHistoryCollector)
+		: FPoseSearchInteractionAvailability(InAvailability)
+		, PoseHistoryName(InPoseHistoryName)
+		, HistoryCollector(InHistoryCollector)
+	{
+	}
+
+	FString GetPoseHistoryName() const;
+	const FAnimNode_PoseSearchHistoryCollector_Base* GetHistoryCollector(const UAnimInstance* AnimInstance) const;
+
+private:
+	FName PoseHistoryName;
+	const FAnimNode_PoseSearchHistoryCollector_Base* HistoryCollector = nullptr;
+};
+
+typedef TMap<TWeakObjectPtr<UObject>, TArray<FInteractionAvailabilityEx>> FAvailabilityRequestsMap;
+
+} // namespace UE::PoseSearch
+
+// World subsystem accepting the publication of characters (via their AnimInstance(s)) FPoseSearchInteractionAvailability, representing the characters willingness to partecipate in an 
+// interaction with other characters from the next frame forward via Query_AnyThread method.
+//
+// The same method will return the FPoseSearchInteractionBlueprintResult from the PREVIOUS Tick processing (categorization of FPoseSearchInteractionAvailability(s) in multiple FInteractionIsland(s)),
+// to the requesting character, containing the animation to play at what time, and the assigned role to partecipate in the selected interaction within the assigned 
+// FInteractionIsland
+
+// Execution model and threading details:
+/////////////////////////////////////////
+// 
+// - by calling UPoseSearchInteractionLibrary::MotionMatchInteraction_Pure(TArray<FPoseSearchInteractionAvailability> Availabilities, UObject* AnimInstance), characters publish their availabilities
+//   to partecipate in interactions to the UPoseSearchInteractionSubsystem
+// - UPoseSearchInteractionSubsystem::Tick processes those FPoseSearchInteractionAvailability(s) and creates/updates UE::PoseSearch::FInteractionIsland. For each FInteractionIsland it injects 
+//   a tick prerequisite via FInteractionIsland::InjectToActor (that calls AddPrerequisite) to all the Actors in the same island.
+//   NoTe: the next frame the execution will be:
+//			for each island[k]
+//			{
+//				for each Actor[k][i]
+//				{
+//					Tick CharacterMovementComponent[k][i]
+//				}
+// 
+//				Tick Island[k].PreTickFunction (that eventually generates the trajectories with all the updated CMCs)
+//
+//				Tick Actor[k][0].SkeletalMeshComponent (that performs the MotionMatchInteraction queries for all the involved actors via DoSearch_AnyThread)
+// 
+//				Tick Island[k].PostTickFunction (currently just a threading fence for the execution of all the other SkeletalMeshComponent(s))
+// 
+//				for each Actor[k][i]
+//				{
+//					if (i != 0)
+//						Tick SkeletalMeshComponent[k][i] (that DoSearch_AnyThread get the cached result calculated by Tick Actor[k][0].SkeletalMeshComponent)
+//				}
+//			}
+// - next frame UPoseSearchInteractionLibrary::MotionMatchInteraction_Pure(TArray<FPoseSearchInteractionAvailability> Availabilities, UObject* AnimInstance), with the context of all the published 
+//   availabilities and created islands, will find the associated FInteractionIsland to the AnimInstance and call FInteractionIsland::DoSearch_AnyThread 
+//   (via UPoseSearchInteractionSubsystem::Query_AnyThread) that will perform ALL (YES, ALL, so the bigger the island the slower the execution) the motion matching searches
+//   for all the possible Actors / databases / Roles permutations, and populate FInteractionIsland::SearchResults with ALL the results for the island.
+//   Ultimately the MotionMatchInteraction_Pure will return the SearchResults associated to the requesting AnimInstance with information about what animation to play
+//   at what time with wich Role.
+
+UCLASS(Experimental, Category="Animation|Pose Search")
+class POSESEARCH_API UPoseSearchInteractionSubsystem : public UTickableWorldSubsystem
+{
+	GENERATED_BODY()
+
+public:
+	static UPoseSearchInteractionSubsystem* GetSubsystem_AnyThread(UObject* AnimInstance);
+
+	virtual void Initialize(FSubsystemCollectionBase& Collection) override;
+	virtual void Deinitialize() override;
+
+	// it processes FPoseSearchInteractionAvailability(s) and creates/updates FInteractionIsland
+	virtual void Tick(float DeltaSeconds) override;
+	
+	virtual TStatId GetStatId() const override;
+
+	// publishing FPoseSearchInteractionAvailability(s) for the requesting character (AnimInstance) and getting the FPoseSearchInteractionBlueprintResult from the PREVIOUS Tick update
+	// containing the animation to play at what time, and the assigned role to partecipate in the selected interaction
+	// Either a PoseHistoryName or a HistoryCollector are required to perform the associated motion matching searches
+	void Query_AnyThread(const TArrayView<const FPoseSearchInteractionAvailability> Availabilities, UObject* AnimInstance, 
+		const FPoseSearchContinuingProperties& ContinuingProperties, FPoseSearchInteractionBlueprintResult& Result,
+		FName PoseHistoryName, const FAnimNode_PoseSearchHistoryCollector_Base* HistoryCollector, bool bValidateResultAgainstAvailabilities);
+private:
+	UE::PoseSearch::FInteractionIsland& CreateIsland();
+	UE::PoseSearch::FInteractionIsland& GetAvailableIsland();
+
+	void DestroyIsland(int32 Index);
+	void DestroyAllIslands();
+	void UninjectAllIslands();
+	bool ValidateAllIslands() const;
+
+	void PopulateContinuingProperties(UE::PoseSearch::FInteractionSearchContext& SearchContext, float DeltaSeconds) const;
+	UE::PoseSearch::FInteractionIsland* FindIsland(UObject* InAnimInstance);
+	
+	void DebugDraw() const;
+
+	UE::PoseSearch::FAvailabilityRequestsMap AvailabilityRequestsMap;
+	FCriticalSection AvailabilityRequestsMapMutex;
+
+	// array of groups of characters that needs to be anaylzed together for possible interactions
+	TArray<UE::PoseSearch::FInteractionIsland*> Islands;
+
+	// critical section to retrieve the subsystem in a thread safe manner
+	static FCriticalSection RetrieveSubsystemMutex;
+};
\ No newline at end of file
