posesearch - fix for 'CVarPoseSearchInteractionShowIslands': undeclared identifier

diff --git a/Engine/Plugins/Animation/PoseSearch/Source/Runtime/Private/PoseSearchInteractionIsland.cpp b/Engine/Plugins/Animation/PoseSearch/Source/Runtime/Private/PoseSearchInteractionIsland.cpp
--- a/Engine/Plugins/Animation/PoseSearch/Source/Runtime/Private/PoseSearchInteractionIsland.cpp
+++ b/Engine/Plugins/Animation/PoseSearch/Source/Runtime/Private/PoseSearchInteractionIsland.cpp
@@ -16,9 +16,9 @@
 namespace UE::PoseSearch
 {
 
-#if ENABLE_ANIM_DEBUG
+#if ENABLE_DRAW_DEBUG && ENABLE_ANIM_DEBUG
 static TAutoConsoleVariable<bool> CVarPoseSearchInteractionShowIslands(TEXT("a.PoseSearchInteraction.ShowIslands"), false, TEXT("Show Pose Search Interaction Islands"));
-#endif
+#endif // ENABLE_DRAW_DEBUG && ENABLE_ANIM_DEBUG
 
 // Utils functions
 ///////////////////////////////////////////////////////////
@@ -303,7 +303,7 @@ FInteractionIsland::~FInteractionIsland()
 void FInteractionIsland::DebugDraw(const FColor& Color) const
 {
 	// called only by UPoseSearchInteractionSubsystem::Tick so no need to lock SearchResultsMutex to protect the read of SearchContexts
-#if ENABLE_DRAW_DEBUG
+#if ENABLE_DRAW_DEBUG && ENABLE_ANIM_DEBUG
 	check(IsInGameThread());
 
 	if (CVarPoseSearchInteractionShowIslands.GetValueOnAnyThread())
@@ -321,7 +321,7 @@ void FInteractionIsland::DebugDraw(const FColor& Color) const
 			}
 		}
 	}
-#endif // ENABLE_DRAW_DEBUG
+#endif // ENABLE_DRAW_DEBUG && ENABLE_ANIM_DEBUG
 }
 
 void FInteractionIsland::InjectToActor(AActor* Actor)
