posesearch
- cleanup MotionMatchMulti
- adding FAnimNode_MotionMatchingInteraction::bEnableWarping
- making consistent comment about "Experimental, this feature might be removed without warning, not for production use"

diff --git a/Engine/Plugins/Animation/PoseSearch/Source/Runtime/Private/AnimNode_MotionMatchingInteraction.cpp b/Engine/Plugins/Animation/PoseSearch/Source/Runtime/Private/AnimNode_MotionMatchingInteraction.cpp
--- a/Engine/Plugins/Animation/PoseSearch/Source/Runtime/Private/AnimNode_MotionMatchingInteraction.cpp
+++ b/Engine/Plugins/Animation/PoseSearch/Source/Runtime/Private/AnimNode_MotionMatchingInteraction.cpp
@@ -157,15 +157,24 @@ void FAnimNode_MotionMatchingInteraction::UpdateAssetPlayer(const FAnimationUpda
 		BlendLerp = FMath::Clamp(BlendLerp + (BlendSign * DeltaTime / BlendTime), 0.f, 1.f);
 	}
 	
-	const float WarpSign = NewInteractingRolesNum > 1 ? 1.f : -1.f;
+	const float WarpSign = bEnableWarping && NewInteractingRolesNum > 1 ? 1.f : -1.f;
 	if (InitialTranslationWarpTime > UE_KINDA_SMALL_NUMBER)
 	{
 		TranslationWarpLerp = FMath::Clamp(TranslationWarpLerp + (WarpSign * DeltaTime / InitialTranslationWarpTime), 0.f, 1.f);
 	}
+	else
+	{
+		TranslationWarpLerp = 0.f;
+	}
+
 	if (InitialRotationWarpTime > UE_KINDA_SMALL_NUMBER)
 	{
 		RotationWarpLerp = FMath::Clamp(RotationWarpLerp + (WarpSign * DeltaTime / InitialRotationWarpTime), 0.f, 1.f);
 	}
+	else
+	{
+		RotationWarpLerp = 0.f;
+	}
 
 	// updating source
 	const float SourceBlendLerp = 1.f - BlendLerp;
diff --git a/Engine/Plugins/Animation/PoseSearch/Source/Runtime/Private/PoseSearchLibrary.cpp b/Engine/Plugins/Animation/PoseSearch/Source/Runtime/Private/PoseSearchLibrary.cpp
--- a/Engine/Plugins/Animation/PoseSearch/Source/Runtime/Private/PoseSearchLibrary.cpp
+++ b/Engine/Plugins/Animation/PoseSearch/Source/Runtime/Private/PoseSearchLibrary.cpp
@@ -672,34 +672,6 @@ void UPoseSearchLibrary::MotionMatch(
 	MotionMatch(AnimInstances, Roles, AssetsToSearchConst, PoseHistoryName, ContinuingProperties, Future, Result);
 }
 
-void UPoseSearchLibrary::MotionMatchMulti(
-	TArray<ACharacter*> Characters,
-	TArray<FName> Roles,
-	TArray<UObject*> AssetsToSearch,
-	const FName PoseHistoryName,
-	const FPoseSearchContinuingProperties ContinuingProperties,
-	const FPoseSearchFutureProperties Future,
-	FPoseSearchBlueprintResult& Result)
-{
-	using namespace UE::PoseSearch;
-
-	FMemMark Mark(FMemStack::Get());
-
-	TArray<UAnimInstance*, TInlineAllocator<PreallocatedRolesNum, TMemStackAllocator<>>> AnimInstances;
-	for (ACharacter* Character : Characters)
-	{
-		UAnimInstance* AnimInstance = nullptr;
-		if (Character && Character->GetMesh())
-		{
-			AnimInstance = Character->GetMesh()->GetAnimInstance();
-		}
-		AnimInstances.Add(AnimInstance);
-	}
-
-	TArray<const UObject*>& AssetsToSearchConst = reinterpret_cast<TArray<const UObject*>&>(AssetsToSearch);
-	MotionMatch(AnimInstances, Roles, AssetsToSearchConst, PoseHistoryName, ContinuingProperties, Future, Result);
-}
-
 void UPoseSearchLibrary::MotionMatch(
 	const TArrayView<UAnimInstance*> AnimInstances,
 	const TArrayView<const UE::PoseSearch::FRole> Roles,
diff --git a/Engine/Plugins/Animation/PoseSearch/Source/Runtime/Public/PoseSearch/AnimNode_MotionMatchingInteraction.h b/Engine/Plugins/Animation/PoseSearch/Source/Runtime/Public/PoseSearch/AnimNode_MotionMatchingInteraction.h
--- a/Engine/Plugins/Animation/PoseSearch/Source/Runtime/Public/PoseSearch/AnimNode_MotionMatchingInteraction.h
+++ b/Engine/Plugins/Animation/PoseSearch/Source/Runtime/Public/PoseSearch/AnimNode_MotionMatchingInteraction.h
@@ -42,6 +42,9 @@ struct POSESEARCH_API FAnimNode_MotionMatchingInteraction : public FAnimNode_Ble
 	UPROPERTY(EditAnywhere, Category = Settings, meta = (PinHiddenByDefault))
 	bool bValidateResultAgainstAvailabilities = true;
 
+	UPROPERTY(EditAnywhere, Category = Settings, meta = (PinHiddenByDefault))
+	bool bEnableWarping = false;
+
 	// time from the beginning of the interaction to warp to full translation alignment
 	UPROPERTY(EditAnywhere, BlueprintReadWrite, Category = Warping, meta = (PinHiddenByDefault, ClampMin = "0"))
 	float InitialTranslationWarpTime = 0.2f;
diff --git a/Engine/Plugins/Animation/PoseSearch/Source/Runtime/Public/PoseSearch/AnimNode_PoseSearchHistoryCollector.h b/Engine/Plugins/Animation/PoseSearch/Source/Runtime/Public/PoseSearch/AnimNode_PoseSearchHistoryCollector.h
--- a/Engine/Plugins/Animation/PoseSearch/Source/Runtime/Public/PoseSearch/AnimNode_PoseSearchHistoryCollector.h
+++ b/Engine/Plugins/Animation/PoseSearch/Source/Runtime/Public/PoseSearch/AnimNode_PoseSearchHistoryCollector.h
@@ -38,7 +38,8 @@ public:
 	bool bStoreScales = false;
 
 	// time in seconds to recover to the reference skeleton root bone transform by RootBoneTranslationRecoveryRatio and RootBoneRotationRecoveryRatio
-	// from any eventual root bone modification. if zero the behaviour will be disabled (Experimental)
+	// from any eventual root bone modification. if zero the behaviour will be disabled
+	// Experimental, this feature might be removed without warning, not for production use
 	UPROPERTY(EditAnywhere, BlueprintReadWrite, Category = Experimental, meta = (ClampMin="0"))
 	float RootBoneRecoveryTime = 0.f;
 
@@ -58,7 +59,8 @@ public:
 	FLinearColor DebugColor = FLinearColor::Red;
 #endif // WITH_EDITORONLY_DATA
 
-	// if true Trajectory the pose history node will generate the trajectory using the TrajectoryData parameters instead of relying on the input Trajectory (Experimental)
+	// if true Trajectory the pose history node will generate the trajectory using the TrajectoryData parameters instead of relying on the input Trajectory
+	// Experimental, this feature might be removed without warning, not for production use
 	UPROPERTY(EditAnywhere, Category = Experimental)
 	bool bGenerateTrajectory = false;
 
diff --git a/Engine/Plugins/Animation/PoseSearch/Source/Runtime/Public/PoseSearch/PoseSearchDatabase.h b/Engine/Plugins/Animation/PoseSearch/Source/Runtime/Public/PoseSearch/PoseSearchDatabase.h
--- a/Engine/Plugins/Animation/PoseSearch/Source/Runtime/Public/PoseSearch/PoseSearchDatabase.h
+++ b/Engine/Plugins/Animation/PoseSearch/Source/Runtime/Public/PoseSearch/PoseSearchDatabase.h
@@ -35,7 +35,8 @@ enum class EPoseSearchMode : int32
 	// Optimized search mode: the database projects the poses into a PCA space using only the most significant "NumberOfPrincipalComponents" dimensions, and construct a kdtree to facilitate the search.
 	PCAKDTree,
 
-	// Optimized search mode using a vantage point tree (Experimental)
+	// Optimized search mode using a vantage point tree
+	// Experimental, this feature might be removed without warning, not for production use
 	VPTree UMETA(DisplayName = "VPTree (Experimental)")
 };
 
diff --git a/Engine/Plugins/Animation/PoseSearch/Source/Runtime/Public/PoseSearch/PoseSearchLibrary.h b/Engine/Plugins/Animation/PoseSearch/Source/Runtime/Public/PoseSearch/PoseSearchLibrary.h
--- a/Engine/Plugins/Animation/PoseSearch/Source/Runtime/Public/PoseSearch/PoseSearchLibrary.h
+++ b/Engine/Plugins/Animation/PoseSearch/Source/Runtime/Public/PoseSearch/PoseSearchLibrary.h
@@ -228,25 +228,6 @@ public:
 		const FPoseSearchFutureProperties Future,
 		FPoseSearchBlueprintResult& Result);
 
-	/**
-	* Implementation of the core motion matching algorithm for multiple characters
-	*
-	* @param AnimInstances					Input animation instances
-	* @param Roles							Input Roles associated to the animation instances
-	* @param AssetsToSearch					Input assets to search (UPoseSearchDatabase or any animation asset containing UAnimNotifyState_PoseSearchBranchIn)
-	* @param PoseHistoryName				Input tag of the associated PoseSearchHistoryCollector node in the anim graphs of the AnimInstances
-	* @param Result							Output FPoseSearchBlueprintResult with the search result
-	*/
-	UFUNCTION(BlueprintPure, Category = "Animation|Pose Search|Experimental", meta = (BlueprintThreadSafe, Keywords = "PoseMatch"))
-	static void MotionMatchMulti(
-		TArray<ACharacter*> Characters,
-		TArray<FName> Roles,
-		TArray<UObject*> AssetsToSearch,
-		const FName PoseHistoryName,
-		const FPoseSearchContinuingProperties ContinuingProperties,
-		const FPoseSearchFutureProperties Future,
-		FPoseSearchBlueprintResult& Result);
-
 	static void MotionMatch(
 		const TArrayView<UAnimInstance*> AnimInstances,
 		const TArrayView<const UE::PoseSearch::FRole> Roles,
diff --git a/Engine/Plugins/Animation/PoseSearch/Source/Runtime/Public/PoseSearch/PoseSearchSchema.h b/Engine/Plugins/Animation/PoseSearch/Source/Runtime/Public/PoseSearch/PoseSearchSchema.h
--- a/Engine/Plugins/Animation/PoseSearch/Source/Runtime/Public/PoseSearch/PoseSearchSchema.h
+++ b/Engine/Plugins/Animation/PoseSearch/Source/Runtime/Public/PoseSearch/PoseSearchSchema.h
@@ -20,9 +20,11 @@ enum class EPoseSearchDataPreprocessor : int32
 	Normalize,
 
 	// The data will be normalized against its deviation
+	// Experimental, this feature might be removed without warning, not for production use
 	NormalizeOnlyByDeviation UMETA(DisplayName = "Normalize Only By Deviation (Experimental)"),
 
 	// same behavior as Normalize, but it'll index all the databases in the normalization set with the same schema
+	// Experimental, this feature might be removed without warning, not for production use
 	NormalizeWithCommonSchema UMETA(DisplayName = "Normalize With Common Schema (Experimental)"),
 };
 
