pose search - fix for Deleting Blendspace samples from a blendspace while its still in Pose Search Database freezes the engine

diff --git a/Engine/Plugins/Animation/PoseSearch/Source/Runtime/Private/PoseSearchAssetSampler.cpp b/Engine/Plugins/Animation/PoseSearch/Source/Runtime/Private/PoseSearchAssetSampler.cpp
--- a/Engine/Plugins/Animation/PoseSearch/Source/Runtime/Private/PoseSearchAssetSampler.cpp
+++ b/Engine/Plugins/Animation/PoseSearch/Source/Runtime/Private/PoseSearchAssetSampler.cpp
@@ -105,7 +105,9 @@ static FTransform ExtractBlendSpaceRootMotion(float StartTime, float DeltaTime,
 {
 	FRootMotionMovementParams RootMotionParams;
 
-	if (DeltaTime != 0.f)
+	// looking for conditions that will calculate an FTransform::Identity to early out
+	const bool bIsTrivial = FMath::IsNearlyZero(DeltaTime) || CachedPlayLength < UE_SMALL_NUMBER || AccumulatedRootTransform.Num() <= 1;
+	if (!bIsTrivial)
 	{
 		bool const bPlayingBackwards = (DeltaTime < 0.f);
 
