pose search - fix for UPoseSearchDatabase::CachedAssetMap garbage collection

diff --git a/Engine/Plugins/Animation/PoseSearch/Source/Runtime/Private/PoseSearchDatabase.cpp b/Engine/Plugins/Animation/PoseSearch/Source/Runtime/Private/PoseSearchDatabase.cpp
--- a/Engine/Plugins/Animation/PoseSearch/Source/Runtime/Private/PoseSearchDatabase.cpp
+++ b/Engine/Plugins/Animation/PoseSearch/Source/Runtime/Private/PoseSearchDatabase.cpp
@@ -600,7 +600,7 @@ void UPoseSearchDatabase::UpdateCachedProperties()
 		}
 	}
 
-	for (TPair<TObjectPtr<UObject>, TArray<int32>>& CachedAssetMapPair : CachedAssetMap)
+	for (TPair<TWeakObjectPtr<UObject>, TArray<int32>>& CachedAssetMapPair : CachedAssetMap)
 	{
 		CachedAssetMapPair.Value.Sort();
 	}
diff --git a/Engine/Plugins/Animation/PoseSearch/Source/Runtime/Public/PoseSearch/PoseSearchDatabase.h b/Engine/Plugins/Animation/PoseSearch/Source/Runtime/Public/PoseSearch/PoseSearchDatabase.h
--- a/Engine/Plugins/Animation/PoseSearch/Source/Runtime/Public/PoseSearch/PoseSearchDatabase.h
+++ b/Engine/Plugins/Animation/PoseSearch/Source/Runtime/Public/PoseSearch/PoseSearchDatabase.h
@@ -489,7 +489,7 @@ private:
 	UE::PoseSearch::FSearchIndex SearchIndexPrivate;
 	
 	// CachedAssetMap is NOT serialized in operator<< but recalculated by UpdateCachedProperties every time SearchIndexPrivate changes
-	TMap<TObjectPtr<UObject>, TArray<int32>> CachedAssetMap;
+	TMap<TWeakObjectPtr<UObject>, TArray<int32>> CachedAssetMap;
 
 #if WITH_EDITOR
 	DECLARE_MULTICAST_DELEGATE(FOnDerivedDataRebuildMulticaster);
