MM interaction - WIP support for warping with IAnimRootMotionProvider

diff --git a/Engine/Plugins/Animation/PoseSearch/Source/Runtime/Private/AnimNode_MotionMatchingInteraction.cpp b/Engine/Plugins/Animation/PoseSearch/Source/Runtime/Private/AnimNode_MotionMatchingInteraction.cpp
--- a/Engine/Plugins/Animation/PoseSearch/Source/Runtime/Private/AnimNode_MotionMatchingInteraction.cpp
+++ b/Engine/Plugins/Animation/PoseSearch/Source/Runtime/Private/AnimNode_MotionMatchingInteraction.cpp
@@ -3,6 +3,7 @@
 #include "PoseSearch/AnimNode_MotionMatchingInteraction.h"
 #include "Animation/AnimInertializationSyncScope.h"
 #include "Animation/AnimInstanceProxy.h"
+#include "Animation/AnimRootMotionProvider.h"
 #include "PoseSearch/PoseSearchInteractionAsset.h"
 #include "PoseSearch/PoseSearchInteractionLibrary.h"
 #include "PoseSearch/PoseHistoryProvider.h"
@@ -65,12 +66,7 @@ void FAnimNode_MotionMatchingInteraction::UpdateAssetPlayer(const FAnimationUpda
 
 	int32 NewInteractingRolesNum = 0;
 	const float DeltaTime = Context.GetDeltaTime();
-	FPoseHistoryProvider* PoseHistoryProvider = Context.GetMessage<FPoseHistoryProvider>();
-	if (!PoseHistoryProvider)
-	{
-		UE_LOG(LogPoseSearch, Error, TEXT("FAnimNode_MotionMatchingInteraction::Update_AnyThread couldn't find the FPoseHistoryProvider"));
-	}
-	else
+	if (FPoseHistoryProvider* PoseHistoryProvider = Context.GetMessage<FPoseHistoryProvider>())
 	{
 		// we're NOT providing the UMultiAnimAsset here, relying on the interaction subsystem to resolve the ContinuingProperties properly
 		FPoseSearchContinuingProperties ContinuingProperties;
@@ -142,6 +138,34 @@ void FAnimNode_MotionMatchingInteraction::UpdateAssetPlayer(const FAnimationUpda
 				bBlendToExecuted = true;
 			}
 		}
+
+		if (bUseAnimRootMotionProvider)
+		{
+			const IPoseHistory& PoseHistory = PoseHistoryProvider->GetPoseHistory();
+			const USkeleton* Skeleton = Context.AnimInstanceProxy->GetSkeleton();
+			FTransform RootBoneTrasform;
+			PoseHistory.GetTransformAtTime(0.f, RootBoneTrasform, Skeleton, RootBoneIndexType, WorldSpaceIndexType);
+
+			const FTransform LerpedAlignedActorRootBoneTransform(
+				FQuat::Slerp(RootBoneTrasform.GetRotation(), FullAlignedActorRootBoneTransform.GetRotation(), RotationWarpLerp),
+				FMath::Lerp(RootBoneTrasform.GetTranslation(), FullAlignedActorRootBoneTransform.GetTranslation(), TranslationWarpLerp),
+				FVector::ZeroVector);
+
+#if ENABLE_ANIM_DEBUG
+			if (CVarAnimNodeMotionMatchingInteractionDebug.GetValueOnAnyThread())
+			{
+				Context.AnimInstanceProxy->AnimDrawDebugCoordinateSystem(RootBoneTrasform.GetLocation(), RootBoneTrasform.Rotator(), 25.f, false, 0.f, 0.f, SDPG_Foreground);
+				Context.AnimInstanceProxy->AnimDrawDebugCoordinateSystem(FullAlignedActorRootBoneTransform.GetLocation(), FullAlignedActorRootBoneTransform.Rotator(), 50.f, false, 0.f, 0.f, SDPG_Foreground);
+			}
+#endif // ENABLE_ANIM_DEBUG
+
+			WarpRootMotionTransform = LerpedAlignedActorRootBoneTransform.GetRelativeTransform(RootBoneTrasform);
+			check(WarpRootMotionTransform.IsRotationNormalized());
+		}
+	}
+	else
+	{
+		UE_LOG(LogPoseSearch, Error, TEXT("FAnimNode_MotionMatchingInteraction::Update_AnyThread couldn't find the FPoseHistoryProvider"));
 	}
 
 	const bool bDidBlendToRequestAnInertialBlend = bBlendToExecuted && bUseInertialBlend;
@@ -235,20 +259,34 @@ void FAnimNode_MotionMatchingInteraction::Evaluate_AnyThread(FPoseContext& Outpu
 
 	if (TranslationWarpLerp > UE_KINDA_SMALL_NUMBER || RotationWarpLerp > UE_KINDA_SMALL_NUMBER)
 	{
-		const FTransform ComponentTransform = Output.AnimInstanceProxy->GetComponentTransform();
-		const FTransform FullAlignedActorRootBoneLocalTransform = FullAlignedActorRootBoneTransform.GetRelativeTransform(ComponentTransform);
-
-#if ENABLE_ANIM_DEBUG
-		if (CVarAnimNodeMotionMatchingInteractionDebug.GetValueOnAnyThread())
+		if (bUseAnimRootMotionProvider)
 		{
-			Output.AnimInstanceProxy->AnimDrawDebugCoordinateSystem(FullAlignedActorRootBoneTransform.GetLocation(), FullAlignedActorRootBoneTransform.Rotator(), 15.f, false, 0.f, 0.f, SDPG_Foreground);
-			Output.AnimInstanceProxy->AnimDrawDebugCoordinateSystem(ComponentTransform.GetLocation(), ComponentTransform.Rotator(), 5.f, false, 0.f, 0.f, SDPG_Foreground);
+			if (const UE::Anim::IAnimRootMotionProvider* RootMotionProvider = UE::Anim::IAnimRootMotionProvider::Get())
+			{
+				RootMotionProvider->OverrideRootMotion(WarpRootMotionTransform, Output.CustomAttributes);
+			}
+			else
+			{
+				UE_LOG(LogPoseSearch, Error, TEXT("FAnimNode_MotionMatchingInteraction::Evaluate_AnyThread couldn't find the IAnimRootMotionProvider"));
+			}
 		}
+		else
+		{
+			const FTransform ComponentTransform = Output.AnimInstanceProxy->GetComponentTransform();
+			const FTransform FullAlignedActorRootBoneLocalTransform = FullAlignedActorRootBoneTransform.GetRelativeTransform(ComponentTransform);
+
+#if ENABLE_ANIM_DEBUG
+			if (CVarAnimNodeMotionMatchingInteractionDebug.GetValueOnAnyThread())
+			{
+				Output.AnimInstanceProxy->AnimDrawDebugCoordinateSystem(FullAlignedActorRootBoneTransform.GetLocation(), FullAlignedActorRootBoneTransform.Rotator(), 15.f, false, 0.f, 0.f, SDPG_Foreground);
+				Output.AnimInstanceProxy->AnimDrawDebugCoordinateSystem(ComponentTransform.GetLocation(), ComponentTransform.Rotator(), 5.f, false, 0.f, 0.f, SDPG_Foreground);
+			}
 #endif // ENABLE_ANIM_DEBUG
 
-		const FCompactPoseBoneIndex RootBoneIndex(0);
-		FTransform& RootBoneTransform = Output.Pose[RootBoneIndex];
-		RootBoneTransform.SetTranslation(FMath::Lerp(RootBoneTransform.GetTranslation(), FullAlignedActorRootBoneLocalTransform.GetTranslation(), TranslationWarpLerp));
-		RootBoneTransform.SetRotation(FQuat::Slerp(RootBoneTransform.GetRotation(), FullAlignedActorRootBoneLocalTransform.GetRotation(), RotationWarpLerp));
+			const FCompactPoseBoneIndex RootBoneIndex(0);
+			FTransform& RootBoneTransform = Output.Pose[RootBoneIndex];
+			RootBoneTransform.SetTranslation(FMath::Lerp(RootBoneTransform.GetTranslation(), FullAlignedActorRootBoneLocalTransform.GetTranslation(), TranslationWarpLerp));
+			RootBoneTransform.SetRotation(FQuat::Slerp(RootBoneTransform.GetRotation(), FullAlignedActorRootBoneLocalTransform.GetRotation(), RotationWarpLerp));
+		}
 	}
 }
diff --git a/Engine/Plugins/Animation/PoseSearch/Source/Runtime/Private/MotionMatchingInteractionAnimNodeLibrary.cpp b/Engine/Plugins/Animation/PoseSearch/Source/Runtime/Private/MotionMatchingInteractionAnimNodeLibrary.cpp
--- a/Engine/Plugins/Animation/PoseSearch/Source/Runtime/Private/MotionMatchingInteractionAnimNodeLibrary.cpp
+++ b/Engine/Plugins/Animation/PoseSearch/Source/Runtime/Private/MotionMatchingInteractionAnimNodeLibrary.cpp
@@ -18,4 +18,26 @@ void UMotionMatchingInteractionAnimNodeLibrary::SetAvailabilities(const FMotionM
 	{
 		UE_LOG(LogPoseSearch, Warning, TEXT("UMotionMatchingInteractionAnimNodeLibrary::SetAvailabilities called on an invalid context or with an invalid type"));
 	}
+}
+
+float UMotionMatchingInteractionAnimNodeLibrary::GetTranslationWarpLerp(const FMotionMatchingInteractionAnimNodeReference& MotionMatchingInteractionNode)
+{
+	if (const FAnimNode_MotionMatchingInteraction* MotionMatchingInteractionNodePtr = MotionMatchingInteractionNode.GetAnimNodePtr<FAnimNode_MotionMatchingInteraction>())
+	{
+		return MotionMatchingInteractionNodePtr->GetTranslationWarpLerp();
+	}
+
+	UE_LOG(LogPoseSearch, Warning, TEXT("UMotionMatchingInteractionAnimNodeLibrary::GetTranslationWarpLerp called on an invalid context or with an invalid type"));
+	return 0.f;
+}
+
+float UMotionMatchingInteractionAnimNodeLibrary::GetRotationWarpLerp(const FMotionMatchingInteractionAnimNodeReference& MotionMatchingInteractionNode)
+{
+	if (FAnimNode_MotionMatchingInteraction* MotionMatchingInteractionNodePtr = MotionMatchingInteractionNode.GetAnimNodePtr<FAnimNode_MotionMatchingInteraction>())
+	{
+		return MotionMatchingInteractionNodePtr->GetRotationWarpLerp();
+	}
+	
+	UE_LOG(LogPoseSearch, Warning, TEXT("UMotionMatchingInteractionAnimNodeLibrary::GetRotationWarpLerp called on an invalid context or with an invalid type"));
+	return 0.f;
 }
\ No newline at end of file
diff --git a/Engine/Plugins/Animation/PoseSearch/Source/Runtime/Public/PoseSearch/AnimNode_MotionMatchingInteraction.h b/Engine/Plugins/Animation/PoseSearch/Source/Runtime/Public/PoseSearch/AnimNode_MotionMatchingInteraction.h
--- a/Engine/Plugins/Animation/PoseSearch/Source/Runtime/Public/PoseSearch/AnimNode_MotionMatchingInteraction.h
+++ b/Engine/Plugins/Animation/PoseSearch/Source/Runtime/Public/PoseSearch/AnimNode_MotionMatchingInteraction.h
@@ -43,8 +43,11 @@ struct POSESEARCH_API FAnimNode_MotionMatchingInteraction : public FAnimNode_Ble
 	bool bValidateResultAgainstAvailabilities = true;
 
 	UPROPERTY(EditAnywhere, Category = Settings, meta = (PinHiddenByDefault))
-	bool bEnableWarping = false;
+	bool bEnableWarping = true;
 
+	UPROPERTY(EditAnywhere, Category = Settings, meta = (PinHiddenByDefault))
+	bool bUseAnimRootMotionProvider = true;
+	
 	// time from the beginning of the interaction to warp to full translation alignment
 	UPROPERTY(EditAnywhere, BlueprintReadWrite, Category = Warping, meta = (PinHiddenByDefault, ClampMin = "0"))
 	float InitialTranslationWarpTime = 0.2f;
@@ -64,6 +67,9 @@ struct POSESEARCH_API FAnimNode_MotionMatchingInteraction : public FAnimNode_Ble
 
 	virtual void Reset() override;
 
+	float GetTranslationWarpLerp() const { return TranslationWarpLerp; }
+	float GetRotationWarpLerp() const { return RotationWarpLerp; }
+
 protected:
 	float BlendLerp = 0.f;
 	float TranslationWarpLerp = 0.f;
@@ -78,4 +84,5 @@ protected:
 
 	// wanted world transform for full aligment interaction
 	FTransform FullAlignedActorRootBoneTransform = FTransform::Identity;
+	FTransform WarpRootMotionTransform = FTransform::Identity;
 };
diff --git a/Engine/Plugins/Animation/PoseSearch/Source/Runtime/Public/PoseSearch/MotionMatchingInteractionAnimNodeLibrary.h b/Engine/Plugins/Animation/PoseSearch/Source/Runtime/Public/PoseSearch/MotionMatchingInteractionAnimNodeLibrary.h
--- a/Engine/Plugins/Animation/PoseSearch/Source/Runtime/Public/PoseSearch/MotionMatchingInteractionAnimNodeLibrary.h
+++ b/Engine/Plugins/Animation/PoseSearch/Source/Runtime/Public/PoseSearch/MotionMatchingInteractionAnimNodeLibrary.h
@@ -36,4 +36,10 @@ public:
 
 	UFUNCTION(BlueprintCallable, Category = "Animation|Pose Search", meta = (BlueprintThreadSafe))
 	static void SetAvailabilities(const FMotionMatchingInteractionAnimNodeReference& MotionMatchingInteractionNode, const TArray<FPoseSearchInteractionAvailability>& Availabilities);
+
+	UFUNCTION(BlueprintPure, Category = "Animation|Pose Search", meta = (BlueprintThreadSafe))
+	static float GetTranslationWarpLerp(const FMotionMatchingInteractionAnimNodeReference& MotionMatchingInteractionNode);
+
+	UFUNCTION(BlueprintPure, Category = "Animation|Pose Search", meta = (BlueprintThreadSafe))
+	static float GetRotationWarpLerp(const FMotionMatchingInteractionAnimNodeReference& MotionMatchingInteractionNode);
 };
\ No newline at end of file
