posesearch - FAnimNode_MotionMatchingInteraction cleanup

diff --git a/Engine/Plugins/Animation/PoseSearch/Source/Runtime/Private/AnimNode_MotionMatchingInteraction.cpp b/Engine/Plugins/Animation/PoseSearch/Source/Runtime/Private/AnimNode_MotionMatchingInteraction.cpp
--- a/Engine/Plugins/Animation/PoseSearch/Source/Runtime/Private/AnimNode_MotionMatchingInteraction.cpp
+++ b/Engine/Plugins/Animation/PoseSearch/Source/Runtime/Private/AnimNode_MotionMatchingInteraction.cpp
@@ -28,22 +28,9 @@ void FAnimNode_MotionMatchingInteraction::GatherDebugData(FNodeDebugData& DebugD
 	DebugData.AddDebugItem(DebugLine);
 }
 
-void FAnimNode_MotionMatchingInteraction::Initialize_AnyThread(const FAnimationInitializeContext& Context)
-{
-	Super::Initialize_AnyThread(Context);
-	Source.Initialize(Context);
-}
-
-void FAnimNode_MotionMatchingInteraction::CacheBones_AnyThread(const FAnimationCacheBonesContext& Context)
-{
-	Super::CacheBones_AnyThread(Context);
-	Source.CacheBones(Context);
-}
-
 void FAnimNode_MotionMatchingInteraction::Reset()
 {
 	Super::Reset();
-	BlendLerp = 0.f;
 	TranslationWarpLerp = 0.f;
 	RotationWarpLerp = 0.f;
 	InteractingRolesNum = 0;
@@ -175,12 +162,6 @@ void FAnimNode_MotionMatchingInteraction::UpdateAssetPlayer(const FAnimationUpda
 	UpdateBlendspaceParameters(BlendspaceUpdateMode, BlendParameters);
 
 	// calculating the translation and rotation warp lerps, used to warp the root transform towards the last computed FullAlignedActorRootBoneTransform
-	if (BlendTime > UE_KINDA_SMALL_NUMBER)
-	{
-		const float BlendSign = NewInteractingRolesNum > 0 ? 1.f : -1.f;
-		BlendLerp = FMath::Clamp(BlendLerp + (BlendSign * DeltaTime / BlendTime), 0.f, 1.f);
-	}
-	
 	const float WarpSign = bEnableWarping && NewInteractingRolesNum > 1 ? 1.f : -1.f;
 	if (InitialTranslationWarpTime > UE_KINDA_SMALL_NUMBER)
 	{
@@ -200,34 +181,8 @@ void FAnimNode_MotionMatchingInteraction::UpdateAssetPlayer(const FAnimationUpda
 		RotationWarpLerp = 0.f;
 	}
 
-	// updating source
-	const float SourceBlendLerp = 1.f - BlendLerp;
-	FAnimationUpdateContext SourceContext = Context.FractionalWeightAndRootMotion(SourceBlendLerp, SourceBlendLerp);
-	if (BlendLerp > UE_KINDA_SMALL_NUMBER)
-	{
-		SourceContext = SourceContext.AsInactive();
-	}
-	Source.Update(SourceContext);
-	
-	if (Source.GetLinkNode())
-	{
-		// updating blend stack
-		if (BlendLerp < UE_KINDA_SMALL_NUMBER && DeltaTime > UE_KINDA_SMALL_NUMBER)
-		{
-			// resetting the blendstack if there's no BlendLerp weight to it
-			Reset();
-		}
-
-		FAnimationUpdateContext BlendStackContext = Context.FractionalWeightAndRootMotion(BlendLerp, BlendLerp);
-
-		// bypassing FAnimNode_BlendStack::UpdateAssetPlayer, since we overridden its behaviour
-		FAnimNode_BlendStack_Standalone::UpdateAssetPlayer(BlendStackContext);
-	}
-	else
-	{
-		// bypassing FAnimNode_BlendStack::UpdateAssetPlayer, since we overridden its behaviour
-		FAnimNode_BlendStack_Standalone::UpdateAssetPlayer(Context);
-	}
+	// bypassing FAnimNode_BlendStack::UpdateAssetPlayer, since we overridden its behaviour
+	FAnimNode_BlendStack_Standalone::UpdateAssetPlayer(Context);
 
 #if ENABLE_ANIM_DEBUG
 	if (UE_TRACE_CHANNELEXPR_IS_ENABLED(AnimationChannel))
@@ -235,7 +190,6 @@ void FAnimNode_MotionMatchingInteraction::UpdateAssetPlayer(const FAnimationUpda
 		TRACE_ANIM_NODE_VALUE(Context, *FString("InteractingRolesNum"), InteractingRolesNum);
 		TRACE_ANIM_NODE_VALUE(Context, *FString("NewInteractingRolesNum"), NewInteractingRolesNum);
 		TRACE_ANIM_NODE_VALUE(Context, *FString("BlendToExecuted"), bBlendToExecuted);
-		TRACE_ANIM_NODE_VALUE(Context, *FString("BlendLerp"), BlendLerp);
 		TRACE_ANIM_NODE_VALUE(Context, *FString("TranslationWarpLerp"), TranslationWarpLerp);
 		TRACE_ANIM_NODE_VALUE(Context, *FString("RotationWarpLerp"), RotationWarpLerp);
 	}
@@ -248,14 +202,7 @@ void FAnimNode_MotionMatchingInteraction::Evaluate_AnyThread(FPoseContext& Outpu
 {
 	check(Output.AnimInstanceProxy);
 
-	// evaluating Source to get the base pose 
-	Source.Evaluate(Output);
-
-	// @todo perform a proper blend between Source and blend stack. In the current WIP implementation blend stack will stomp over Source 
-	if (!AnimPlayers.IsEmpty())
-	{
-		Super::Evaluate_AnyThread(Output);
-	}
+	Super::Evaluate_AnyThread(Output);
 
 	if (TranslationWarpLerp > UE_KINDA_SMALL_NUMBER || RotationWarpLerp > UE_KINDA_SMALL_NUMBER)
 	{
diff --git a/Engine/Plugins/Animation/PoseSearch/Source/Runtime/Public/PoseSearch/AnimNode_MotionMatchingInteraction.h b/Engine/Plugins/Animation/PoseSearch/Source/Runtime/Public/PoseSearch/AnimNode_MotionMatchingInteraction.h
--- a/Engine/Plugins/Animation/PoseSearch/Source/Runtime/Public/PoseSearch/AnimNode_MotionMatchingInteraction.h
+++ b/Engine/Plugins/Animation/PoseSearch/Source/Runtime/Public/PoseSearch/AnimNode_MotionMatchingInteraction.h
@@ -30,9 +30,6 @@ struct POSESEARCH_API FAnimNode_MotionMatchingInteraction : public FAnimNode_Ble
 {
 	GENERATED_BODY();
 
-	UPROPERTY(EditAnywhere, BlueprintReadWrite, Category = Links)
-	FPoseLink Source;
-
 	UPROPERTY(EditAnywhere, Category = Settings)
 	EMotionMatchingInteractionEvaluationMode EvaluationMode = EMotionMatchingInteractionEvaluationMode::ContinuousReselection;
 
@@ -58,9 +55,7 @@ struct POSESEARCH_API FAnimNode_MotionMatchingInteraction : public FAnimNode_Ble
 
 	// FAnimNode_Base interface
 	virtual void GatherDebugData(FNodeDebugData& DebugData) override;
-	virtual void Initialize_AnyThread(const FAnimationInitializeContext& Context) override;
 	virtual void Evaluate_AnyThread(FPoseContext& Output) override;
-	virtual void CacheBones_AnyThread(const FAnimationCacheBonesContext& Context) override;
 
 	virtual void UpdateAssetPlayer(const FAnimationUpdateContext& Context) override;
 	// End of FAnimNode_Base interface
@@ -71,7 +66,6 @@ struct POSESEARCH_API FAnimNode_MotionMatchingInteraction : public FAnimNode_Ble
 	float GetRotationWarpLerp() const { return RotationWarpLerp; }
 
 protected:
-	float BlendLerp = 0.f;
 	float TranslationWarpLerp = 0.f;
 	float RotationWarpLerp = 0.f;
 	// if a search is successful InteractingRolesNum > 0. 
