posesearch - database asset picker and browser support for UMultiAnimAsset

diff --git a/Engine/Plugins/Animation/PoseSearch/Source/Editor/Private/PoseSearchDatabaseAssetBrowser.cpp b/Engine/Plugins/Animation/PoseSearch/Source/Editor/Private/PoseSearchDatabaseAssetBrowser.cpp
--- a/Engine/Plugins/Animation/PoseSearch/Source/Editor/Private/PoseSearchDatabaseAssetBrowser.cpp
+++ b/Engine/Plugins/Animation/PoseSearch/Source/Editor/Private/PoseSearchDatabaseAssetBrowser.cpp
@@ -1,4 +1,4 @@
-﻿// Copyright Epic Games, Inc. All Rights Reserved.
+// Copyright Epic Games, Inc. All Rights Reserved.
 
 #include "PoseSearchDatabaseAssetBrowser.h"
 
@@ -12,6 +12,7 @@
 #include "Widgets/SBoxPanel.h"
 #include "UObject/AssetRegistryTagsContext.h"
 #include "PoseSearchDatabaseViewModel.h"
+#include "PoseSearch/MultiAnimAsset.h"
 #include "PoseSearch/PoseSearchDatabase.h"
 #include "PoseSearch/PoseSearchSchema.h"
 #include "Subsystems/AssetEditorSubsystem.h"
@@ -46,11 +47,10 @@ void SPoseSearchDatabaseAssetBrowser::Construct(const FArguments& InArgs, TShare
 
 void SPoseSearchDatabaseAssetBrowser::RefreshView()
 {
-	// @TODO: Add support for MultiAnimAsset.
-	
 	FAssetPickerConfig AssetPickerConfig;
 	
 	AssetPickerConfig.Filter.ClassPaths.Add(UAnimationAsset::StaticClass()->GetClassPathName());
+	AssetPickerConfig.Filter.ClassPaths.Add(UMultiAnimAsset::StaticClass()->GetClassPathName());
 	AssetPickerConfig.Filter.bRecursiveClasses = true;
 	
 	AssetPickerConfig.bAddFilterUI = true;
@@ -91,64 +91,31 @@ SPoseSearchDatabaseAssetBrowser::~SPoseSearchDatabaseAssetBrowser()
 
 void SPoseSearchDatabaseAssetBrowser::OnAssetDoubleClicked(const FAssetData& AssetData)
 {
-	if (!AssetData.GetAsset())
-	{
-		return;
-	}
-
-	UAnimationAsset* NewAnimationAsset = Cast<UAnimationAsset>(AssetData.GetAsset());
-	if (!NewAnimationAsset)
+	if (const UObject* Asset = AssetData.GetAsset())
 	{
-		return;
+		if (Cast<UAnimationAsset>(Asset) || Cast<UMultiAnimAsset>(Asset))
+		{
+			GEditor->GetEditorSubsystem<UAssetEditorSubsystem>()->OpenEditorForAsset(Asset);
+		}
 	}
-
-	// Just open asset in persona.
-	GEditor->GetEditorSubsystem<UAssetEditorSubsystem>()->OpenEditorForAsset(NewAnimationAsset);
 }
 
 bool SPoseSearchDatabaseAssetBrowser::OnShouldFilterAsset(const FAssetData& AssetData)
 {
-	bool bAssetHasCompatibleSkeleton = false;
-
-	if (DatabaseViewModel)
-	{
-		bAssetHasCompatibleSkeleton = FPoseSearchEditorUtils::IsAssetCompatibleWithDatabase(DatabaseViewModel->GetPoseSearchDatabase(), AssetData);
-	}
-	
-	if (AssetData.GetClass()->IsChildOf(UAnimSequence::StaticClass()) && bAssetHasCompatibleSkeleton)
-	{
-		return false;
-	}
-
-	if (AssetData.GetClass()->IsChildOf(UAnimComposite::StaticClass()) && bAssetHasCompatibleSkeleton)
-	{
-		return false;
-	}
-
-	if (AssetData.GetClass()->IsChildOf(UAnimMontage::StaticClass()) && bAssetHasCompatibleSkeleton)
-	{
-		return false;
-	}
-	
-	if (AssetData.GetClass()->IsChildOf(UBlendSpace::StaticClass()) && bAssetHasCompatibleSkeleton)
-	{
-		return false;
-	}
-	
-	return true;
+	const bool bIsAssetCompatibleWithDatabase = DatabaseViewModel && FPoseSearchEditorUtils::IsAssetCompatibleWithDatabase(DatabaseViewModel->GetPoseSearchDatabase(), AssetData);
+	return !bIsAssetCompatibleWithDatabase;
 }
 	
 void SPoseSearchDatabaseAssetBrowser::OnObjectPropertyChanged(UObject* InObject, FPropertyChangedEvent& InPropertyChangedEvent) const
 {
-	if (DatabaseViewModel)
+	if (DatabaseViewModel && InObject)
 	{
-		if (const UPoseSearchDatabase* Database = DatabaseViewModel->GetPoseSearchDatabase())
+		const UPoseSearchDatabase* Database = DatabaseViewModel->GetPoseSearchDatabase();
+		if (Database && Database->Schema == InObject ||
+			InObject->GetClass()->IsChildOf(UAnimationAsset::StaticClass()) ||
+			InObject->GetClass()->IsChildOf(UMultiAnimAsset::StaticClass()))
 		{
-			// Only refresh asset browser if our target schema has changed.
-			if (Database->Schema && Database->Schema == InObject)
-			{
-				RefreshAssetViewDelegate.ExecuteIfBound(true);	
-			}
+			RefreshAssetViewDelegate.ExecuteIfBound(true);
 		}
 	}
 }
diff --git a/Engine/Plugins/Animation/PoseSearch/Source/Editor/Private/PoseSearchDatabaseEditorUtils.cpp b/Engine/Plugins/Animation/PoseSearch/Source/Editor/Private/PoseSearchDatabaseEditorUtils.cpp
--- a/Engine/Plugins/Animation/PoseSearch/Source/Editor/Private/PoseSearchDatabaseEditorUtils.cpp
+++ b/Engine/Plugins/Animation/PoseSearch/Source/Editor/Private/PoseSearchDatabaseEditorUtils.cpp
@@ -1,7 +1,8 @@
-﻿// Copyright Epic Games, Inc. All Rights Reserved.
+// Copyright Epic Games, Inc. All Rights Reserved.
 
 #include "PoseSearchDatabaseEditorUtils.h"
 
+#include "PoseSearch/MultiAnimAsset.h"
 #include "PoseSearch/PoseSearchDatabase.h"
 #include "PoseSearch/PoseSearchSchema.h"
 
@@ -13,16 +14,46 @@ namespace UE::PoseSearch
 		{
 			if (InDatabase->Schema)
 			{
-				TArray<FPoseSearchRoledSkeleton> RoledSkeletons = InDatabase->Schema->GetRoledSkeletons();
-
-				for (const FPoseSearchRoledSkeleton& RoledSkeleton : RoledSkeletons)
+				const TArray<FPoseSearchRoledSkeleton> RoledSkeletons = InDatabase->Schema->GetRoledSkeletons();
+				if (RoledSkeletons.Num() == 1 )
 				{
-					if (RoledSkeleton.Skeleton && RoledSkeleton.Skeleton->IsCompatibleForEditor(InAssetData))
+					if (InAssetData.GetClass()->IsChildOf(UAnimationAsset::StaticClass()) &&
+						RoledSkeletons[0].Skeleton &&
+						RoledSkeletons[0].Skeleton->IsCompatibleForEditor(InAssetData))
 					{
 						// We found a compatible skeleton in the schema.
 						return true;
 					}
 				}
+				else if (RoledSkeletons.Num() > 1)
+				{
+					if (const UMultiAnimAsset* MultiAnimAsset = Cast<UMultiAnimAsset>(InAssetData.GetAsset()))
+					{
+						if (MultiAnimAsset->GetNumRoles() != RoledSkeletons.Num())
+						{
+							return false;
+						}
+						
+						for (const FPoseSearchRoledSkeleton& RoledSkeleton : RoledSkeletons)
+						{
+							if (const UAnimationAsset* RoledAnimationAsset = MultiAnimAsset->GetAnimationAsset(RoledSkeleton.Role))
+							{
+								if (!RoledAnimationAsset->GetSkeleton()->IsCompatibleForEditor(RoledSkeleton.Skeleton))
+								{
+									return false;
+								}
+							}
+							else
+							{
+								// couldn't find a necessary asset for the RoledSkeleton
+								return false;
+							}
+						}
+
+						// we passed all the compatibility requirements!
+						return true;
+					}
+				}
 			}
 		}
 		
diff --git a/Engine/Plugins/Animation/PoseSearch/Source/Editor/Private/PoseSearchDatabaseViewModel.cpp b/Engine/Plugins/Animation/PoseSearch/Source/Editor/Private/PoseSearchDatabaseViewModel.cpp
--- a/Engine/Plugins/Animation/PoseSearch/Source/Editor/Private/PoseSearchDatabaseViewModel.cpp
+++ b/Engine/Plugins/Animation/PoseSearch/Source/Editor/Private/PoseSearchDatabaseViewModel.cpp
@@ -731,7 +731,7 @@ bool FDatabaseViewModel::SetAnimationAsset(int32 AnimationAssetIndex, UObject* A
 			{
 				// Ensure that our target database item matches the input object's class.
 				const UClass* AssetClass = AnimAsset->GetClass();
-				if (AssetClass == DatabaseAnimationAsset->GetAnimationAssetStaticClass())
+				if (AssetClass->IsChildOf(DatabaseAnimationAsset->GetAnimationAssetStaticClass()))
 				{
 					if (AssetClass->IsChildOf( UAnimSequence::StaticClass()))
 					{
