posesearch - defaulting FPoseSearchInteractionAvailability::MaxCost to 0

diff --git a/Engine/Plugins/Animation/PoseSearch/Source/Runtime/Private/PoseSearchInteractionSubsystem.cpp b/Engine/Plugins/Animation/PoseSearch/Source/Runtime/Private/PoseSearchInteractionSubsystem.cpp
--- a/Engine/Plugins/Animation/PoseSearch/Source/Runtime/Private/PoseSearchInteractionSubsystem.cpp
+++ b/Engine/Plugins/Animation/PoseSearch/Source/Runtime/Private/PoseSearchInteractionSubsystem.cpp
@@ -97,7 +97,12 @@ namespace UE::PoseSearch
 				{
 					// we already found the FRoledAnimInstances. Updating the BroadPhaseRadius and MaxCost accordingly
 					RoledAnimInstance.BroadPhaseRadius = FMath::Max(RoledAnimInstance.BroadPhaseRadius, BroadPhaseRadius);
-					RoledAnimInstance.MaxCost = FMath::Min(RoledAnimInstance.MaxCost, MaxCost);
+					
+					// MaxCost is valid only if greater than zero
+					if (MaxCost > 0.f)
+					{
+						RoledAnimInstance.MaxCost = FMath::Min(RoledAnimInstance.MaxCost, MaxCost);
+					}
 					return;
 				}
 			}
@@ -108,7 +113,10 @@ namespace UE::PoseSearch
 			NewRoledAnimInstance.HistoryCollector = HistoryCollector;
 			NewRoledAnimInstance.Role = Role;
 			NewRoledAnimInstance.BroadPhaseRadius = BroadPhaseRadius;
-			NewRoledAnimInstance.MaxCost = MaxCost;
+			if (MaxCost > 0.f)
+			{
+				NewRoledAnimInstance.MaxCost = MaxCost;
+			}
 
 			// @todo: use a better datastructure to keep things FRoledAnimInstance(s) sorted? Or use insertion sort?
 			RoledAnimInstances.Sort([](const FRoledAnimInstance& A, const FRoledAnimInstance& B)
diff --git a/Engine/Plugins/Animation/PoseSearch/Source/Runtime/Public/PoseSearch/PoseSearchInteractionLibrary.h b/Engine/Plugins/Animation/PoseSearch/Source/Runtime/Public/PoseSearch/PoseSearchInteractionLibrary.h
--- a/Engine/Plugins/Animation/PoseSearch/Source/Runtime/Public/PoseSearch/PoseSearchInteractionLibrary.h
+++ b/Engine/Plugins/Animation/PoseSearch/Source/Runtime/Public/PoseSearch/PoseSearchInteractionLibrary.h
@@ -34,9 +34,9 @@ struct POSESEARCH_API FPoseSearchInteractionAvailability
 	UPROPERTY(EditAnywhere, BlueprintReadWrite, Category=Settings)
 	float BroadPhaseRadius = 500.f;
 
-	// the associated character to this FPoseSearchInteractionAvailability will not partecipate in an interaction if the motion matching search cost result is higher than MaxCost
+	// if MaxCost if greater than zero, the associated character to this FPoseSearchInteractionAvailability will not partecipate in an interaction if the motion matching search cost result is higher than MaxCost
 	UPROPERTY(EditAnywhere, BlueprintReadWrite, Category=Settings)
-	float MaxCost = 1000.f;
+	float MaxCost = 0.f;
 };
 
 USTRUCT(Experimental, BlueprintType, Category="Animation|Pose Search")
