posesearch - initial supporting null AnimInstace

diff --git a/Engine/Plugins/Animation/PoseSearch/Source/Runtime/Private/PoseSearchLibrary.cpp b/Engine/Plugins/Animation/PoseSearch/Source/Runtime/Private/PoseSearchLibrary.cpp
--- a/Engine/Plugins/Animation/PoseSearch/Source/Runtime/Private/PoseSearchLibrary.cpp
+++ b/Engine/Plugins/Animation/PoseSearch/Source/Runtime/Private/PoseSearchLibrary.cpp
@@ -295,14 +295,16 @@ void UPoseSearchLibrary::TraceMotionMatching(UE::PoseSearch::FSearchContext& Sea
 
 	for (int32 AnimInstanceIndex = 0; AnimInstanceIndex < AnimInstancesNum; ++AnimInstanceIndex)
 	{
-		const UAnimInstance* AnimInstance = SearchContext.GetAnimInstances()[AnimInstanceIndex];
-		const UObject* SkeletalMeshComponent = AnimInstance->GetOuter();
+		if (const UAnimInstance* AnimInstance = SearchContext.GetAnimInstances()[AnimInstanceIndex])
+		{
+			const UObject* SkeletalMeshComponent = AnimInstance->GetOuter();
 
-		TRACE_OBJECT(AnimInstance);
+			TRACE_OBJECT(AnimInstance);
 
-		TraceState.SkeletalMeshComponentIds[AnimInstanceIndex] = FObjectTrace::GetObjectId(SkeletalMeshComponent);
+			TraceState.SkeletalMeshComponentIds[AnimInstanceIndex] = FObjectTrace::GetObjectId(SkeletalMeshComponent);
 
-		SearchId = HashCombineFast(SearchId, GetTypeHash(FObjectTrace::GetObjectId(AnimInstance)));
+			SearchId = HashCombineFast(SearchId, GetTypeHash(FObjectTrace::GetObjectId(AnimInstance)));
+		}
 	}
 
 	TraceState.Roles.SetNum(AnimInstancesNum);
@@ -863,7 +865,7 @@ UE::PoseSearch::FSearchResult UPoseSearchLibrary::MotionMatch(
 	TAssetsToSearchPerDatabaseMap AssetsToSearchPerDatabaseMap;
 	
 	// collecting all the possible continuing pose search (it could be multiple searches, but most likely only one)
-	const float DeltaSeconds = AnimInstances[0]->GetDeltaSeconds();
+	const float DeltaSeconds = AnimInstances[0] ? AnimInstances[0]->GetDeltaSeconds() : FiniteDelta;
 	if (const UObject* PlayingAnimationAsset = ContinuingProperties.PlayingAsset.Get())
 	{
 		AddToSearch(AssetsToSearchPerDatabaseMap, PlayingAnimationAsset);
@@ -1030,8 +1032,7 @@ UE::PoseSearch::FSearchResult UPoseSearchLibrary::MotionMatch(
 	FMotionMatchingState MotionMatchingState;
 	MotionMatchingState.CurrentSearchResult = SearchResult;
 	MotionMatchingState.ElapsedPoseSearchTime = 0.0f;
-	TraceMotionMatching(SearchContext, MotionMatchingState,
-		DeltaSeconds, true, FObjectTrace::GetWorldElapsedTime(AnimInstances[0]->GetWorld()));
+	TraceMotionMatching(SearchContext, MotionMatchingState, DeltaSeconds, true, FObjectTrace::GetWorldElapsedTime(AnimInstances[0] ? AnimInstances[0]->GetWorld() : nullptr));
 #endif // UE_POSE_SEARCH_TRACE_ENABLED
 
 	return SearchResult;
diff --git a/Engine/Plugins/Animation/PoseSearch/Source/Runtime/Private/PoseSearchResult.cpp b/Engine/Plugins/Animation/PoseSearch/Source/Runtime/Private/PoseSearchResult.cpp
--- a/Engine/Plugins/Animation/PoseSearch/Source/Runtime/Private/PoseSearchResult.cpp
+++ b/Engine/Plugins/Animation/PoseSearch/Source/Runtime/Private/PoseSearchResult.cpp
@@ -39,13 +39,6 @@ void FSearchResult::Update(float NewAssetTime)
 	}
 }
 
-void FSearchResult::Reset()
-{
-	PoseIdx = INDEX_NONE;
-	Database = nullptr;
-	AssetTime = 0.0f;
-}
-
 const FSearchIndexAsset* FSearchResult::GetSearchIndexAsset(bool bMandatory) const
 {
 	if (bMandatory)
diff --git a/Engine/Plugins/Animation/PoseSearch/Source/Runtime/Private/PoseSearchTrajectoryLibrary.cpp b/Engine/Plugins/Animation/PoseSearch/Source/Runtime/Private/PoseSearchTrajectoryLibrary.cpp
--- a/Engine/Plugins/Animation/PoseSearch/Source/Runtime/Private/PoseSearchTrajectoryLibrary.cpp
+++ b/Engine/Plugins/Animation/PoseSearch/Source/Runtime/Private/PoseSearchTrajectoryLibrary.cpp
@@ -14,25 +14,30 @@ void FPoseSearchTrajectoryData::UpdateData(
 	FDerived& TrajectoryDataDerived,
 	FState& TrajectoryDataState) const
 {
-	UpdateData(DeltaTime, Cast<const UAnimInstance>(AnimInstanceProxy.GetAnimInstanceObject()), TrajectoryDataDerived, TrajectoryDataState);
+	UpdateData(DeltaTime, AnimInstanceProxy.GetAnimInstanceObject(), TrajectoryDataDerived, TrajectoryDataState);
 }
 
 void FPoseSearchTrajectoryData::UpdateData(
 	float DeltaTime,
-	const UAnimInstance* AnimInstance,
+	const UObject* Context,
 	FDerived& TrajectoryDataDerived,
 	FState& TrajectoryDataState) const
 {
 	// An AnimInstance might call this during an AnimBP recompile with 0 delta time.
-	if (DeltaTime <= 0.f || !AnimInstance)
+	if (DeltaTime <= 0.f)
 	{
 		return;
 	}
 
-	const ACharacter* Character = Cast<ACharacter>(AnimInstance->GetOwningActor());
+	const ACharacter* Character = Cast<ACharacter>(Context);
 	if (!Character)
 	{
-		return;
+		const UAnimInstance* AnimInstance = Cast<UAnimInstance>(Context);
+		Character = Cast<ACharacter>(AnimInstance->GetOwningActor());
+		if (!Character)
+		{
+			return;
+		}
 	}
 
 	if (const UCharacterMovementComponent* MoveComp = Character->GetCharacterMovement())
@@ -383,7 +388,7 @@ void UPoseSearchTrajectoryLibrary::UpdatePrediction_SimulateCharacterMovement(
 }
 
 void UPoseSearchTrajectoryLibrary::PoseSearchGenerateTrajectory(
-	const UAnimInstance* InAnimInstance, 
+	const UObject* Context, 
 	UPARAM(ref)	const FPoseSearchTrajectoryData& InTrajectoryData,
 	float InDeltaTime,
 	UPARAM(ref) FPoseSearchQueryTrajectory& InOutTrajectory,
@@ -404,7 +409,7 @@ void UPoseSearchTrajectoryLibrary::PoseSearchGenerateTrajectory(
 	TrajectoryDataState.DesiredControllerYawLastUpdate = InOutDesiredControllerYawLastUpdate;
 
 	FPoseSearchTrajectoryData::FDerived TrajectoryDataDerived;
-	InTrajectoryData.UpdateData(InDeltaTime, InAnimInstance, TrajectoryDataDerived, TrajectoryDataState);
+	InTrajectoryData.UpdateData(InDeltaTime, Context, TrajectoryDataDerived, TrajectoryDataState);
 	InitTrajectorySamples(InOutTrajectory, InTrajectoryData, TrajectoryDataDerived.Position, TrajectoryDataDerived.Facing, TrajectoryDataSampling, InDeltaTime);
 	UpdateHistory_TransformHistory(InOutTrajectory, InTrajectoryData, TrajectoryDataDerived.Position, TrajectoryDataDerived.Velocity, TrajectoryDataSampling, InDeltaTime);
 	UpdatePrediction_SimulateCharacterMovement(InOutTrajectory, InTrajectoryData, TrajectoryDataDerived, TrajectoryDataSampling, InDeltaTime);
diff --git a/Engine/Plugins/Animation/PoseSearch/Source/Runtime/Public/PoseSearch/PoseSearchHistory.h b/Engine/Plugins/Animation/PoseSearch/Source/Runtime/Public/PoseSearch/PoseSearchHistory.h
--- a/Engine/Plugins/Animation/PoseSearch/Source/Runtime/Public/PoseSearch/PoseSearchHistory.h
+++ b/Engine/Plugins/Animation/PoseSearch/Source/Runtime/Public/PoseSearch/PoseSearchHistory.h
@@ -98,7 +98,7 @@ struct POSESEARCH_API FArchivedPoseHistory : public IPoseHistory
 };
 FArchive& operator<<(FArchive& Ar, FArchivedPoseHistory& Entry);
 
-struct FPoseHistory : public IPoseHistory
+struct POSESEARCH_API FPoseHistory : public IPoseHistory
 {
 	FPoseHistory() = default;
 	FPoseHistory(const FPoseHistory& Other);
@@ -131,6 +131,9 @@ struct FPoseHistory : public IPoseHistory
 #endif
 	// End of IPoseHistory interface
 	
+	int32 GetMaxNumPoses() const { return MaxNumPoses; }
+	float GetSamplingInterval() const { return SamplingInterval; }
+
 private:
 
 	// caching MaxNumPoses, since FData::Entries.Max() is a padded number
@@ -176,7 +179,7 @@ private:
 #endif // ENABLE_ANIM_DEBUG
 };
 
-struct FMemStackPoseHistory : public IPoseHistory
+struct POSESEARCH_API FMemStackPoseHistory : public IPoseHistory
 {
 	void Init(const IPoseHistory* InPoseHistory);
 
diff --git a/Engine/Plugins/Animation/PoseSearch/Source/Runtime/Public/PoseSearch/PoseSearchResult.h b/Engine/Plugins/Animation/PoseSearch/Source/Runtime/Public/PoseSearch/PoseSearchResult.h
--- a/Engine/Plugins/Animation/PoseSearch/Source/Runtime/Public/PoseSearch/PoseSearchResult.h
+++ b/Engine/Plugins/Animation/PoseSearch/Source/Runtime/Public/PoseSearch/PoseSearchResult.h
@@ -34,7 +34,7 @@ struct FSearchResult
 
 	bool IsValid() const { return PoseIdx != INDEX_NONE && Database != nullptr; }
 
-	void Reset();
+	void Reset() { PoseIdx = INDEX_NONE; Database = nullptr; AssetTime = 0.0f; bIsContinuingPoseSearch = false; }
 
 	POSESEARCH_API const FSearchIndexAsset* GetSearchIndexAsset(bool bMandatory = false) const;
 	
diff --git a/Engine/Plugins/Animation/PoseSearch/Source/Runtime/Public/PoseSearch/PoseSearchTrajectoryLibrary.h b/Engine/Plugins/Animation/PoseSearch/Source/Runtime/Public/PoseSearch/PoseSearchTrajectoryLibrary.h
--- a/Engine/Plugins/Animation/PoseSearch/Source/Runtime/Public/PoseSearch/PoseSearchTrajectoryLibrary.h
+++ b/Engine/Plugins/Animation/PoseSearch/Source/Runtime/Public/PoseSearch/PoseSearchTrajectoryLibrary.h
@@ -51,7 +51,7 @@ public:
 	};
 
 	void UpdateData(float DeltaTime, const FAnimInstanceProxy& AnimInstanceProxy, FDerived& TrajectoryDataDerived, FState& TrajectoryDataState) const;
-	void UpdateData(float DeltaTime, const UAnimInstance* AnimInstance, FDerived& TrajectoryDataDerived, FState& TrajectoryDataState) const;
+	void UpdateData(float DeltaTime, const UObject* Context, FDerived& TrajectoryDataDerived, FState& TrajectoryDataState) const;
 	FVector StepCharacterMovementGroundPrediction(float DeltaTime, const FVector& InVelocity, const FVector& InAcceleration, const FDerived& TrajectoryDataDerived) const;
 	
 	// If the character is forward facing (i.e. bOrientRotationToMovement is true), this controls how quickly the trajectory will rotate
@@ -122,9 +122,10 @@ public:
 	// Update prediction by simulating the movement math for ground locomotion from UCharacterMovementComponent.
 	static void UpdatePrediction_SimulateCharacterMovement(FPoseSearchQueryTrajectory& Trajectory, const FPoseSearchTrajectoryData& TrajectoryData, const FPoseSearchTrajectoryData::FDerived& TrajectoryDataDerived, const FPoseSearchTrajectoryData::FSampling& TrajectoryDataSampling, float DeltaTime);
 
+	// @todo: rename InAnimInstance into Context, since it can now be either an AnimInstance or a Character
 	// Generates a prediction trajectory based of the current character intent. For use with Character actors.
 	UFUNCTION(BlueprintCallable, Category = "Animation|PoseSearch", meta = (BlueprintThreadSafe, DisplayName = "Pose Search Generate Trajectory (for Character)"))
-	static void PoseSearchGenerateTrajectory(const UAnimInstance* InAnimInstance, 
+	static void PoseSearchGenerateTrajectory(const UObject* InAnimInstance, 
 		UPARAM(ref) const FPoseSearchTrajectoryData& InTrajectoryData, float InDeltaTime,
 		UPARAM(ref) FPoseSearchQueryTrajectory& InOutTrajectory, UPARAM(ref) float& InOutDesiredControllerYawLastUpdate, FPoseSearchQueryTrajectory& OutTrajectory,
 		float InHistorySamplingInterval = 0.04f, int32 InTrajectoryHistoryCount = 10, float InPredictionSamplingInterval = 0.2f, int32 InTrajectoryPredictionCount = 8);
