posesearch
- fix for missing FAsyncPoseSearchDatabasesManagement::RequestAsyncBuildIndex protection in MotionMatch function
- UPoseSearchTrajectoryLibrary::DrawTrajectory
- more AnimNext compatibility

diff --git a/Engine/Plugins/Animation/PoseSearch/Source/Runtime/Private/PoseSearchHistory.cpp b/Engine/Plugins/Animation/PoseSearch/Source/Runtime/Private/PoseSearchHistory.cpp
--- a/Engine/Plugins/Animation/PoseSearch/Source/Runtime/Private/PoseSearchHistory.cpp
+++ b/Engine/Plugins/Animation/PoseSearch/Source/Runtime/Private/PoseSearchHistory.cpp
@@ -672,8 +672,6 @@ void FPoseHistory::PreUpdate()
 	CheckThreadSafetyWrite(ReadPoseDataThreadSafeCounter);
 	CheckThreadSafetyWrite(WritePoseDataThreadSafeCounter);
 
-	check(IsInGameThread());
-
 	ReadPoseDataIndex = GetWritePoseDataIndex();
 }
 
diff --git a/Engine/Plugins/Animation/PoseSearch/Source/Runtime/Private/PoseSearchLibrary.cpp b/Engine/Plugins/Animation/PoseSearch/Source/Runtime/Private/PoseSearchLibrary.cpp
--- a/Engine/Plugins/Animation/PoseSearch/Source/Runtime/Private/PoseSearchLibrary.cpp
+++ b/Engine/Plugins/Animation/PoseSearch/Source/Runtime/Private/PoseSearchLibrary.cpp
@@ -45,9 +45,34 @@ namespace UE::PoseSearch
 	typedef TMap<const UPoseSearchDatabase*, TAssetsToSearch, TInlineSetAllocator<MAX_STACK_ALLOCATED_SETS, TMemStackSetAllocator<>>> TAssetsToSearchPerDatabaseMap;
 	typedef TPair<const UPoseSearchDatabase*, TAssetsToSearch> TAssetsToSearchPerDatabasePair;
 
-	static void AddToSearchForDatabase(TAssetsToSearchPerDatabaseMap& AssetsToSearchPerDatabaseMap, const UObject* AssetToSearch, const UPoseSearchDatabase* Database)
+	// this function adds AssetToSearch to the search of Database
+	// returns bAsyncBuildIndexInProgress
+	static bool AddToSearchForDatabase(TAssetsToSearchPerDatabaseMap& AssetsToSearchPerDatabaseMap, const UObject* AssetToSearch, const UPoseSearchDatabase* Database, bool bContainsIsMandatory)
 	{
-		if (TAssetsToSearch* AssetsToSearch = AssetsToSearchPerDatabaseMap.Find(Database))
+		TAssetsToSearch* AssetsToSearch = AssetsToSearchPerDatabaseMap.Find(Database);
+
+#if WITH_EDITOR
+		// no need to check if Database is indexing if found into AssetsToSearchPerDatabaseMap, since it already passed RequestAsyncBuildIndex successfully in a previous AddToSearchForDatabase call
+		if (!AssetsToSearch && (EAsyncBuildIndexResult::Success != FAsyncPoseSearchDatabasesManagement::RequestAsyncBuildIndex(Database, ERequestAsyncBuildFlag::ContinueRequest)))
+		{
+			// database is still indexing.. moving on
+			return true;
+		}
+
+		if (!Database->Contains(AssetToSearch))
+		{
+			if (bContainsIsMandatory)
+			{
+				UE_LOG(LogPoseSearch, Error, TEXT("improperly setup UAnimSequenceBase. Database %s doesn't contain UAnimSequenceBase %s"), *Database->GetName(), *AssetToSearch->GetName());
+			}
+			return false;
+		}
+#endif // WITH_EDITOR
+
+		// making sure AssetToSearch is not a databases! later on we could add support for nested databases, but currently we don't support that
+		check(Cast<const UPoseSearchDatabase>(AssetToSearch) == nullptr);
+
+		if (AssetsToSearch)
 		{
 			// an empty TAssetsToSearch associated to Database means we need to search ALL the assets, so we don't need to add this AssetToSearch
 			if (!AssetsToSearch->IsEmpty())
@@ -57,12 +82,20 @@ namespace UE::PoseSearch
 		}
 		else
 		{
-			AssetsToSearchPerDatabaseMap.Add(Database).AddUnique(AssetToSearch);
+			// mp meed to AddUnique since it's the first one
+			AssetsToSearchPerDatabaseMap.Add(Database).Add(AssetToSearch);
 		}
+
+		return false;
 	}
 
-	static void AddToSearch(TAssetsToSearchPerDatabaseMap& AssetsToSearchPerDatabaseMap, const UObject* AssetToSearch)
+	// this function is looking for UPoseSearchDatabase(s) to search for the input AssetToSearch:
+	// if AssetToSearch is a database search it ALL,
+	// if it's a sequence containing UAnimNotifyState_PoseSearchBranchIn, we add to the search of the dabase UAnimNotifyState_PoseSearchBranchIn::Database the asset AssetToSearch
+	// returns bAsyncBuildIndexInProgress
+	static bool AddToSearch(TAssetsToSearchPerDatabaseMap& AssetsToSearchPerDatabaseMap, const UObject* AssetToSearch)
 	{
+		bool bAsyncBuildIndexInProgress = false;
 		if (const UAnimSequenceBase* SequenceBase = Cast<const UAnimSequenceBase>(AssetToSearch))
 		{
 			for (const FAnimNotifyEvent& NotifyEvent : SequenceBase->Notifies)
@@ -74,29 +107,38 @@ namespace UE::PoseSearch
 						UE_LOG(LogPoseSearch, Error, TEXT("improperly setup UAnimNotifyState_PoseSearchBranchIn with null Database in %s"), *SequenceBase->GetName());
 						continue;
 					}
-#if WITH_EDITOR
-					if (EAsyncBuildIndexResult::Success != FAsyncPoseSearchDatabasesManagement::RequestAsyncBuildIndex(PoseSearchBranchIn->Database, ERequestAsyncBuildFlag::ContinueRequest))
-					{
-						// database is still indexing.. moving on
-						continue;
-					}
-
-					if (!PoseSearchBranchIn->Database->Contains(SequenceBase))
+					
+					// we just skip indexing databases to keep the experience as smooth as possible
+					if (AddToSearchForDatabase(AssetsToSearchPerDatabaseMap, SequenceBase, PoseSearchBranchIn->Database, true))
 					{
-						UE_LOG(LogPoseSearch, Error, TEXT("improperly setup UAnimSequenceBase. Database %s doesn't contain UAnimSequenceBase %s"), *PoseSearchBranchIn->Database->GetName(), *SequenceBase->GetName());
-						continue;
+						bAsyncBuildIndexInProgress = true;
 					}
-#endif // WITH_EDITOR
-					
-					AddToSearchForDatabase(AssetsToSearchPerDatabaseMap, SequenceBase, PoseSearchBranchIn->Database);
 				}
 			}
 		}
 		else if (const UPoseSearchDatabase* Database = Cast<UPoseSearchDatabase>(AssetToSearch))
 		{
-			// an empty TAssetsToSearch associated to Database means we need to search ALL the assets
-			AssetsToSearchPerDatabaseMap.FindOrAdd(Database).Reset();
+			// we already added Database to AssetsToSearchPerDatabaseMap, so it already successfully passed RequestAsyncBuildIndex
+			if (TAssetsToSearch* AssetsToSearch = AssetsToSearchPerDatabaseMap.Find(Database))
+			{
+				// an empty TAssetsToSearch associated to Database means we need to search ALL the assets
+				AssetsToSearch->Reset();
+			}
+			else
+#if WITH_EDITOR
+			if (EAsyncBuildIndexResult::Success != FAsyncPoseSearchDatabasesManagement::RequestAsyncBuildIndex(Database, ERequestAsyncBuildFlag::ContinueRequest))
+			{
+				bAsyncBuildIndexInProgress = true;
+			}
+			else
+#endif // WITH_EDITOR
+			{
+				// an empty TAssetsToSearch associated to Database means we need to search ALL the assets
+				AssetsToSearchPerDatabaseMap.Add(Database);
+			}
 		}
+
+		return bAsyncBuildIndexInProgress;
 	}
 
 	static bool IsForceInterrupt(EPoseSearchInterruptMode InterruptMode, const UPoseSearchDatabase* CurrentResultDatabase, const TArray<TObjectPtr<const UPoseSearchDatabase>>& Databases)
@@ -864,23 +906,27 @@ UE::PoseSearch::FSearchResult UPoseSearchLibrary::MotionMatch(
 
 	TAssetsToSearchPerDatabaseMap AssetsToSearchPerDatabaseMap;
 	
+	bool bAsyncBuildIndexInProgress = false;
+
 	// collecting all the possible continuing pose search (it could be multiple searches, but most likely only one)
 	const float DeltaSeconds = AnimInstances[0] ? AnimInstances[0]->GetDeltaSeconds() : FiniteDelta;
 	if (const UObject* PlayingAnimationAsset = ContinuingProperties.PlayingAsset.Get())
 	{
-		AddToSearch(AssetsToSearchPerDatabaseMap, PlayingAnimationAsset);
+		// checking if PlayingAnimationAsset has an associated database
+		if (AddToSearch(AssetsToSearchPerDatabaseMap, PlayingAnimationAsset))
+		{
+			bAsyncBuildIndexInProgress = true;
+		}
 		
+		// checking if any of the AssetsToSearch (databse) contains PlayingAnimationAsset
 		for (const UObject* AssetToSearch : AssetsToSearch)
 		{
 			if (const UPoseSearchDatabase* Database = Cast<UPoseSearchDatabase>(AssetToSearch))
 			{
-				if (Database->Contains(PlayingAnimationAsset))
+				// since it cannot be a database we can directly add it to AssetsToSearchPerDatabaseMap
+				if (AddToSearchForDatabase(AssetsToSearchPerDatabaseMap, PlayingAnimationAsset, Database, false))
 				{
-					// checking just in case later on we add support for databases containing other databases
-					check(Cast<const UPoseSearchDatabase>(PlayingAnimationAsset) == nullptr);
-
-					// since it cannot be a database we can directly add it to AssetsToSearchPerDatabaseMap
-					AddToSearchForDatabase(AssetsToSearchPerDatabaseMap, AssetToSearch, Database);
+					bAsyncBuildIndexInProgress = true;
 				}
 			}
 		}
@@ -888,75 +934,65 @@ UE::PoseSearch::FSearchResult UPoseSearchLibrary::MotionMatch(
 		for (const TAssetsToSearchPerDatabasePair& AssetsToSearchPerDatabasePair : AssetsToSearchPerDatabaseMap)
 		{
 			const UPoseSearchDatabase* Database = AssetsToSearchPerDatabasePair.Key;
+			check(Database);
 
-#if WITH_EDITOR
-			if (EAsyncBuildIndexResult::Success != FAsyncPoseSearchDatabasesManagement::RequestAsyncBuildIndex(Database, ERequestAsyncBuildFlag::ContinueRequest))
-			{
-				SearchContext.SetAsyncBuildIndexInProgress();
-			}
-			else
-#endif // WITH_EDITOR
+			const FSearchIndex& SearchIndex = Database->GetSearchIndex();
+			for (int32 AssetIndex : Database->GetAssetIndexesForSourceAsset(PlayingAnimationAsset))
 			{
-				check(Database);
-
-				const FSearchIndex& SearchIndex = Database->GetSearchIndex();
-				for (int32 AssetIndex : Database->GetAssetIndexesForSourceAsset(PlayingAnimationAsset))
-				{
-					const FSearchIndexAsset& SearchIndexAsset = SearchIndex.Assets[AssetIndex];
+				const FSearchIndexAsset& SearchIndexAsset = SearchIndex.Assets[AssetIndex];
 
-					const float FirstSampleTime = SearchIndexAsset.GetFirstSampleTime(Database->Schema->SampleRate);
-					const float LastSampleTime = SearchIndexAsset.GetLastSampleTime(Database->Schema->SampleRate);
+				const float FirstSampleTime = SearchIndexAsset.GetFirstSampleTime(Database->Schema->SampleRate);
+				const float LastSampleTime = SearchIndexAsset.GetLastSampleTime(Database->Schema->SampleRate);
 
-					bool bCanAdvance = true;
-					if (SearchIndexAsset.IsLooping())
+				bool bCanAdvance = true;
+				if (SearchIndexAsset.IsLooping())
+				{
+					const float DeltaSampleTime = LastSampleTime - FirstSampleTime;
+					if (DeltaSampleTime < UE_SMALL_NUMBER)
+					{
+						ReconstructedPreviousSearchResult.AssetTime = FirstSampleTime;
+					}
+					else if (ContinuingProperties.PlayingAssetAccumulatedTime < FirstSampleTime)
 					{
-						const float DeltaSampleTime = LastSampleTime - FirstSampleTime;
-						if (DeltaSampleTime < UE_SMALL_NUMBER)
-						{
-							ReconstructedPreviousSearchResult.AssetTime = FirstSampleTime;
-						}
-						else if (ContinuingProperties.PlayingAssetAccumulatedTime < FirstSampleTime)
-						{
-							ReconstructedPreviousSearchResult.AssetTime = FMath::Fmod(ContinuingProperties.PlayingAssetAccumulatedTime - FirstSampleTime, DeltaSampleTime) + DeltaSampleTime + FirstSampleTime;
-						}
-						else if (ContinuingProperties.PlayingAssetAccumulatedTime > LastSampleTime)
-						{
-							ReconstructedPreviousSearchResult.AssetTime = FMath::Fmod(ContinuingProperties.PlayingAssetAccumulatedTime - FirstSampleTime, DeltaSampleTime) + FirstSampleTime;
-						}
-						else
-						{
-							ReconstructedPreviousSearchResult.AssetTime = ContinuingProperties.PlayingAssetAccumulatedTime;
-						}
+						ReconstructedPreviousSearchResult.AssetTime = FMath::Fmod(ContinuingProperties.PlayingAssetAccumulatedTime - FirstSampleTime, DeltaSampleTime) + DeltaSampleTime + FirstSampleTime;
+					}
+					else if (ContinuingProperties.PlayingAssetAccumulatedTime > LastSampleTime)
+					{
+						ReconstructedPreviousSearchResult.AssetTime = FMath::Fmod(ContinuingProperties.PlayingAssetAccumulatedTime - FirstSampleTime, DeltaSampleTime) + FirstSampleTime;
 					}
 					else
 					{
-						const float MaxTimeToBeAbleToContinuingPlayingAnimation = LastSampleTime - DeltaSeconds;
-						bCanAdvance = ContinuingProperties.PlayingAssetAccumulatedTime >= FirstSampleTime && ContinuingProperties.PlayingAssetAccumulatedTime < MaxTimeToBeAbleToContinuingPlayingAnimation;
 						ReconstructedPreviousSearchResult.AssetTime = ContinuingProperties.PlayingAssetAccumulatedTime;
 					}
+				}
+				else
+				{
+					const float MaxTimeToBeAbleToContinuingPlayingAnimation = LastSampleTime - DeltaSeconds;
+					bCanAdvance = ContinuingProperties.PlayingAssetAccumulatedTime >= FirstSampleTime && ContinuingProperties.PlayingAssetAccumulatedTime < MaxTimeToBeAbleToContinuingPlayingAnimation;
+					ReconstructedPreviousSearchResult.AssetTime = ContinuingProperties.PlayingAssetAccumulatedTime;
+				}
 
-					if (bCanAdvance)
-					{
-						ReconstructedPreviousSearchResult.Database = Database;
-						ReconstructedPreviousSearchResult.PoseIdx = Database->GetPoseIndexFromTime(ContinuingProperties.PlayingAssetAccumulatedTime, SearchIndexAsset);
-						SearchContext.UpdateCurrentResultPoseVector();
+				if (bCanAdvance)
+				{
+					ReconstructedPreviousSearchResult.Database = Database;
+					ReconstructedPreviousSearchResult.PoseIdx = Database->GetPoseIndexFromTime(ContinuingProperties.PlayingAssetAccumulatedTime, SearchIndexAsset);
+					SearchContext.UpdateCurrentResultPoseVector();
 
-						const FSearchResult NewSearchResult = Database->SearchContinuingPose(SearchContext);
+					const FSearchResult NewSearchResult = Database->SearchContinuingPose(SearchContext);
 
 #if WITH_EDITOR && ENABLE_ANIM_DEBUG && UE_POSE_SEARCH_TRACE_ENABLED
-						const FPoseSearchCost BestBruteForcePoseCost = NewSearchResult.BruteForcePoseCost < SearchResult.BruteForcePoseCost ? NewSearchResult.BruteForcePoseCost : SearchResult.BruteForcePoseCost;
+					const FPoseSearchCost BestBruteForcePoseCost = NewSearchResult.BruteForcePoseCost < SearchResult.BruteForcePoseCost ? NewSearchResult.BruteForcePoseCost : SearchResult.BruteForcePoseCost;
 #endif // WITH_EDITOR && ENABLE_ANIM_DEBUG && UE_POSE_SEARCH_TRACE_ENABLED
 
-						if (NewSearchResult.PoseCost < SearchResult.PoseCost)
-						{
-							SearchResult = NewSearchResult;
-							SearchContext.UpdateCurrentBestCost(SearchResult.PoseCost);
-						}
+					if (NewSearchResult.PoseCost < SearchResult.PoseCost)
+					{
+						SearchResult = NewSearchResult;
+						SearchContext.UpdateCurrentBestCost(SearchResult.PoseCost);
+					}
 								
 #if WITH_EDITOR && ENABLE_ANIM_DEBUG && UE_POSE_SEARCH_TRACE_ENABLED
-						SearchResult.BruteForcePoseCost = BestBruteForcePoseCost;
+					SearchResult.BruteForcePoseCost = BestBruteForcePoseCost;
 #endif // WITH_EDITOR && ENABLE_ANIM_DEBUG && UE_POSE_SEARCH_TRACE_ENABLED
-					}
 				}
 			}
 		}
@@ -969,7 +1005,10 @@ UE::PoseSearch::FSearchResult UPoseSearchLibrary::MotionMatch(
 	{
 		for (const UObject* AssetToSearch : AssetsToSearch)
 		{
-			AddToSearch(AssetsToSearchPerDatabaseMap, AssetToSearch);
+			if (AddToSearch(AssetsToSearchPerDatabaseMap, AssetToSearch))
+			{
+				bAsyncBuildIndexInProgress = true;
+			}
 		}
 
 		for (const TAssetsToSearchPerDatabasePair& AssetsToSearchPerDatabasePair : AssetsToSearchPerDatabaseMap)
@@ -997,6 +1036,13 @@ UE::PoseSearch::FSearchResult UPoseSearchLibrary::MotionMatch(
 		}
 	}
 
+#if WITH_EDITOR
+	if (bAsyncBuildIndexInProgress)
+	{
+		SearchContext.SetAsyncBuildIndexInProgress();
+	}
+#endif // WITH_EDITOR
+
 #if ENABLE_DRAW_DEBUG && ENABLE_ANIM_DEBUG
 	if (SearchResult.IsValid())
 	{
diff --git a/Engine/Plugins/Animation/PoseSearch/Source/Runtime/Private/PoseSearchTrajectoryLibrary.cpp b/Engine/Plugins/Animation/PoseSearch/Source/Runtime/Private/PoseSearchTrajectoryLibrary.cpp
--- a/Engine/Plugins/Animation/PoseSearch/Source/Runtime/Private/PoseSearchTrajectoryLibrary.cpp
+++ b/Engine/Plugins/Animation/PoseSearch/Source/Runtime/Private/PoseSearchTrajectoryLibrary.cpp
@@ -627,4 +627,14 @@ void UPoseSearchTrajectoryLibrary::GetTrajectoryAngularVelocity(const FPoseSearc
 		FMath::RadiansToDegrees(AngularVelocityInRadians.X),
 		FMath::RadiansToDegrees(AngularVelocityInRadians.Y),
 		FMath::RadiansToDegrees(AngularVelocityInRadians.Z));
+}
+
+void UPoseSearchTrajectoryLibrary::DrawTrajectory(const UObject* WorldContextObject, const FPoseSearchQueryTrajectory& InTrajectory, const float DebugThickness, float HeightOffset)
+{
+#if ENABLE_ANIM_DEBUG
+	if (UWorld* World = GEngine->GetWorldFromContextObject(WorldContextObject, EGetWorldErrorMode::LogAndReturnNull))
+	{
+		InTrajectory.DebugDrawTrajectory(World, DebugThickness, HeightOffset);
+	}
+#endif // ENABLE_ANIM_DEBUG
 }
\ No newline at end of file
diff --git a/Engine/Plugins/Animation/PoseSearch/Source/Runtime/Public/PoseSearch/PoseSearchTrajectoryLibrary.h b/Engine/Plugins/Animation/PoseSearch/Source/Runtime/Public/PoseSearch/PoseSearchTrajectoryLibrary.h
--- a/Engine/Plugins/Animation/PoseSearch/Source/Runtime/Public/PoseSearch/PoseSearchTrajectoryLibrary.h
+++ b/Engine/Plugins/Animation/PoseSearch/Source/Runtime/Public/PoseSearch/PoseSearchTrajectoryLibrary.h
@@ -167,6 +167,9 @@ public:
 		return InTrajectorySample.GetTransform();
 	};
 	
+	UFUNCTION(BlueprintCallable, Category="Animation|PoseSearch|Experimental", meta=(WorldContext="WorldContextObject"))
+	static void DrawTrajectory(const UObject* WorldContextObject, UPARAM(ref) const FPoseSearchQueryTrajectory& InTrajectory, const float DebugThickness, float HeightOffset);
+
 private:
 	static FVector RemapVectorMagnitudeWithCurve(const FVector& Vector, bool bUseCurve, const FRuntimeFloatCurve& Curve);
 };
\ No newline at end of file
diff --git a/Engine/Plugins/Animation/PoseSearch/Source/Runtime/Public/PoseSearch/PoseSearchTrajectoryTypes.h b/Engine/Plugins/Animation/PoseSearch/Source/Runtime/Public/PoseSearch/PoseSearchTrajectoryTypes.h
--- a/Engine/Plugins/Animation/PoseSearch/Source/Runtime/Public/PoseSearch/PoseSearchTrajectoryTypes.h
+++ b/Engine/Plugins/Animation/PoseSearch/Source/Runtime/Public/PoseSearch/PoseSearchTrajectoryTypes.h
@@ -29,7 +29,7 @@ struct POSESEARCH_API FPoseSearchQueryTrajectorySample
 };
 POSESEARCH_API FArchive& operator<<(FArchive& Ar, FPoseSearchQueryTrajectorySample& TrajectorySample);
 
-USTRUCT(BlueprintType, Category = "Motion Trajectory")
+USTRUCT(BlueprintType, Category = "Pose Search Trajectory")
 struct POSESEARCH_API FPoseSearchQueryTrajectory
 {
 	GENERATED_BODY()
