posesearch - fix for CBP compilation when FAnimNode_PoseSearchHistoryCollector_Base::Initialize_AnyThread has bInitializeWithRefPose

diff --git a/Engine/Plugins/Animation/PoseSearch/Source/Runtime/Private/AnimNode_PoseSearchHistoryCollector.cpp b/Engine/Plugins/Animation/PoseSearch/Source/Runtime/Private/AnimNode_PoseSearchHistoryCollector.cpp
--- a/Engine/Plugins/Animation/PoseSearch/Source/Runtime/Private/AnimNode_PoseSearchHistoryCollector.cpp
+++ b/Engine/Plugins/Animation/PoseSearch/Source/Runtime/Private/AnimNode_PoseSearchHistoryCollector.cpp
@@ -21,12 +21,16 @@ void FAnimNode_PoseSearchHistoryCollector_Base::Initialize_AnyThread(const FAnim
 
 	if (bInitializeWithRefPose)
 	{
-		// initializing PoseHistory with a ref pose at FAnimInstanceProxy location/facing
-		FMemMark Mark(FMemStack::Get());
-		FCSPose<FCompactPose> ComponentSpacePose;
-		ComponentSpacePose.InitPose(&Context.AnimInstanceProxy->GetRequiredBones());
-		PoseHistory.EvaluateComponentSpace_AnyThread(0.f, ComponentSpacePose, bStoreScales, 
-			RootBoneRecoveryTime, RootBoneTranslationRecoveryRatio, RootBoneRotationRecoveryRatio, true, true, GetRequiredBones(Context.AnimInstanceProxy));
+		const FBoneContainer& BoneContainer = Context.AnimInstanceProxy->GetRequiredBones();
+		if (BoneContainer.IsValid())
+		{
+			// initializing PoseHistory with a ref pose at FAnimInstanceProxy location/facing
+			FMemMark Mark(FMemStack::Get());
+			FCSPose<FCompactPose> ComponentSpacePose;
+			ComponentSpacePose.InitPose(&BoneContainer);
+			PoseHistory.EvaluateComponentSpace_AnyThread(0.f, ComponentSpacePose, bStoreScales,
+				RootBoneRecoveryTime, RootBoneTranslationRecoveryRatio, RootBoneRotationRecoveryRatio, true, true, GetRequiredBones(Context.AnimInstanceProxy));
+		}
 	}
 }
 
