pose search - fix for UE::PoseSearch::AddToSearchForDatabase WITH_EDITOR

diff --git a/Engine/Plugins/Animation/PoseSearch/Source/Runtime/Private/PoseSearchLibrary.cpp b/Engine/Plugins/Animation/PoseSearch/Source/Runtime/Private/PoseSearchLibrary.cpp
--- a/Engine/Plugins/Animation/PoseSearch/Source/Runtime/Private/PoseSearchLibrary.cpp
+++ b/Engine/Plugins/Animation/PoseSearch/Source/Runtime/Private/PoseSearchLibrary.cpp
@@ -58,6 +58,7 @@ namespace UE::PoseSearch
 			// database is still indexing.. moving on
 			return true;
 		}
+#endif // WITH_EDITOR
 
 		if (!Database->Contains(AssetToSearch))
 		{
@@ -67,7 +68,6 @@ namespace UE::PoseSearch
 			}
 			return false;
 		}
-#endif // WITH_EDITOR
 
 		// making sure AssetToSearch is not a databases! later on we could add support for nested databases, but currently we don't support that
 		check(Cast<const UPoseSearchDatabase>(AssetToSearch) == nullptr);
@@ -82,7 +82,7 @@ namespace UE::PoseSearch
 		}
 		else
 		{
-			// mp meed to AddUnique since it's the first one
+			// no need to AddUnique since it's the first one
 			AssetsToSearchPerDatabaseMap.Add(Database).Add(AssetToSearch);
 		}
 
