posesearch - fix for NoPCH Compile UnrealGame Mac

diff --git a/Engine/Plugins/Animation/PoseSearch/Source/Runtime/Private/PoseSearchTrajectoryLibrary.cpp b/Engine/Plugins/Animation/PoseSearch/Source/Runtime/Private/PoseSearchTrajectoryLibrary.cpp
--- a/Engine/Plugins/Animation/PoseSearch/Source/Runtime/Private/PoseSearchTrajectoryLibrary.cpp
+++ b/Engine/Plugins/Animation/PoseSearch/Source/Runtime/Private/PoseSearchTrajectoryLibrary.cpp
@@ -4,6 +4,7 @@
 #include "PoseSearch/PoseSearchTrajectoryPredictor.h"
 #include "Animation/AnimInstanceProxy.h"
 #include "Components/SkeletalMeshComponent.h"
+#include "Engine/Engine.h"
 #include "GameFramework/Character.h"
 #include "GameFramework/CharacterMovementComponent.h"
 #include "PoseSearch/PoseSearchDefines.h"
