Pose Search: Add support for curve values in Pose History

diff --git a/Engine/Plugins/Animation/PoseSearch/Source/Runtime/Private/AnimNode_PoseSearchHistoryCollector.cpp b/Engine/Plugins/Animation/PoseSearch/Source/Runtime/Private/AnimNode_PoseSearchHistoryCollector.cpp
--- a/Engine/Plugins/Animation/PoseSearch/Source/Runtime/Private/AnimNode_PoseSearchHistoryCollector.cpp
+++ b/Engine/Plugins/Animation/PoseSearch/Source/Runtime/Private/AnimNode_PoseSearchHistoryCollector.cpp
@@ -27,9 +27,11 @@ void FAnimNode_PoseSearchHistoryCollector_Base::Initialize_AnyThread(const FAnim
 			// initializing PoseHistory with a ref pose at FAnimInstanceProxy location/facing
 			FMemMark Mark(FMemStack::Get());
 			FCSPose<FCompactPose> ComponentSpacePose;
+			FBlendedCurve EmptyCurves;
 			ComponentSpacePose.InitPose(&BoneContainer);
 			PoseHistory.EvaluateComponentSpace_AnyThread(0.f, ComponentSpacePose, bStoreScales,
-				RootBoneRecoveryTime, RootBoneTranslationRecoveryRatio, RootBoneRotationRecoveryRatio, true, true, GetRequiredBones(Context.AnimInstanceProxy));
+				RootBoneRecoveryTime, RootBoneTranslationRecoveryRatio, RootBoneRotationRecoveryRatio, true, true, 
+				GetRequiredBones(Context.AnimInstanceProxy), EmptyCurves, MakeConstArrayView(CollectedCurves));
 		}
 	}
 }
@@ -151,7 +153,8 @@ void FAnimNode_PoseSearchHistoryCollector::Evaluate_AnyThread(FPoseContext& Outp
 	}
 
 	PoseHistory.EvaluateComponentSpace_AnyThread(Output.AnimInstanceProxy->GetDeltaSeconds(), ComponentSpacePose, bStoreScales,
-		RootBoneRecoveryTime, RootBoneTranslationRecoveryRatio, RootBoneRotationRecoveryRatio, bNeedsReset, bCacheBones, RequiredBones);
+		RootBoneRecoveryTime, RootBoneTranslationRecoveryRatio, RootBoneRotationRecoveryRatio, bNeedsReset, bCacheBones, 
+		RequiredBones, Output.Curve, MakeConstArrayView(CollectedCurves));
 
 	bCacheBones = false;
 
diff --git a/Engine/Plugins/Animation/PoseSearch/Source/Runtime/Private/PoseSearchContext.cpp b/Engine/Plugins/Animation/PoseSearch/Source/Runtime/Private/PoseSearchContext.cpp
--- a/Engine/Plugins/Animation/PoseSearch/Source/Runtime/Private/PoseSearchContext.cpp
+++ b/Engine/Plugins/Animation/PoseSearch/Source/Runtime/Private/PoseSearchContext.cpp
@@ -507,7 +507,7 @@ float FSearchContext::GetSampleCurveValueInternal(float SampleTime, const FName&
 	{
 		if (ensure(PoseHistory))
 		{
-			// @TODO: Add curve support to the pose history node.
+			PoseHistory->GetCurveValueAtTime(SampleTime, CurveName, OutCurveValue);
 			return OutCurveValue;
 		}
 	}
diff --git a/Engine/Plugins/Animation/PoseSearch/Source/Runtime/Private/PoseSearchHistory.cpp b/Engine/Plugins/Animation/PoseSearch/Source/Runtime/Private/PoseSearchHistory.cpp
--- a/Engine/Plugins/Animation/PoseSearch/Source/Runtime/Private/PoseSearchHistory.cpp
+++ b/Engine/Plugins/Animation/PoseSearch/Source/Runtime/Private/PoseSearchHistory.cpp
@@ -112,6 +112,43 @@ static FComponentSpaceTransformIndex GetRemappedComponentSpaceTransformIndex(con
 	return BoneTransformIndex;
 }
 
+static bool LerpEntries(float Time, bool bExtrapolate, const FPoseHistoryEntry& PrevEntry, const FPoseHistoryEntry& NextEntry, const FName& CurveName, const TConstArrayView<FName>& CollectedCurves, float& OutCurveValue)
+{
+	bool bSuccess = true;
+
+	const int32 CurveIndex = CollectedCurves.Find(CurveName);
+	if (CurveIndex == INDEX_NONE)
+	{
+		OutCurveValue = 0.0f;
+		bSuccess = false;
+	}
+	else
+	{
+		const float Denominator = NextEntry.AccumulatedSeconds - PrevEntry.AccumulatedSeconds;
+		float LerpValue = 0.f;
+		if (!FMath::IsNearlyZero(Denominator))
+		{
+			const float Numerator = Time - PrevEntry.AccumulatedSeconds;
+			LerpValue = bExtrapolate ? Numerator / Denominator : FMath::Clamp(Numerator / Denominator, 0.f, 1.f);
+		}
+
+		if (FMath::IsNearlyZero(LerpValue, ZERO_ANIMWEIGHT_THRESH))
+		{
+			OutCurveValue = PrevEntry.GetCurveValue(CurveIndex);
+		}
+		else if (FMath::IsNearlyZero(LerpValue - 1.f, ZERO_ANIMWEIGHT_THRESH))
+		{
+			OutCurveValue = NextEntry.GetCurveValue(CurveIndex);
+		}
+		else
+		{
+			OutCurveValue = FMath::Lerp(PrevEntry.GetCurveValue(CurveIndex), NextEntry.GetCurveValue(CurveIndex), LerpValue);
+		}
+	}
+
+	return bSuccess;
+}
+
 static bool LerpEntries(float Time, bool bExtrapolate, const FPoseHistoryEntry& PrevEntry, const FPoseHistoryEntry& NextEntry, const USkeleton* BoneIndexSkeleton, const USkeleton* LastUpdateSkeleton,
 	const FBoneToTransformMap& BoneToTransformMap, FBoneIndexType BoneIndexType, FBoneIndexType ReferenceBoneIndexType, FTransform& OutBoneTransform)
 {
@@ -240,7 +277,7 @@ static void CopyEntries(const TRingBuffer<FPoseHistoryEntry>& From, TRingBuffer<
 
 //////////////////////////////////////////////////////////////////////////
 // FPoseHistoryEntry
-void FPoseHistoryEntry::Update(float Time, FCSPose<FCompactPose>& ComponentSpacePose, const FBoneToTransformMap& BoneToTransformMap, bool bStoreScales)
+void FPoseHistoryEntry::Update(float Time, FCSPose<FCompactPose>& ComponentSpacePose, const FBoneToTransformMap& BoneToTransformMap, bool bStoreScales, const FBlendedCurve& Curves, const TConstArrayView<FName>& CollectedCurves)
 {
 	AccumulatedSeconds = Time;
 
@@ -271,6 +308,14 @@ void FPoseHistoryEntry::Update(float Time, FCSPose<FCompactPose>& ComponentSpace
 			SetComponentSpaceTransform(BoneToTransformPair.Value, (CompactBoneIdx.IsValid() ? ComponentSpacePose.GetComponentSpaceTransform(CompactBoneIdx) : RefBonePose[SkeletonBoneIdx.GetInt()]));
 		}
 	}
+
+	const int32 NumCurves = CollectedCurves.Num();
+	CurveValues.SetNum(NumCurves);
+	for (int i = 0; i < NumCurves; ++i)
+	{
+		const FName& CurveName = CollectedCurves[i];
+		CurveValues[i] = Curves.Get(CurveName);
+	}
 }
 
 void FPoseHistoryEntry::SetNum(int32 Num, bool bStoreScales)
@@ -314,11 +359,25 @@ FTransform FPoseHistoryEntry::GetComponentSpaceTransform(int32 Index) const
 	return FTransform(Quat, ComponentSpacePositions[Index], Scale);
 }
 
+float FPoseHistoryEntry::GetCurveValue(int32 Index) const
+{
+#if WITH_EDITOR
+	if (Index < 0 || Index >= CurveValues.Num())
+	{
+		UE_LOG(LogPoseSearch, Error, TEXT("FPoseHistoryEntry::GetCurveValue - Index %d out of bound [0, %d)"), Index, CurveValues.Num());
+		return 0.0f;
+	}
+#endif // WITH_EDITOR
+
+	return CurveValues[Index];
+}
+
 FArchive& operator<<(FArchive& Ar, FPoseHistoryEntry& Entry)
 {
 	Ar << Entry.ComponentSpaceRotations;
 	Ar << Entry.ComponentSpacePositions;
 	Ar << Entry.ComponentSpaceScales;
+	Ar << Entry.CurveValues;
 	Ar << Entry.AccumulatedSeconds;
 	return Ar;
 }
@@ -374,7 +433,7 @@ void FArchivedPoseHistory::InitFrom(const IPoseHistory* PoseHistory)
 	{
 		Trajectory = PoseHistory->GetTrajectory();
 		BoneToTransformMap = PoseHistory->GetBoneToTransformMap();
-			
+		CollectedCurves = PoseHistory->GetCollectedCurves();
 		const int32 NumEntries = PoseHistory->GetNumEntries();
 		Entries.SetNum(NumEntries);
 
@@ -432,6 +491,36 @@ bool FArchivedPoseHistory::GetTransformAtTime(float Time, FTransform& OutBoneTra
 	return bSuccess;
 }
 
+bool FArchivedPoseHistory::GetCurveValueAtTime(float Time, const FName& CurveName, float& OutCurveValue, bool bExtrapolate /*= false*/) const
+{
+	bool bSuccess = false;
+
+	const int32 NumEntries = Entries.Num();
+	if (NumEntries > 0)
+	{
+		int32 NextIdx = 0;
+		int32 PrevIdx = 0;
+
+		if (NumEntries > 1)
+		{
+			const int32 LowerBoundIdx = Algo::LowerBound(Entries, Time, [](const FPoseHistoryEntry& Entry, float Value) { return Value > Entry.AccumulatedSeconds; });
+			NextIdx = FMath::Clamp(LowerBoundIdx, 1, NumEntries - 1);
+			PrevIdx = NextIdx - 1;
+		}
+
+		const FPoseHistoryEntry& PrevEntry = Entries[PrevIdx];
+		const FPoseHistoryEntry& NextEntry = Entries[NextIdx];
+
+		bSuccess = LerpEntries(Time, bExtrapolate, PrevEntry, NextEntry, CurveName, GetCollectedCurves(), OutCurveValue);
+	}
+	else
+	{
+		OutCurveValue = 0.0f;
+	}
+
+	return bSuccess;
+}
+
 #if ENABLE_DRAW_DEBUG && ENABLE_ANIM_DEBUG
 void FArchivedPoseHistory::DebugDraw(const UWorld* World, FColor Color) const
 {
@@ -486,6 +575,7 @@ void FArchivedPoseHistory::DebugDraw(const UWorld* World, FColor Color) const
 FArchive& operator<<(FArchive& Ar, FArchivedPoseHistory& Entry)
 {
 	Ar << Entry.BoneToTransformMap;
+	Ar << Entry.CollectedCurves;
 	Ar << Entry.Entries;
 	Ar << Entry.Trajectory;
 	return Ar;
@@ -617,6 +707,38 @@ bool FPoseHistory::GetTransformAtTime(float Time, FTransform& OutBoneTransform,
 	return bSuccess;
 }
 
+bool FPoseHistory::GetCurveValueAtTime(float Time, const FName& CurveName, float& OutCurveValue, bool bExtrapolate /*= false*/) const
+{
+	CheckThreadSafetyRead(ReadPoseDataThreadSafeCounter);
+
+	bool bSuccess = false;
+	const FPoseData& ReadPoseData = GetReadPoseData();
+	const int32 NumEntries = ReadPoseData.Entries.Num();
+	if (NumEntries > 0)
+	{
+		int32 NextIdx = 0;
+		int32 PrevIdx = 0;
+
+		if (NumEntries > 1)
+		{
+			const int32 LowerBoundIdx = LowerBound(ReadPoseData.Entries.begin(), ReadPoseData.Entries.end(), Time, [](const FPoseHistoryEntry& Entry, float Value) { return Value > Entry.AccumulatedSeconds; });
+			NextIdx = FMath::Clamp(LowerBoundIdx, 1, NumEntries - 1);
+			PrevIdx = NextIdx - 1;
+		}
+
+		const FPoseHistoryEntry& PrevEntry = ReadPoseData.Entries[PrevIdx];
+		const FPoseHistoryEntry& NextEntry = ReadPoseData.Entries[NextIdx];
+
+		bSuccess = LerpEntries(Time, bExtrapolate, PrevEntry, NextEntry, CurveName, ReadPoseData.CollectedCurves, OutCurveValue);
+	}
+	else
+	{
+		OutCurveValue = 0.0f;
+	}
+
+	return bSuccess;
+}
+
 const FPoseSearchQueryTrajectory& FPoseHistory::GetTrajectory() const
 {
 	CheckThreadSafetyRead(ReadPoseDataThreadSafeCounter);
@@ -641,6 +763,12 @@ const FBoneToTransformMap& FPoseHistory::GetBoneToTransformMap() const
 	return GetReadPoseData().BoneToTransformMap;
 }
 
+const TConstArrayView<FName> FPoseHistory::GetCollectedCurves() const
+{
+	CheckThreadSafetyRead(ReadPoseDataThreadSafeCounter);
+	return MakeConstArrayView(GetReadPoseData().CollectedCurves);
+}
+
 int32 FPoseHistory::GetNumEntries() const
 {
 	CheckThreadSafetyRead(ReadPoseDataThreadSafeCounter);
@@ -697,9 +825,9 @@ void FPoseHistory::SetTrajectory(const FPoseSearchQueryTrajectory& InTrajectory,
 	}
 }
 
-void FPoseHistory::EvaluateComponentSpace_AnyThread(float DeltaTime, FCSPose<FCompactPose>& ComponentSpacePose, bool bStoreScales,
-	float RootBoneRecoveryTime, float RootBoneTranslationRecoveryRatio, float RootBoneRotationRecoveryRatio,
-	bool bNeedsReset, bool bCacheBones, const TArray<FBoneIndexType>& RequiredBones)
+void FPoseHistory::EvaluateComponentSpace_AnyThread(float DeltaTime, FCSPose<FCompactPose>& ComponentSpacePose, bool bStoreScales, float RootBoneRecoveryTime,
+	float RootBoneTranslationRecoveryRatio, float RootBoneRotationRecoveryRatio, bool bNeedsReset, bool bCacheBones,
+	const TArray<FBoneIndexType>& RequiredBones, const FBlendedCurve& Curves, const TConstArrayView<FName>& CollectedCurves)
 {
 	CheckThreadSafetyWrite(WritePoseDataThreadSafeCounter);
 	CheckThreadSafetyRead(ReadPoseDataThreadSafeCounter);
@@ -717,6 +845,7 @@ void FPoseHistory::EvaluateComponentSpace_AnyThread(float DeltaTime, FCSPose<FCo
 	if (bCacheBones)
 	{
 		WritePoseData.BoneToTransformMap.Reset();
+		WritePoseData.CollectedCurves = CollectedCurves;
 		if (!RequiredBones.IsEmpty())
 		{
 			// making sure we always collect the root bone transform (by construction BoneToTransformMap[0] = 0)
@@ -788,7 +917,7 @@ void FPoseHistory::EvaluateComponentSpace_AnyThread(float DeltaTime, FCSPose<FCo
 
 	// Regardless of the retention policy, we always update the most recent Entry
 	FPoseHistoryEntry& MostRecentEntry = WritePoseData.Entries.Last();
-	MostRecentEntry.Update(0.f, ComponentSpacePose, WritePoseData.BoneToTransformMap, bStoreScales);
+	MostRecentEntry.Update(0.f, ComponentSpacePose, WritePoseData.BoneToTransformMap, bStoreScales, Curves, WritePoseData.CollectedCurves);
 
 	if (RootBoneRecoveryTime > 0.f && !Trajectory.Samples.IsEmpty())
 	{
@@ -814,6 +943,16 @@ void FPoseHistory::EvaluateComponentSpace_AnyThread(float DeltaTime, FCSPose<FCo
 	}
 }
 
+void FPoseHistory::EvaluateComponentSpace_AnyThread(float DeltaTime, FCSPose<FCompactPose>& ComponentSpacePose, bool bStoreScales,
+	float RootBoneRecoveryTime, float RootBoneTranslationRecoveryRatio, float RootBoneRotationRecoveryRatio,
+	bool bNeedsReset, bool bCacheBones, const TArray<FBoneIndexType>& RequiredBones)
+{
+	FBlendedCurve Curves;
+	EvaluateComponentSpace_AnyThread(DeltaTime, ComponentSpacePose, bStoreScales,
+		RootBoneRecoveryTime, RootBoneTranslationRecoveryRatio, RootBoneRotationRecoveryRatio, bNeedsReset, bCacheBones, 
+		RequiredBones, Curves, TConstArrayView<FName>());
+}
+
 #if ENABLE_DRAW_DEBUG && ENABLE_ANIM_DEBUG
 void FPoseHistory::DebugDraw(FAnimInstanceProxy& AnimInstanceProxy, FColor Color) const
 {
@@ -900,6 +1039,27 @@ bool FMemStackPoseHistory::GetTransformAtTime(float Time, FTransform& OutBoneTra
 	return PoseHistory->GetTransformAtTime(Time, OutBoneTransform, BoneIndexSkeleton, BoneIndexType, ReferenceBoneIndexType, bExtrapolate);
 }
 
+bool FMemStackPoseHistory::GetCurveValueAtTime(float Time, const FName& CurveName, float& OutCurveValue, bool bExtrapolate /*= false*/) const
+{
+	check(PoseHistory);
+	if (Time > 0.f)
+	{
+		const int32 Num = FutureEntries.Num();
+		if (Num > 0)
+		{
+			const int32 LowerBoundIdx = Algo::LowerBound(FutureEntries, Time, [](const FPoseHistoryEntry& Entry, float Value) { return Value > Entry.AccumulatedSeconds; });
+			const int32 NextIdx = FMath::Min(LowerBoundIdx, Num - 1);
+			const FPoseHistoryEntry& NextEntry = FutureEntries[NextIdx];
+			const FPoseHistoryEntry& PrevEntry = NextIdx > 0 ? FutureEntries[NextIdx - 1] : PoseHistory->GetNumEntries() > 0 ? PoseHistory->GetEntry(PoseHistory->GetNumEntries() - 1) : NextEntry;
+
+			const bool bSuccess = LerpEntries(Time, bExtrapolate, PrevEntry, NextEntry, CurveName, GetCollectedCurves(), OutCurveValue);
+			return bSuccess;
+		}
+	}
+
+	return PoseHistory->GetCurveValueAtTime(Time, CurveName, OutCurveValue, bExtrapolate);
+}
+
 void FMemStackPoseHistory::AddFutureRootBone(float Time, const FTransform& FutureRootBoneTransform, bool bStoreScales)
 {
 	// we don't allow to add "past" or "present" poses to FutureEntries
@@ -913,12 +1073,18 @@ void FMemStackPoseHistory::AddFutureRootBone(float Time, const FTransform& Futur
 }
 
 void FMemStackPoseHistory::AddFuturePose(float Time, FCSPose<FCompactPose>& ComponentSpacePose)
+{
+	FBlendedCurve Curves;
+	AddFuturePose(Time, ComponentSpacePose, Curves);
+}
+
+void FMemStackPoseHistory::AddFuturePose(float Time, FCSPose<FCompactPose>& ComponentSpacePose, const FBlendedCurve& Curves)
 {
 	// we don't allow to add "past" or "present" poses to FutureEntries
 	check(Time > 0.f);
 	check(PoseHistory);	
 	const int32 LowerBoundIdx = Algo::LowerBound(FutureEntries, Time, [](const FPoseHistoryEntry& Entry, float Value) { return Value > Entry.AccumulatedSeconds; });
-	FutureEntries.InsertDefaulted_GetRef(LowerBoundIdx).Update(Time, ComponentSpacePose, GetBoneToTransformMap(), true);
+	FutureEntries.InsertDefaulted_GetRef(LowerBoundIdx).Update(Time, ComponentSpacePose, GetBoneToTransformMap(), true, Curves, MakeConstArrayView(GetCollectedCurves()));
 }
 
 int32 FMemStackPoseHistory::GetNumEntries() const
diff --git a/Engine/Plugins/Animation/PoseSearch/Source/Runtime/Private/PoseSearchLibrary.cpp b/Engine/Plugins/Animation/PoseSearch/Source/Runtime/Private/PoseSearchLibrary.cpp
--- a/Engine/Plugins/Animation/PoseSearch/Source/Runtime/Private/PoseSearchLibrary.cpp
+++ b/Engine/Plugins/Animation/PoseSearch/Source/Runtime/Private/PoseSearchLibrary.cpp
@@ -857,6 +857,7 @@ UE::PoseSearch::FSearchResult UPoseSearchLibrary::MotionMatch(
 			// extracting 2 poses to be able to calculate velocities
 			FCSPose<FCompactPose> ComponentSpacePose;
 			FCompactPose Pose;
+			FBlendedCurve Curves;
 			Pose.SetBoneContainer(&AnimInstances[RoleIndex]->GetRequiredBonesOnAnyThread());
 
 			// @todo: add input BlendParameters to support sampling FutureAnimation blendspaces and support for multi character
@@ -879,7 +880,7 @@ UE::PoseSearch::FSearchResult UPoseSearchLibrary::MotionMatch(
 				const float FuturePoseExtractionTime = FutureAnimationTime + (i - 1) * FiniteDelta;
 				const float FuturePoseAnimationTime = FutureIntervalTime + (i - 1) * FiniteDelta;
 
-				Sampler.ExtractPose(FuturePoseExtractionTime, Pose);
+				Sampler.ExtractPose(FuturePoseExtractionTime, Pose, Curves);
 				ComponentSpacePose.InitPose(Pose);
 				MemStackPoseHistories[RoleIndex].AddFuturePose(FuturePoseAnimationTime, ComponentSpacePose);
 			}
diff --git a/Engine/Plugins/Animation/PoseSearch/Source/Runtime/Public/PoseSearch/AnimNode_PoseSearchHistoryCollector.h b/Engine/Plugins/Animation/PoseSearch/Source/Runtime/Public/PoseSearch/AnimNode_PoseSearchHistoryCollector.h
--- a/Engine/Plugins/Animation/PoseSearch/Source/Runtime/Public/PoseSearch/AnimNode_PoseSearchHistoryCollector.h
+++ b/Engine/Plugins/Animation/PoseSearch/Source/Runtime/Public/PoseSearch/AnimNode_PoseSearchHistoryCollector.h
@@ -25,6 +25,9 @@ public:
 	UPROPERTY(EditAnywhere, Category = Settings)
 	TArray<FBoneReference> CollectedBones;
 
+	UPROPERTY(EditAnywhere, Category = Settings)
+	TArray<FName> CollectedCurves;
+
 	// if true, the pose history will be initialized with a ref pose at the location and orientation of the AnimInstance.
 	UPROPERTY(EditAnywhere, Category = Settings)
 	bool bInitializeWithRefPose = false;
diff --git a/Engine/Plugins/Animation/PoseSearch/Source/Runtime/Public/PoseSearch/PoseSearchHistory.h b/Engine/Plugins/Animation/PoseSearch/Source/Runtime/Public/PoseSearch/PoseSearchHistory.h
--- a/Engine/Plugins/Animation/PoseSearch/Source/Runtime/Public/PoseSearch/PoseSearchHistory.h
+++ b/Engine/Plugins/Animation/PoseSearch/Source/Runtime/Public/PoseSearch/PoseSearchHistory.h
@@ -33,15 +33,17 @@ struct FPoseHistoryEntry
 	TArray<FQuat4f> ComponentSpaceRotations;
 	TArray<FVector> ComponentSpacePositions;
 	TArray<FVector3f> ComponentSpaceScales;
+	TArray<float> CurveValues;
 	float AccumulatedSeconds = 0.f;
 
-	void Update(float Time, FCSPose<FCompactPose>& ComponentSpacePose, const FBoneToTransformMap& BoneToTransformMap, bool bStoreScales);
+	void Update(float Time, FCSPose<FCompactPose>& ComponentSpacePose, const FBoneToTransformMap& BoneToTransformMap, bool bStoreScales, const FBlendedCurve& Curves, const TConstArrayView<FName>& CollectedCurves);
 
 	POSESEARCH_API void SetNum(int32 Num, bool bStoreScales);
 	int32 Num() const;
 
 	POSESEARCH_API void SetComponentSpaceTransform(int32 Index, const FTransform& Transform);
 	FTransform GetComponentSpaceTransform(int32 Index) const;
+	float GetCurveValue(int32 Index) const;
 };
 FArchive& operator<<(FArchive& Ar, FPoseHistoryEntry& Entry);
 
@@ -55,6 +57,8 @@ public:
 	// if ReferenceBoneIndexType is FBoneIndexType(-1) (ComponentSpaceIndexType), OutBoneTransform is in component space
 	// if ReferenceBoneIndexType is FBoneIndexType(-2) (WorldSpaceIndexType), OutBoneTransform is in world space
 	virtual bool GetTransformAtTime(float Time, FTransform& OutBoneTransform, const USkeleton* BoneIndexSkeleton = nullptr, FBoneIndexType BoneIndexType = RootBoneIndexType, FBoneIndexType ReferenceBoneIndexType = ComponentSpaceIndexType, bool bExtrapolate = false) const = 0;
+	// @todo: consider consolidating into a (templated?) get X value at time once we add custom attributes to pose history.
+	virtual bool GetCurveValueAtTime(float Time, const FName& CurveName, float& OutCurveValue, bool bExtrapolate = false) const = 0;
 	virtual const FPoseSearchQueryTrajectory& GetTrajectory() const = 0;
 	
 	// @todo: deprecate this API. TrajectorySpeedMultiplier should be a global query scaling value passed as input parameter of FSearchContext during config BuildQuery
@@ -62,6 +66,7 @@ public:
 	virtual bool IsEmpty() const = 0;
 
 	virtual const FBoneToTransformMap& GetBoneToTransformMap() const = 0;
+	virtual const TConstArrayView<FName> GetCollectedCurves() const = 0;
 	virtual int32 GetNumEntries() const = 0;
 	virtual const FPoseHistoryEntry& GetEntry(int32 EntryIndex) const = 0;
 
@@ -78,11 +83,12 @@ struct POSESEARCH_API FArchivedPoseHistory : public IPoseHistory
 
 	// IPoseHistory interface
 	virtual bool GetTransformAtTime(float Time, FTransform& OutBoneTransform, const USkeleton* BoneIndexSkeleton = nullptr, FBoneIndexType BoneIndexType = RootBoneIndexType, FBoneIndexType ReferenceBoneIndexType = ComponentSpaceIndexType, bool bExtrapolate = false) const override;
+	virtual bool GetCurveValueAtTime(float Time, const FName& CurveName, float& outcurvevalue, bool bExtrapolate = false) const override;
 	virtual const FPoseSearchQueryTrajectory& GetTrajectory() const override { return Trajectory; }
 	virtual float GetTrajectorySpeedMultiplier() const override { return 1.f; }
 	virtual bool IsEmpty() const override { return Entries.IsEmpty(); }
 	virtual const FBoneToTransformMap& GetBoneToTransformMap() const override { return BoneToTransformMap; }
-	
+	virtual const TConstArrayView<FName> GetCollectedCurves() const override { return CollectedCurves; }
 	virtual int32 GetNumEntries() const override { return Entries.Num(); }
 	virtual const FPoseHistoryEntry& GetEntry(int32 EntryIndex) const override { return Entries[EntryIndex]; }
 
@@ -93,6 +99,8 @@ struct POSESEARCH_API FArchivedPoseHistory : public IPoseHistory
 	// End of IPoseHistory interface
 
 	FBoneToTransformMap BoneToTransformMap;
+	// @todo: Make this a map if this is expected to be big.
+	TArray<FName> CollectedCurves;
 	TArray<FPoseHistoryEntry> Entries;
 	FPoseSearchQueryTrajectory Trajectory;
 };
@@ -113,13 +121,18 @@ struct POSESEARCH_API FPoseHistory : public IPoseHistory
 	void EvaluateComponentSpace_AnyThread(float DeltaTime, FCSPose<FCompactPose>& ComponentSpacePose, bool bStoreScales,
 		float RootBoneRecoveryTime, float RootBoneTranslationRecoveryRatio, float RootBoneRotationRecoveryRatio,
 		bool bNeedsReset, bool bCacheBones, const TArray<FBoneIndexType>& RequiredBones);
+	void EvaluateComponentSpace_AnyThread(float DeltaTime, FCSPose<FCompactPose>& ComponentSpacePose, bool bStoreScales,
+		float RootBoneRecoveryTime, float RootBoneTranslationRecoveryRatio, float RootBoneRotationRecoveryRatio,
+		bool bNeedsReset, bool bCacheBones, const TArray<FBoneIndexType>& RequiredBones, const FBlendedCurve& Curves, const TConstArrayView<FName>& CollectedCurves);
 
 	// IPoseHistory interface
 	virtual bool GetTransformAtTime(float Time, FTransform& OutBoneTransform, const USkeleton* BoneIndexSkeleton = nullptr, FBoneIndexType BoneIndexType = RootBoneIndexType, FBoneIndexType ReferenceBoneIndexType = ComponentSpaceIndexType, bool bExtrapolate = false) const override;
+	virtual bool GetCurveValueAtTime(float Time, const FName& CurveName, float& outcurvevalue, bool bExtrapolate = false) const override;
 	virtual const FPoseSearchQueryTrajectory& GetTrajectory() const override;
 	virtual float GetTrajectorySpeedMultiplier() const override;
 	virtual bool IsEmpty() const override;
 	virtual const FBoneToTransformMap& GetBoneToTransformMap() const override;
+	virtual const TConstArrayView<FName> GetCollectedCurves() const override;
 	
 	void SetTrajectory(const FPoseSearchQueryTrajectory& InTrajectory, float InTrajectorySpeedMultiplier = 1.f);
 	virtual int32 GetNumEntries() const override;
@@ -154,6 +167,9 @@ private:
 		// map of FBoneIndexType(s) to collect. If Empty all the bones get collected
 		FBoneToTransformMap BoneToTransformMap;
 		
+		// list of curves that we want to collect into our history.
+		TArray<FName> CollectedCurves;
+
 		// GetTypeHash for BoneToTransformMap
 		uint32 BoneToTransformMapTypeHash = 0;
 
@@ -185,10 +201,12 @@ struct POSESEARCH_API FMemStackPoseHistory : public IPoseHistory
 
 	// IPoseHistory interface
 	virtual bool GetTransformAtTime(float Time, FTransform& OutBoneTransform, const USkeleton* BoneIndexSkeleton = nullptr, FBoneIndexType BoneIndexType = RootBoneIndexType, FBoneIndexType ReferenceBoneIndexType = ComponentSpaceIndexType, bool bExtrapolate = false) const override;
+	virtual bool GetCurveValueAtTime(float Time, const FName& CurveName, float& outcurvevalue, bool bExtrapolate = false) const override;
 	virtual const FPoseSearchQueryTrajectory& GetTrajectory() const override { check(PoseHistory); return PoseHistory->GetTrajectory(); }
 	virtual float GetTrajectorySpeedMultiplier() const override { check(PoseHistory); return PoseHistory->GetTrajectorySpeedMultiplier(); }
 	virtual bool IsEmpty() const override { check(PoseHistory); return PoseHistory->IsEmpty() && FutureEntries.IsEmpty(); }
 	virtual const FBoneToTransformMap& GetBoneToTransformMap() const override { check(PoseHistory); return PoseHistory->GetBoneToTransformMap(); }
+	virtual const TConstArrayView<FName> GetCollectedCurves() const override { check(PoseHistory); return PoseHistory->GetCollectedCurves(); }
 	
 	virtual int32 GetNumEntries() const override;
 	virtual const FPoseHistoryEntry& GetEntry(int32 EntryIndex) const override;
@@ -201,6 +219,7 @@ struct POSESEARCH_API FMemStackPoseHistory : public IPoseHistory
 
 	void AddFutureRootBone(float Time, const FTransform& FutureRootBoneTransform, bool bStoreScales);
 	void AddFuturePose(float Time, FCSPose<FCompactPose>& ComponentSpacePose);
+	void AddFuturePose(float Time, FCSPose<FCompactPose>& ComponentSpacePose, const FBlendedCurve& Curves);
 
 	const IPoseHistory* GetThisOrPoseHistory() const { return FutureEntries.IsEmpty() ? PoseHistory : this; }
 
