posesearch - UPoseSearchLibrary::IsAnimationAssetLooping

diff --git a/Engine/Plugins/Animation/PoseSearch/Source/Runtime/Private/PoseSearchLibrary.cpp b/Engine/Plugins/Animation/PoseSearch/Source/Runtime/Private/PoseSearchLibrary.cpp
--- a/Engine/Plugins/Animation/PoseSearch/Source/Runtime/Private/PoseSearchLibrary.cpp
+++ b/Engine/Plugins/Animation/PoseSearch/Source/Runtime/Private/PoseSearchLibrary.cpp
@@ -694,6 +694,26 @@ void UPoseSearchLibrary::UpdateMotionMatchingState(
 #endif
 }
 
+void UPoseSearchLibrary::IsAnimationAssetLooping(const UObject* Asset, bool& bIsAssetLooping)
+{
+	if (const UAnimSequenceBase* SequenceBase = Cast<const UAnimSequenceBase>(Asset))
+	{
+		bIsAssetLooping = SequenceBase->bLoop;
+	}
+	else if (const UBlendSpace* BlendSpace = Cast<const UBlendSpace>(Asset))
+	{
+		bIsAssetLooping = BlendSpace->bLoop;
+	}
+	else if (const UMultiAnimAsset* MultiAnimAsset = Cast<const UMultiAnimAsset>(Asset))
+	{
+		bIsAssetLooping = MultiAnimAsset->IsLooping();
+	}
+	else
+	{
+		bIsAssetLooping = false;
+	}
+}
+
 void UPoseSearchLibrary::MotionMatch(
 	UAnimInstance* AnimInstance,
 	TArray<UObject*> AssetsToSearch,
diff --git a/Engine/Plugins/Animation/PoseSearch/Source/Runtime/Public/PoseSearch/PoseSearchLibrary.h b/Engine/Plugins/Animation/PoseSearch/Source/Runtime/Public/PoseSearch/PoseSearchLibrary.h
--- a/Engine/Plugins/Animation/PoseSearch/Source/Runtime/Public/PoseSearch/PoseSearchLibrary.h
+++ b/Engine/Plugins/Animation/PoseSearch/Source/Runtime/Public/PoseSearch/PoseSearchLibrary.h
@@ -285,5 +285,8 @@ public:
 	static const FAnimNode_PoseSearchHistoryCollector_Base* FindPoseHistoryNode(
 		const FName PoseHistoryName,
 		const UAnimInstance* AnimInstance);
+
+	UFUNCTION(BlueprintPure, Category = "Animation|Pose Search|Experimental", meta = (BlueprintThreadSafe))
+	static void IsAnimationAssetLooping(const UObject* Asset, bool& bIsAssetLooping);
 };
 
