AnimSandbox: Fix for wavy trajectories on slopes

diff --git a/Engine/Plugins/Animation/PoseSearch/Source/Runtime/Private/PoseSearchTrajectoryLibrary.cpp b/Engine/Plugins/Animation/PoseSearch/Source/Runtime/Private/PoseSearchTrajectoryLibrary.cpp
--- a/Engine/Plugins/Animation/PoseSearch/Source/Runtime/Private/PoseSearchTrajectoryLibrary.cpp
+++ b/Engine/Plugins/Animation/PoseSearch/Source/Runtime/Private/PoseSearchTrajectoryLibrary.cpp
@@ -8,6 +8,7 @@
 #include "GameFramework/Character.h"
 #include "GameFramework/CharacterMovementComponent.h"
 #include "PoseSearch/PoseSearchDefines.h"
+#include "Kismet/KismetMathLibrary.h"
 
 void FPoseSearchTrajectoryData::UpdateData(
 	float DeltaTime,
@@ -532,11 +533,10 @@ void UPoseSearchTrajectoryLibrary::HandleTrajectoryWorldCollisionsWithGravity(co
 
 				FreeFallAccumulatedSeconds += Sample.AccumulatedSeconds - PrevSample.AccumulatedSeconds;
 
-				// projecting Sample.Position on the HitResult plane and offsetting it by FloorCollisionsOffset
 				if (bIsLastImpactValid)
 				{
-					const FVector DeltaImpactPoint = Sample.Position - LastImpactPoint;
-					Sample.Position += (FloorCollisionsOffset - (DeltaImpactPoint | LastImpactNormal)) * LastImpactNormal;
+					const FPlane GroundPlane = FPlane(PrevSample.Position, -GravityDirection);
+					Sample.Position = FPlane::PointPlaneProject(Sample.Position, GroundPlane);
 				}
 
 				// applying gravity
@@ -546,11 +546,12 @@ void UPoseSearchTrajectoryLibrary::HandleTrajectoryWorldCollisionsWithGravity(co
 				FHitResult HitResult;
 				if (FloorCollisionsOffset > 0.f && UKismetSystemLibrary::LineTraceSingle(WorldContextObject, Sample.Position + (GravityDirection * -MaxObstacleHeight), Sample.Position, TraceChannel, bTraceComplex, ActorsToIgnore, DrawDebugType, HitResult, bIgnoreSelf, TraceColor, TraceHitColor, DrawTime))
 				{
-					LastImpactPoint = HitResult.ImpactPoint;
+					// Only allow our trace to move trajectory along gravity direction.
+					LastImpactPoint = UKismetMathLibrary::FindClosestPointOnLine(HitResult.ImpactPoint, Sample.Position, GravityDirection);
 					LastImpactNormal = HitResult.Normal;
 					bIsLastImpactValid = true;
 
-					Sample.Position = LastImpactPoint + (LastImpactNormal * FloorCollisionsOffset);
+					Sample.Position = LastImpactPoint - GravityDirection * FloorCollisionsOffset;
 
 					if (bIsFirstFall)
 					{
