engine - avoiding root motion extraction if !ExtractionContext.bExtractRootMotion in UAnimSequence::GetAnimationPose

diff --git a/Engine/Plugins/Animation/PoseSearch/Source/Runtime/Private/PoseSearchAssetSampler.cpp b/Engine/Plugins/Animation/PoseSearch/Source/Runtime/Private/PoseSearchAssetSampler.cpp
--- a/Engine/Plugins/Animation/PoseSearch/Source/Runtime/Private/PoseSearchAssetSampler.cpp
+++ b/Engine/Plugins/Animation/PoseSearch/Source/Runtime/Private/PoseSearchAssetSampler.cpp
@@ -409,6 +409,10 @@ void FAnimationAssetSampler::ExtractPose(float Time, FCompactPose& OutPose) cons
 	DeltaTimeRecord.Set(Time, 0.f);
 	FAnimExtractContext ExtractionCtx(double(Time), false, DeltaTimeRecord, IsLoopable());
 
+#if WITH_EDITOR
+	ExtractionCtx.bExtractWithRootMotionProvider = false;
+#endif // WITH_EDITOR
+
 	ExtractPose(ExtractionCtx, AnimPoseData);
 }
 
@@ -422,6 +426,10 @@ void FAnimationAssetSampler::ExtractPose(float Time, FCompactPose& OutPose, FBle
 	DeltaTimeRecord.Set(Time, 0.f);
 	FAnimExtractContext ExtractionCtx(double(Time), false, DeltaTimeRecord, IsLoopable());
 
+#if WITH_EDITOR
+	ExtractionCtx.bExtractWithRootMotionProvider = false;
+#endif // WITH_EDITOR
+
 	ExtractPose(ExtractionCtx, AnimPoseData);
 }
 
diff --git a/Engine/Source/Runtime/Engine/Classes/Animation/AnimationAsset.h b/Engine/Source/Runtime/Engine/Classes/Animation/AnimationAsset.h
--- a/Engine/Source/Runtime/Engine/Classes/Animation/AnimationAsset.h
+++ b/Engine/Source/Runtime/Engine/Classes/Animation/AnimationAsset.h
@@ -292,6 +292,9 @@ struct FAnimExtractContext
 
 #if WITH_EDITOR
 	bool bIgnoreRootLock;
+
+	// Experimental
+	bool bExtractWithRootMotionProvider;
 #endif 
 	
 	UE_DEPRECATED(5.1, "FAnimExtractContext construct with float-based time value is deprecated, use other signature")
@@ -305,6 +308,7 @@ struct FAnimExtractContext
 		, InterpolationOverride()
 #if WITH_EDITOR
 		, bIgnoreRootLock(false)
+		, bExtractWithRootMotionProvider(true)
 #endif 
 	{
 	}
@@ -319,6 +323,7 @@ struct FAnimExtractContext
 		, InterpolationOverride()
 #if WITH_EDITOR
 		, bIgnoreRootLock(false)
+		, bExtractWithRootMotionProvider(true)
 #endif 
 	{
 	}
diff --git a/Engine/Source/Runtime/Engine/Private/Animation/AnimSequence.cpp b/Engine/Source/Runtime/Engine/Private/Animation/AnimSequence.cpp
--- a/Engine/Source/Runtime/Engine/Private/Animation/AnimSequence.cpp
+++ b/Engine/Source/Runtime/Engine/Private/Animation/AnimSequence.cpp
@@ -1602,7 +1602,11 @@ void UAnimSequence::GetAnimationPose(FAnimationPoseData& OutAnimationPoseData, c
 	}
 
 	// If the sequence has root motion enabled, allow sampling of a root motion delta into the custom attribute container of the outgoing pose
-	if (HasRootMotion())
+	if (HasRootMotion()
+#if WITH_EDITOR
+		&& ExtractionContext.bExtractWithRootMotionProvider
+#endif // WITH_EDITOR
+		)
 	{
 		if (const UE::Anim::IAnimRootMotionProvider* RootMotionProvider = UE::Anim::IAnimRootMotionProvider::Get())
 		{
