posesearch - database lazy indexing improvements

diff --git a/Engine/Plugins/Animation/PoseSearch/Source/Editor/Private/PoseSearchDatabasePreviewScene.cpp b/Engine/Plugins/Animation/PoseSearch/Source/Editor/Private/PoseSearchDatabasePreviewScene.cpp
--- a/Engine/Plugins/Animation/PoseSearch/Source/Editor/Private/PoseSearchDatabasePreviewScene.cpp
+++ b/Engine/Plugins/Animation/PoseSearch/Source/Editor/Private/PoseSearchDatabasePreviewScene.cpp
@@ -56,20 +56,5 @@ namespace UE::PoseSearch
 		}
 
 		GetWorld()->Tick(LEVELTICK_All, InDeltaTime);
-
-		FDatabaseViewModel* ViewModel = GetEditor()->GetViewModel();
-		const UPoseSearchDatabase* Database = ViewModel->GetPoseSearchDatabase();
-
-		if (!ViewModel->GetPreviewActors().IsEmpty() && EAsyncBuildIndexResult::Success == FAsyncPoseSearchDatabasesManagement::RequestAsyncBuildIndex(Database, ERequestAsyncBuildFlag::ContinueRequest))
-		{
-			const bool bDisplayRootMotionSpeed = ViewModel->IsDisplayRootMotionSpeedChecked();
-			const bool bDisplayBlockTransition = ViewModel->IsDisplayBlockTransitionChecked();
-			bool bDrawQueryVector = ViewModel->ShouldDrawQueryVector();
-
-			for (TArray<FDatabasePreviewActor>& PreviewActorGroup : ViewModel->GetPreviewActors())
-			{
-				bDrawQueryVector &= !FDatabasePreviewActor::DrawPreviewActors(PreviewActorGroup, Database, bDisplayRootMotionSpeed, bDisplayBlockTransition, bDrawQueryVector ? ViewModel->GetQueryVector() : TConstArrayView<float>());
-			}
-		}
 	}
 }
diff --git a/Engine/Plugins/Animation/PoseSearch/Source/Editor/Private/PoseSearchDatabaseViewModel.cpp b/Engine/Plugins/Animation/PoseSearch/Source/Editor/Private/PoseSearchDatabaseViewModel.cpp
--- a/Engine/Plugins/Animation/PoseSearch/Source/Editor/Private/PoseSearchDatabaseViewModel.cpp
+++ b/Engine/Plugins/Animation/PoseSearch/Source/Editor/Private/PoseSearchDatabaseViewModel.cpp
@@ -225,11 +225,6 @@ bool FDatabasePreviewActor::DrawPreviewActors(TConstArrayView<FDatabasePreviewAc
 
 	for (const FDatabasePreviewActor& PreviewActor : PreviewActors)
 	{
-		if (EAsyncBuildIndexResult::Success != FAsyncPoseSearchDatabasesManagement::RequestAsyncBuildIndex(PoseSearchDatabase, ERequestAsyncBuildFlag::ContinueRequest))
-		{
-			return false;
-		}
-
 		const FSearchIndex& SearchIndex = PoseSearchDatabase->GetSearchIndex();
 
 		// This condition happens when the database got reindexed and the new valid SearchIndex has different cardinality for assets or poses.
@@ -516,16 +511,18 @@ void FDatabaseViewModel::OnPreviewActorClassChanged()
 
 void FDatabaseViewModel::Tick(float DeltaSeconds)
 {
-	const float DeltaPlayTime = DeltaSeconds * DeltaTimeMultiplier;
-	if (!FMath::IsNearlyZero(DeltaPlayTime))
+	if (!PreviewActors.IsEmpty())
 	{
-		PlayTime += DeltaPlayTime;
-		PlayTime = FMath::Clamp(PlayTime, MinPreviewPlayLength, MaxPreviewPlayLength);
+		const float DeltaPlayTime = DeltaSeconds * DeltaTimeMultiplier;
 
-		if (const UPoseSearchDatabase* Database = GetPoseSearchDatabase())
+		const UPoseSearchDatabase* Database = GetPoseSearchDatabase();
+		if (EAsyncBuildIndexResult::Success == FAsyncPoseSearchDatabasesManagement::RequestAsyncBuildIndex(Database, ERequestAsyncBuildFlag::ContinueRequest))
 		{
-			if (EAsyncBuildIndexResult::Success == FAsyncPoseSearchDatabasesManagement::RequestAsyncBuildIndex(Database, ERequestAsyncBuildFlag::ContinueRequest))
+			if (!FMath::IsNearlyZero(DeltaPlayTime))
 			{
+				PlayTime += DeltaPlayTime;
+				PlayTime = FMath::Clamp(PlayTime, MinPreviewPlayLength, MaxPreviewPlayLength);
+
 				for (TArray<FDatabasePreviewActor>& PreviewActorGroup : PreviewActors)
 				{
 					for (FDatabasePreviewActor& PreviewActor : PreviewActorGroup)
@@ -534,6 +531,12 @@ void FDatabaseViewModel::Tick(float DeltaSeconds)
 					}
 				}
 			}
+
+			bool bShouldDrawQueryVector = ShouldDrawQueryVector();
+			for (TArray<FDatabasePreviewActor>& PreviewActorGroup : PreviewActors)
+			{
+				bShouldDrawQueryVector &= !FDatabasePreviewActor::DrawPreviewActors(PreviewActorGroup, Database, bDisplayRootMotionSpeed, bDisplayBlockTransition, bShouldDrawQueryVector ? GetQueryVector() : TConstArrayView<float>());
+			}
 		}
 	}
 }
@@ -957,8 +960,9 @@ void FDatabaseViewModel::SetDrawQueryVector(bool bValue)
 
 const FSearchIndexAsset* FDatabaseViewModel::GetSelectedActorIndexAsset() const
 {
-	if (const UPoseSearchDatabase* Database = GetPoseSearchDatabase())
+	if (SelectedActorIndexAssetIndex >= 0)
 	{
+		const UPoseSearchDatabase* Database = GetPoseSearchDatabase();
 		if (EAsyncBuildIndexResult::Success == FAsyncPoseSearchDatabasesManagement::RequestAsyncBuildIndex(Database, ERequestAsyncBuildFlag::ContinueRequest))
 		{
 			const FSearchIndex& SearchIndex = Database->GetSearchIndex();
diff --git a/Engine/Plugins/Animation/PoseSearch/Source/Runtime/Private/PoseSearchDerivedData.cpp b/Engine/Plugins/Animation/PoseSearch/Source/Runtime/Private/PoseSearchDerivedData.cpp
--- a/Engine/Plugins/Animation/PoseSearch/Source/Runtime/Private/PoseSearchDerivedData.cpp
+++ b/Engine/Plugins/Animation/PoseSearch/Source/Runtime/Private/PoseSearchDerivedData.cpp
@@ -89,7 +89,8 @@ static TAutoConsoleVariable<int32> CVarMotionMatchTestNumIterations(TEXT("a.Moti
 static bool AnyTestFlags(int32 Flags) { return (CVarMotionMatchTestFlags.GetValueOnAnyThread() & Flags) != 0; }
 #endif // ENABLE_ANIM_DEBUG
 
-static TAutoConsoleVariable<bool> CVarMotionMatchReindexCancelledDatabases(TEXT("a.MotionMatch.ReindexCancelledDatabases"), true, TEXT("Reindex Cancelled Databases"));
+static TAutoConsoleVariable<bool> CVarMotionMatchReindexCancelledDatabases(TEXT("a.MotionMatch.ReindexCancelledDatabases"), false, TEXT("Reindex Cancelled Databases"));
+static TAutoConsoleVariable<bool> CVarMotionMatchReindexAllReferencedDatabases(TEXT("a.MotionMatch.ReindexAllReferencedDatabases"), true, TEXT("Reindex All Referenced Databases"));
 
 static const UE::DerivedData::FValueId Id(UE::DerivedData::FValueId::FromName("Data"));
 static const UE::DerivedData::FCacheBucket Bucket("PoseSearchDatabase");
@@ -102,6 +103,28 @@ static FCookStatsManager::FAutoRegisterCallback RegisterCookStats([](FCookStatsM
 	});
 #endif // ENABLE_COOK_STATS
 
+typedef TSet<const UPoseSearchDatabase*, DefaultKeyFuncs<const UPoseSearchDatabase*>, TInlineSetAllocator<256>> FDatabaseSet;
+
+static void RecursivePopulateDependentDatabases(const UPoseSearchDatabase* Database, FDatabaseSet& DatabaseSet)
+{
+	if (Database)
+	{
+		bool bIsAlreadyInSet = false;
+		DatabaseSet.Add(Database, &bIsAlreadyInSet);
+
+		if (!bIsAlreadyInSet)
+		{
+			if (const UPoseSearchNormalizationSet* NormalizationSet = Database->NormalizationSet)
+			{
+				for (const UPoseSearchDatabase* DependentDatabase : NormalizationSet->Databases)
+				{
+					RecursivePopulateDependentDatabases(DependentDatabase, DatabaseSet);
+				}
+			}
+		}
+	}
+}
+
 // helper struct to calculate mean deviations
 struct FMeanDeviationCalculator
 {
@@ -2398,7 +2421,7 @@ void FAsyncPoseSearchDatabasesManagement::AddReferencedObjects(FReferenceCollect
 	}
 }
 
-EAsyncBuildIndexResult FAsyncPoseSearchDatabasesManagement::RequestAsyncBuildIndex(const UPoseSearchDatabase* Database, ERequestAsyncBuildFlag Flag)
+EAsyncBuildIndexResult FAsyncPoseSearchDatabasesManagement::RequestAsyncBuildIndexInternal(const UPoseSearchDatabase* Database, ERequestAsyncBuildFlag Flag)
 {
 	if (!IsValid(Database))
 	{
@@ -2461,6 +2484,26 @@ EAsyncBuildIndexResult FAsyncPoseSearchDatabasesManagement::RequestAsyncBuildInd
 	return EAsyncBuildIndexResult::InProgress;
 }
 
+EAsyncBuildIndexResult FAsyncPoseSearchDatabasesManagement::RequestAsyncBuildIndex(const UPoseSearchDatabase* Database, ERequestAsyncBuildFlag Flag)
+{
+	if (CVarMotionMatchReindexAllReferencedDatabases.GetValueOnAnyThread())
+	{
+		FDatabaseSet DatabaseSet;
+		RecursivePopulateDependentDatabases(Database, DatabaseSet);
+
+		for (FDatabaseSet::TConstIterator Iter = DatabaseSet.CreateConstIterator(); Iter; ++Iter)
+		{
+			const UPoseSearchDatabase* DependentDatabase = *Iter;
+			if (DependentDatabase != Database)
+			{
+				RequestAsyncBuildIndexInternal(DependentDatabase, ERequestAsyncBuildFlag::ContinueRequest);
+			}
+		}
+	}
+
+	return RequestAsyncBuildIndexInternal(Database, Flag);
+}
+
 } // namespace UE::PoseSearch
 
 #undef LOCTEXT_NAMESPACE
diff --git a/Engine/Plugins/Animation/PoseSearch/Source/Runtime/Public/PoseSearch/PoseSearchDerivedData.h b/Engine/Plugins/Animation/PoseSearch/Source/Runtime/Public/PoseSearch/PoseSearchDerivedData.h
--- a/Engine/Plugins/Animation/PoseSearch/Source/Runtime/Public/PoseSearch/PoseSearchDerivedData.h
+++ b/Engine/Plugins/Animation/PoseSearch/Source/Runtime/Public/PoseSearch/PoseSearchDerivedData.h
@@ -40,6 +40,7 @@ namespace UE::PoseSearch
 		~FAsyncPoseSearchDatabasesManagement();
 
 		static FAsyncPoseSearchDatabasesManagement& Get();
+		static EAsyncBuildIndexResult RequestAsyncBuildIndexInternal(const UPoseSearchDatabase* Database, ERequestAsyncBuildFlag Flag);
 
 		void OnObjectModified(UObject* Object);
 		void OnObjectTransacted(UObject* Object, const FTransactionObjectEvent& TransactionObjectEvent);
