posesearch - indexing performance improvements, added TestDDCKeyDeterminism

diff --git a/Engine/Plugins/Animation/PoseSearch/Source/Runtime/Private/PoseSearchAssetIndexer.cpp b/Engine/Plugins/Animation/PoseSearch/Source/Runtime/Private/PoseSearchAssetIndexer.cpp
--- a/Engine/Plugins/Animation/PoseSearch/Source/Runtime/Private/PoseSearchAssetIndexer.cpp
+++ b/Engine/Plugins/Animation/PoseSearch/Source/Runtime/Private/PoseSearchAssetIndexer.cpp
@@ -219,6 +219,8 @@ FAssetIndexer::FAssetIndexer(const TConstArrayView<FBoneContainer> InBoneContain
 {
 	check(BoneContainers.Num() == AssetSamplers.Num() && BoneContainers.Num() == RoleToIndex.Num());
 	check(IsValid(RoleToIndex));
+
+	CachedEntries.Reserve(SearchIndexAsset.GetNumPoses());
 }
 
 void FAssetIndexer::AssignWorkingData(int32 InStartPoseIdx, TArrayView<float> InOutFeatureVectorTable, TArrayView<FPoseMetadata> InOutPoseMetadata)
@@ -488,9 +490,7 @@ FAssetIndexer::FCachedEntry& FAssetIndexer::GetEntry(float SampleTime)
 				AssetSamplers.MirrorPose(Pose, RoleIndex);
 			}
 
-			FCSPose<FCompactPose> StackComponentSpacePose;
-			StackComponentSpacePose.InitPose(MoveTemp(Pose));
-			Entry->ComponentSpacePose[RoleIndex].CopyPose(StackComponentSpacePose);
+			Entry->ComponentSpacePose[RoleIndex].InitPose(Pose);
 			Entry->Curves[RoleIndex].CopyFrom(Curve);
 
 			Entry->RootTransform[RoleIndex] = MirrorTransform(SampleRootTransform, RoleIndex);
diff --git a/Engine/Plugins/Animation/PoseSearch/Source/Runtime/Private/PoseSearchDerivedData.cpp b/Engine/Plugins/Animation/PoseSearch/Source/Runtime/Private/PoseSearchDerivedData.cpp
--- a/Engine/Plugins/Animation/PoseSearch/Source/Runtime/Private/PoseSearchDerivedData.cpp
+++ b/Engine/Plugins/Animation/PoseSearch/Source/Runtime/Private/PoseSearchDerivedData.cpp
@@ -83,6 +83,9 @@ enum EMotionMatchTestFlags
 
 	// test FAnimationAssetSampler determinism across multiple editro executions. It'll store some bin files in \Engine\TestAssetSamplerDeterminism
 	TestAssetSamplerDeterminismFromPreviousExecution = 1 << 13,
+
+	// test DDC key generation determinism
+	TestDDCKeyDeterminism = 1 << 14,
 };
 static TAutoConsoleVariable<int32> CVarMotionMatchTestFlags(TEXT("a.MotionMatch.TestFlags"), EMotionMatchTestFlags::None, TEXT("Test Motion Matching using EMotionMatchTestFlags"));
 static TAutoConsoleVariable<int32> CVarMotionMatchTestNumIterations(TEXT("a.MotionMatch.TestNumIterations"), 10, TEXT("Test Motion Matching Num Iterations"));
@@ -1418,6 +1421,7 @@ void FPoseSearchDatabaseAsyncCacheTask::StartNewRequestIfNeeded(bool bPerformCon
 			DerivedDataKey = NewDerivedDataKey;
 
 			DatabaseDependencies.Reset();
+			DatabaseDependencies.Reserve(KeyBuilder.GetDependencies().Num());
 			for (const UObject* Dependency : KeyBuilder.GetDependencies())
 			{
 				DatabaseDependencies.Add(Dependency);
@@ -2399,6 +2403,32 @@ void FAsyncPoseSearchDatabasesManagement::Tick(float DeltaTime)
 			Tasks[TaskIndex]->Update(Mutex);
 		}
 	}
+	
+#if ENABLE_ANIM_DEBUG
+	if (AnyTestFlags(EMotionMatchTestFlags::TestDDCKeyDeterminism))
+	{
+		const int32 NumIterations = CVarMotionMatchTestNumIterations.GetValueOnAnyThread();
+		for (int32 TaskIndex = 0; TaskIndex < Tasks.Num(); ++TaskIndex)
+		{
+			if (const UPoseSearchDatabase* Database = Tasks[TaskIndex]->GetDatabase())
+			{
+				const FKeyBuilder KeyBuilder(Database, false, false);
+				const FIoHash IoHash = KeyBuilder.Finalize();
+
+				for (int32 IterationIndex = 0; IterationIndex < NumIterations; ++IterationIndex)
+				{
+					const FKeyBuilder TestKeyBuilder(Database, false, false);
+					const FIoHash TestIoHash = TestKeyBuilder.Finalize();
+
+					if (IoHash != TestIoHash)
+					{
+						UE_LOG(LogPoseSearch, Error, TEXT("FKeyBuilder - key generation is not deterministic: %s / %s for asset %s"), *LexToString(IoHash), *LexToString(TestIoHash), *Database->GetName());
+					}
+				}
+			}
+		}
+	}
+#endif ENABLE_ANIM_DEBUG
 }
 
 void FAsyncPoseSearchDatabasesManagement::TickCook(float DeltaTime, bool bCookCompete)
diff --git a/Engine/Plugins/Animation/PoseSearch/Source/Runtime/Public/PoseSearch/PoseSearchAssetIndexer.h b/Engine/Plugins/Animation/PoseSearch/Source/Runtime/Public/PoseSearch/PoseSearchAssetIndexer.h
--- a/Engine/Plugins/Animation/PoseSearch/Source/Runtime/Public/PoseSearch/PoseSearchAssetIndexer.h
+++ b/Engine/Plugins/Animation/PoseSearch/Source/Runtime/Public/PoseSearch/PoseSearchAssetIndexer.h
@@ -128,6 +128,17 @@ private:
 	// bClamped will be true if SampleTime or OriginTime are outside the animation duration boundaries
 	bool GetSamplePositionInternal(FVector& OutSamplePosition, float SampleTime, float OriginTime, bool& bClamped, int8 SchemaSampleBoneIdx, int8 SchemaOriginBoneIdx, const FRole& SampleRole, const FRole& OriginRole, int32 SamplingAttributeId);
 
+	struct FCachedCSPose : public FCSPose<FCompactHeapPose>
+	{
+		void InitPose(const FCompactPose& SrcPose)
+		{
+			Pose.CopyBonesFrom(SrcPose);
+			ComponentSpaceFlags.Empty(Pose.GetNumBones());
+			ComponentSpaceFlags.AddZeroed(Pose.GetNumBones());
+			ComponentSpaceFlags[0] = 1;
+		}
+	};
+
 	struct FCachedEntry
 	{
 		float SampleTime = 0.f;
@@ -135,7 +146,7 @@ private:
 
 		// RootTransform and ComponentSpacePose are stored mirrored in case SearchIndexAsset.IsMirrored
 		TArray<FTransform> RootTransform;
-		TArray<FCSPose<FCompactHeapPose>> ComponentSpacePose;
+		TArray<FCachedCSPose> ComponentSpacePose;
 		TArray<FBlendedHeapCurve> Curves;
 	};
 
