posesearch - fix for error: extra tokens at end of #endif directive [-Werror,-Wextra-tokens] #endif ENABLE_ANIM_DEBUG

diff --git a/Engine/Plugins/Animation/PoseSearch/Source/Runtime/Private/PoseSearchDerivedData.cpp b/Engine/Plugins/Animation/PoseSearch/Source/Runtime/Private/PoseSearchDerivedData.cpp
--- a/Engine/Plugins/Animation/PoseSearch/Source/Runtime/Private/PoseSearchDerivedData.cpp
+++ b/Engine/Plugins/Animation/PoseSearch/Source/Runtime/Private/PoseSearchDerivedData.cpp
@@ -2428,7 +2428,7 @@ void FAsyncPoseSearchDatabasesManagement::Tick(float DeltaTime)
 			}
 		}
 	}
-#endif ENABLE_ANIM_DEBUG
+#endif // ENABLE_ANIM_DEBUG
 }
 
 void FAsyncPoseSearchDatabasesManagement::TickCook(float DeltaTime, bool bCookCompete)
