posesearch - DDC key generation speed up via partial caching of object hashes

diff --git a/Engine/Plugins/Animation/PoseSearch/Source/Runtime/Private/PoseSearchDerivedData.cpp b/Engine/Plugins/Animation/PoseSearch/Source/Runtime/Private/PoseSearchDerivedData.cpp
--- a/Engine/Plugins/Animation/PoseSearch/Source/Runtime/Private/PoseSearchDerivedData.cpp
+++ b/Engine/Plugins/Animation/PoseSearch/Source/Runtime/Private/PoseSearchDerivedData.cpp
@@ -93,7 +93,10 @@ static bool AnyTestFlags(int32 Flags) { return (CVarMotionMatchTestFlags.GetValu
 #endif // ENABLE_ANIM_DEBUG
 
 static TAutoConsoleVariable<bool> CVarMotionMatchReindexCancelledDatabases(TEXT("a.MotionMatch.ReindexCancelledDatabases"), false, TEXT("Reindex Cancelled Databases"));
+// Experimental, this feature might be removed without warning, not for production use
 static TAutoConsoleVariable<bool> CVarMotionMatchReindexAllReferencedDatabases(TEXT("a.MotionMatch.ReindexAllReferencedDatabases"), true, TEXT("Reindex All Referenced Databases"));
+// Experimental, this feature might be removed without warning, not for production use
+static TAutoConsoleVariable<int32> CVarMotionMatchPartialKeyHashesMode(TEXT("a.MotionMatch.PartialKeyHashesMode"), 0, TEXT("0: use partial key hashes, 1: do not use partial key hashes, 2: do not use and validate partial key hashes"));
 
 static const UE::DerivedData::FValueId Id(UE::DerivedData::FValueId::FromName("Data"));
 static const UE::DerivedData::FCacheBucket Bucket("PoseSearchDatabase");
@@ -1300,9 +1303,9 @@ struct FPoseSearchDatabaseAsyncCacheTask
 		Failed		// the task has ended unsuccessfully
 	};
 
-	FPoseSearchDatabaseAsyncCacheTask(UPoseSearchDatabase* InDatabase, bool bPerformConditionalPostLoadIfRequired);
-	void StartNewRequestIfNeeded(bool bPerformConditionalPostLoadIfRequired);
-	void Update(FCriticalSection& OuterMutex);
+	FPoseSearchDatabaseAsyncCacheTask(UPoseSearchDatabase* InDatabase, bool bPerformConditionalPostLoadIfRequired, FPartialKeyHashes& PartialKeyHashes);
+	void StartNewRequestIfNeeded(bool bPerformConditionalPostLoadIfRequired, FPartialKeyHashes& PartialKeyHashes);
+	void Update(FCriticalSection& OuterMutex, FPartialKeyHashes& PartialKeyHashes);
 	void Wait(FCriticalSection& OuterMutex);
 	void PreCancelIfDependsOn(const UObject* Object);
 	void Cancel();
@@ -1346,7 +1349,7 @@ private:
 
 class FPoseSearchDatabaseAsyncCacheTasks : public TArray<TUniquePtr<FPoseSearchDatabaseAsyncCacheTask>> {};
 
-FPoseSearchDatabaseAsyncCacheTask::FPoseSearchDatabaseAsyncCacheTask(UPoseSearchDatabase* InDatabase, bool bPerformConditionalPostLoadIfRequired)
+FPoseSearchDatabaseAsyncCacheTask::FPoseSearchDatabaseAsyncCacheTask(UPoseSearchDatabase* InDatabase, bool bPerformConditionalPostLoadIfRequired, FPartialKeyHashes& PartialKeyHashes)
 	: Database(InDatabase)
 	, Owner(UE::DerivedData::EPriority::Normal)
 	, DerivedDataKey(FIoHash::Zero)
@@ -1354,7 +1357,7 @@ FPoseSearchDatabaseAsyncCacheTask::FPoseSearchDatabaseAsyncCacheTask(UPoseSearch
 	if (IsInGameThread())
 	{
 		// it is safe to compose DDC key only on the game thread, since assets can modified in this thread execution
-		StartNewRequestIfNeeded(bPerformConditionalPostLoadIfRequired);
+		StartNewRequestIfNeeded(bPerformConditionalPostLoadIfRequired, PartialKeyHashes);
 	}
 	else
 	{
@@ -1393,7 +1396,7 @@ void FPoseSearchDatabaseAsyncCacheTask::TestSynchronizeWithExternalDependencies(
 }
 #endif //ENABLE_ANIM_DEBUG
 
-void FPoseSearchDatabaseAsyncCacheTask::StartNewRequestIfNeeded(bool bPerformConditionalPostLoadIfRequired)
+void FPoseSearchDatabaseAsyncCacheTask::StartNewRequestIfNeeded(bool bPerformConditionalPostLoadIfRequired, FPartialKeyHashes& PartialKeyHashes)
 {
 	using namespace UE::DerivedData;
 
@@ -1403,8 +1406,22 @@ void FPoseSearchDatabaseAsyncCacheTask::StartNewRequestIfNeeded(bool bPerformCon
 	// Owner.Cancel must be performed before SearchIndex.Reset() in case any task is flying (launched by Owner.LaunchTask)
 	Owner.Cancel();
 
+	FKeyBuilder::EDebugPartialKeyHashesMode DebugPartialKeyHashesMode;
+	switch (CVarMotionMatchPartialKeyHashesMode.GetValueOnAnyThread())
+	{
+	case 0:
+		DebugPartialKeyHashesMode = FKeyBuilder::EDebugPartialKeyHashesMode::Use;
+		break;
+	case 1:
+		DebugPartialKeyHashesMode = FKeyBuilder::EDebugPartialKeyHashesMode::DoNotUse;
+		break;
+	default:
+		DebugPartialKeyHashesMode = FKeyBuilder::EDebugPartialKeyHashesMode::Validate;
+		break;
+	}
+
 	// composing the key
-	const FKeyBuilder KeyBuilder(Database.Get(), true, bPerformConditionalPostLoadIfRequired);
+	const FKeyBuilder KeyBuilder(Database.Get(), true, bPerformConditionalPostLoadIfRequired, &PartialKeyHashes, DebugPartialKeyHashesMode);
 	if (KeyBuilder.AnyAssetNotReady())
 	{
 		DerivedDataKey = FIoHash::Zero;
@@ -1483,7 +1500,7 @@ void FPoseSearchDatabaseAsyncCacheTask::Cancel()
 	SetState(EState::Cancelled);
 }
 
-void FPoseSearchDatabaseAsyncCacheTask::Update(FCriticalSection& OuterMutex)
+void FPoseSearchDatabaseAsyncCacheTask::Update(FCriticalSection& OuterMutex, FPartialKeyHashes& PartialKeyHashes)
 {
 	check(IsInGameThread());
 
@@ -1491,7 +1508,7 @@ void FPoseSearchDatabaseAsyncCacheTask::Update(FCriticalSection& OuterMutex)
 
 	if (GetState() == EState::Notstarted)
 	{
-		StartNewRequestIfNeeded(false);
+		StartNewRequestIfNeeded(false, PartialKeyHashes);
 	}
 
 	if (GetState() == EState::Prestarted && Poll())
@@ -2076,6 +2093,8 @@ void FAsyncPoseSearchDatabasesManagement::PreModified(UObject* Object)
 
 	FScopeLock Lock(&Mutex);
 
+	PartialKeyHashes.Remove(Object);
+
 	// iterating backwards because of the possible RemoveAtSwap
 	for (int32 TaskIndex = Tasks.Num() - 1; TaskIndex >= 0; --TaskIndex)
 	{
@@ -2400,7 +2419,7 @@ void FAsyncPoseSearchDatabasesManagement::Tick(float DeltaTime)
 		}
 		else
 		{
-			Tasks[TaskIndex]->Update(Mutex);
+			Tasks[TaskIndex]->Update(Mutex, PartialKeyHashes);
 		}
 	}
 	
@@ -2417,9 +2436,14 @@ void FAsyncPoseSearchDatabasesManagement::Tick(float DeltaTime)
 
 				for (int32 IterationIndex = 0; IterationIndex < NumIterations; ++IterationIndex)
 				{
-					const FKeyBuilder TestKeyBuilder(Database, false, false);
+					const FKeyBuilder TestKeyBuilder(Database, false, false, &PartialKeyHashes, FKeyBuilder::EDebugPartialKeyHashesMode::Use);
 					const FIoHash TestIoHash = TestKeyBuilder.Finalize();
 
+					if (!KeyBuilder.ValidateAgainst(TestKeyBuilder))
+					{
+						UE_LOG(LogPoseSearch, Error, TEXT("FKeyBuilder - key generation is not deterministic: %s / %s for asset %s"), *LexToString(IoHash), *LexToString(TestIoHash), *Database->GetName());
+					}
+
 					if (IoHash != TestIoHash)
 					{
 						UE_LOG(LogPoseSearch, Error, TEXT("FKeyBuilder - key generation is not deterministic: %s / %s for asset %s"), *LexToString(IoHash), *LexToString(TestIoHash), *Database->GetName());
@@ -2479,7 +2503,7 @@ EAsyncBuildIndexResult FAsyncPoseSearchDatabasesManagement::RequestAsyncBuildInd
 				{
 					Task->Cancel();
 				}
-				Task->StartNewRequestIfNeeded(bWaitForCompletion);
+				Task->StartNewRequestIfNeeded(bWaitForCompletion, This.PartialKeyHashes);
 			}
 			break;
 		}
@@ -2488,7 +2512,7 @@ EAsyncBuildIndexResult FAsyncPoseSearchDatabasesManagement::RequestAsyncBuildInd
 	if (!Task)
 	{
 		// we didn't find the Task, so we Emplace a new one
-		This.Tasks.Emplace(MakeUnique<FPoseSearchDatabaseAsyncCacheTask>(const_cast<UPoseSearchDatabase*>(Database), bWaitForCompletion));
+		This.Tasks.Emplace(MakeUnique<FPoseSearchDatabaseAsyncCacheTask>(const_cast<UPoseSearchDatabase*>(Database), bWaitForCompletion, This.PartialKeyHashes));
 		Task = This.Tasks.Last().Get();
 	}
 
diff --git a/Engine/Plugins/Animation/PoseSearch/Source/Runtime/Private/PoseSearchDerivedDataKey.cpp b/Engine/Plugins/Animation/PoseSearch/Source/Runtime/Private/PoseSearchDerivedDataKey.cpp
--- a/Engine/Plugins/Animation/PoseSearch/Source/Runtime/Private/PoseSearchDerivedDataKey.cpp
+++ b/Engine/Plugins/Animation/PoseSearch/Source/Runtime/Private/PoseSearchDerivedDataKey.cpp
@@ -19,6 +19,16 @@
 namespace UE::PoseSearch
 {
 
+// log properties and UObjects names
+#ifndef UE_POSE_SEARCH_DERIVED_DATA_LOGGING
+	#define UE_POSE_SEARCH_DERIVED_DATA_LOGGING 0
+#endif
+
+// log properties data
+#ifndef UE_POSE_SEARCH_DERIVED_DATA_LOGGING_VERBOSE
+	#define UE_POSE_SEARCH_DERIVED_DATA_LOGGING_VERBOSE 0
+#endif
+
 FKeyBuilder::FKeyBuilder()
 {
 	ArIgnoreOuterRef = true;
@@ -30,11 +40,19 @@ FKeyBuilder::FKeyBuilder()
 	SetIsSaving(true);
 }
 
-FKeyBuilder::FKeyBuilder(const UObject* Object, bool bUseDataVer, bool bPerformConditionalPostLoadIfRequired)
+FKeyBuilder::FKeyBuilder(const UObject* Object, bool bUseDataVer, bool bPerformConditionalPostLoadIfRequired, FPartialKeyHashes* InPartialKeyHashes, EDebugPartialKeyHashesMode InDebugPartialKeyHashesMode)
 : FKeyBuilder()
 {
 	check(Object);
+
+	// preallocating a reasonable amount of memory to avoid multiple reallocations
+	ObjectsToSerialize.Reserve(256);
+	ObjectBeingSerializedDependencies.Reserve(256);
+	LocalPartialKeyHashes.Reserve(1024);
+
 	bPerformConditionalPostLoad = bPerformConditionalPostLoadIfRequired;
+	PartialKeyHashes = InPartialKeyHashes;
+	DebugPartialKeyHashesMode = InDebugPartialKeyHashesMode;
 
 	if (bUseDataVer)
 	{
@@ -51,6 +69,16 @@ FKeyBuilder::FKeyBuilder(const UObject* Object, bool bUseDataVer, bool bPerformC
 	// FKeyBuilder is a saving only archiver, and since it doesn't modify the input Object it's safe to do a const_cast 
 	UObject* NonConstObject = const_cast<UObject*>(Object);
 	*this << NonConstObject;
+
+	while (!ObjectsToSerialize.IsEmpty() && !bAnyAssetNotReady)
+	{
+		SerializeObjectInternal(ObjectsToSerialize.Pop(EAllowShrinking::No));
+	}
+}
+
+FKeyBuilder::FKeyBuilder(const UObject* Object, bool bUseDataVer, bool bPerformConditionalPostLoadIfRequired)
+	: FKeyBuilder(Object, bUseDataVer, bPerformConditionalPostLoadIfRequired, nullptr, FKeyBuilder::EDebugPartialKeyHashesMode::DoNotUse)
+{
 }
 
 void FKeyBuilder::Seek(int64 InPos)
@@ -69,7 +97,7 @@ bool FKeyBuilder::ShouldSkipProperty(const FProperty* InProperty) const
 	if (Super::ShouldSkipProperty(InProperty))
 	{
 		#if UE_POSE_SEARCH_DERIVED_DATA_LOGGING
-		UE_LOG(LogPoseSearch, Log, TEXT("%s x %s (ShouldSkipProperty)"), *GetIndentation(), *InProperty->GetFullName());
+		UE_LOG(LogPoseSearch, Log, TEXT("  x %s (ShouldSkipProperty)"), *InProperty->GetFullName());
 		#endif
 		return true;
 	}
@@ -77,7 +105,7 @@ bool FKeyBuilder::ShouldSkipProperty(const FProperty* InProperty) const
 	if (InProperty->HasAllPropertyFlags(CPF_Transient))
 	{
 		#if UE_POSE_SEARCH_DERIVED_DATA_LOGGING
-		UE_LOG(LogPoseSearch, Log, TEXT("%s x %s (Transient)"), *GetIndentation(), *InProperty->GetFullName());
+		UE_LOG(LogPoseSearch, Log, TEXT("  x %s (Transient)"), *InProperty->GetFullName());
 		#endif
 		return true;
 	}
@@ -85,7 +113,7 @@ bool FKeyBuilder::ShouldSkipProperty(const FProperty* InProperty) const
 	if (InProperty->HasMetaData(ExcludeFromHashName))
 	{
 		#if UE_POSE_SEARCH_DERIVED_DATA_LOGGING
-		UE_LOG(LogPoseSearch, Log, TEXT("%s x %s (ExcludeFromHash)"), *GetIndentation(), *InProperty->GetFullName());
+		UE_LOG(LogPoseSearch, Log, TEXT("  x %s (ExcludeFromHash)"), *InProperty->GetFullName());
 		#endif
 		return true;
 	}
@@ -93,7 +121,7 @@ bool FKeyBuilder::ShouldSkipProperty(const FProperty* InProperty) const
 	if (InProperty->HasMetaData(IgnoreForMemberInitializationTestName))
 	{
 		#if UE_POSE_SEARCH_DERIVED_DATA_LOGGING
-		UE_LOG(LogPoseSearch, Log, TEXT("%s x %s (IgnoreForMemberInitializationTest)"), *GetIndentation(), *InProperty->GetFullName());
+		UE_LOG(LogPoseSearch, Log, TEXT("  x %s (IgnoreForMemberInitializationTest)"), *InProperty->GetFullName());
 		#endif
 		return true;
 	}
@@ -101,7 +129,7 @@ bool FKeyBuilder::ShouldSkipProperty(const FProperty* InProperty) const
 	check(!InProperty->HasMetaData(NeverInHashName));
 
 	#if UE_POSE_SEARCH_DERIVED_DATA_LOGGING
-	UE_LOG(LogPoseSearch, Log, TEXT("%s - %s"), *GetIndentation(), *InProperty->GetFullName());
+	UE_LOG(LogPoseSearch, Log, TEXT("  - %s"), *InProperty->GetFullName());
 	#endif
 
 	return false;
@@ -113,7 +141,7 @@ void FKeyBuilder::Serialize(void* Data, int64 Length)
 
 	#if UE_POSE_SEARCH_DERIVED_DATA_LOGGING_VERBOSE
 	FString RawBytesString = BytesToString(HasherData, Length);
-	UE_LOG(LogPoseSearch, Log, TEXT("%s  > %s"), *GetIndentation(), *RawBytesString);
+	UE_LOG(LogPoseSearch, Log, TEXT("  > %s"), *RawBytesString);
 	#endif
 
 	Hasher.Update(HasherData, Length);
@@ -131,94 +159,161 @@ FArchive& FKeyBuilder::operator<<(FName& Name)
 	return *this;
 }
 
-FArchive& FKeyBuilder::operator<<(class UObject*& Object)
+FArchive& FKeyBuilder::TryAddDependency(UObject* Object, bool bAddToPartialKeyHashes)
 {
-	if (Object)
+	if (Object->HasAnyFlags(RF_NeedPostLoad))
 	{
-		if (Object->HasAnyFlags(RF_NeedPostLoad))
+		if (bPerformConditionalPostLoad)
 		{
-			if (bPerformConditionalPostLoad)
-			{
-				Object->ConditionalPostLoad();
-			}
-			else
-			{
-				bAnyAssetNotReady = true;
-			}
+			Object->ConditionalPostLoad();
+		}
+		else
+		{
+			bAnyAssetNotReady = true;
+			return *this;
 		}
+	}
 
-		// @todo: should we WaitOnExistingCompression?
-		//if (UAnimSequence* AnimSequence = Cast<UAnimSequence>(Object))
-		//{
-		//	AnimSequence->WaitOnExistingCompression();
-		//}
+	// @todo: should we WaitOnExistingCompression?
+	//if (UAnimSequence* AnimSequence = Cast<UAnimSequence>(Object))
+	//{
+	//	AnimSequence->WaitOnExistingCompression();
+	//}
 
-		if (!bAnyAssetNotReady)
-		{
-			#if UE_POSE_SEARCH_DERIVED_DATA_LOGGING
-			++Indentation;
-			#endif
+	// collecting ALL the dependencies of the object being serialized, so we can then cache it in PartialKeyHashes
+	ObjectBeingSerializedDependencies.Add(Object);
 
-			if (Object->HasAnyFlags(RF_Transient))
-			{
-				#if UE_POSE_SEARCH_DERIVED_DATA_LOGGING
-				UE_LOG(LogPoseSearch, Log, TEXT("%sTransient '%s' (%s)"), *GetIndentation(), *Object->GetName(), *Object->GetClass()->GetName());
-				#endif
-			}
-			else if (IsExcludedType(Object))
+	bool bAlreadyProcessed = false;
+	Dependencies.Add(Object, &bAlreadyProcessed);
+
+	// If we haven't already serialized this object
+	if (bAlreadyProcessed)
+	{
+#if UE_POSE_SEARCH_DERIVED_DATA_LOGGING
+		UE_LOG(LogPoseSearch, Log, TEXT("AlreadyProcessed '%s' (%s)"), *Object->GetName(), *Object->GetClass()->GetName());
+#endif
+		return *this;
+	}
+
+	ObjectsToSerialize.Add(Object);
+	return *this;
+}
+
+FArchive& FKeyBuilder::operator<<(class UObject*& Object)
+{
+	if (!Object)
+	{
+		return *this;
+	}
+
+	if (Object->HasAnyFlags(RF_Transient))
+	{
+#if UE_POSE_SEARCH_DERIVED_DATA_LOGGING
+		UE_LOG(LogPoseSearch, Log, TEXT("Transient '%s' (%s)"), *Object->GetName(), *Object->GetClass()->GetName());
+#endif
+		return *this;
+	}
+
+	if (IsExcludedType(Object))
+	{
+#if UE_POSE_SEARCH_DERIVED_DATA_LOGGING
+		UE_LOG(LogPoseSearch, Log, TEXT("Excluded '%s' (%s)"), *Object->GetName(), *Object->GetClass()->GetName());
+#endif
+		return *this;
+	}
+
+	return TryAddDependency(Object, true);
+}
+
+void FKeyBuilder::SerializeObjectInternal(UObject* Object)
+{
+	Hasher.Reset();
+
+	check(!bAnyAssetNotReady && Object);
+
+	// test to validate the PartialKeyHashes
+	bool bIsValidTestPartialKeyHash = false;
+	FPartialKeyHashes::FEntry TestEntry;
+
+	// adding the LocalCachedHash here to keep the order consistent with Dependencies
+	FLocalPartialKeyHash& LocalCachedHash = LocalPartialKeyHashes.AddDefaulted_GetRef();
+
+	if (PartialKeyHashes)
+	{
+		if (DebugPartialKeyHashesMode == EDebugPartialKeyHashesMode::Validate)
+		{
+			if (const FPartialKeyHashes::FEntry* Entry = PartialKeyHashes->Find(Object))
 			{
-				#if UE_POSE_SEARCH_DERIVED_DATA_LOGGING
-				UE_LOG(LogPoseSearch, Log, TEXT("%sExcluded '%s' (%s)"), *GetIndentation(), *Object->GetName(), *Object->GetClass()->GetName());
-				#endif
+				bIsValidTestPartialKeyHash = true;
+				TestEntry = *Entry;
 			}
-			else
+		}
+		else if (DebugPartialKeyHashesMode == EDebugPartialKeyHashesMode::Use)
+		{
+			if (const FPartialKeyHashes::FEntry* Entry = PartialKeyHashes->Find(Object))
 			{
-				bool bAlreadyProcessed = false;
-				Dependencies.Add(Object, &bAlreadyProcessed);
-
-				// If we haven't already serialized this object
-				if (bAlreadyProcessed)
+				for (const TWeakObjectPtr<>& DependencyPtr : Entry->Dependencies)
 				{
-					#if UE_POSE_SEARCH_DERIVED_DATA_LOGGING
-					UE_LOG(LogPoseSearch, Log, TEXT("%sAlreadyProcessed '%s' (%s)"), *GetIndentation(), *Object->GetName(), *Object->GetClass()->GetName());
-					#endif
+					if (UObject* Dependency = DependencyPtr.Get())
+					{
+						TryAddDependency(Dependency, false);
+					}
 				}
-				// for specific types we only add their names to the hash
-				else if (IsAddNameOnlyType(Object))
-				{
-					#if UE_POSE_SEARCH_DERIVED_DATA_LOGGING
-					UE_LOG(LogPoseSearch, Log, TEXT("%sAddingNameOnly '%s' (%s)"), *GetIndentation(), *Object->GetName(), *Object->GetClass()->GetName());
-					#endif
 
-					FString ObjectName = GetFullNameSafe(Object);
-					*this << ObjectName;
-				}
-				else
-				{
-					const UObject* PreviousObjectBeingSerialized = ObjectBeingSerialized;
-					ObjectBeingSerialized = Object;
+				LocalCachedHash.Object = Object;
+				LocalCachedHash.Hash = Entry->Hash;
+				return;
+			}
+		}
+	}
+
+	// making sure we don't call Object->Serialize recursively!
+	check(ObjectBeingSerialized == nullptr);
+	ObjectBeingSerialized = Object;
 
-					#if UE_POSE_SEARCH_DERIVED_DATA_LOGGING
-					UE_LOG(LogPoseSearch, Log, TEXT("%sBegin '%s' (%s)"), *GetIndentation(), *Object->GetName(), *Object->GetClass()->GetName());
-					#endif
+	ObjectBeingSerializedDependencies.Reset();
 
-					const_cast<UObject*>(Object)->Serialize(*this);
+	if (IsAddNameOnlyType(Object))
+	{
+		// for specific types we only add their names to the hash
+#if UE_POSE_SEARCH_DERIVED_DATA_LOGGING
+		UE_LOG(LogPoseSearch, Log, TEXT("AddingNameOnly '%s' (%s)"), *Object->GetName(), *Object->GetClass()->GetName());
+#endif
+		FString ObjectName = GetFullNameSafe(Object);
+		*this << ObjectName;
+	}
+	else
+	{
+#if UE_POSE_SEARCH_DERIVED_DATA_LOGGING
+		UE_LOG(LogPoseSearch, Log, TEXT("Begin '%s' (%s)"), *Object->GetName(), *Object->GetClass()->GetName());
+#endif
+		Object->Serialize(*this);
 
-					#if UE_POSE_SEARCH_DERIVED_DATA_LOGGING
-					UE_LOG(LogPoseSearch, Log, TEXT("%sEnd '%s' (%s)"), *GetIndentation(), *Object->GetName(), *Object->GetClass()->GetName());
-					#endif
+#if UE_POSE_SEARCH_DERIVED_DATA_LOGGING
+		UE_LOG(LogPoseSearch, Log, TEXT("End '%s' (%s)"), *Object->GetName(), *Object->GetClass()->GetName());
+#endif
+	}
 
-					ObjectBeingSerialized = PreviousObjectBeingSerialized;
-				}
-			}
+	// making sure we don't call Object->Serialize recursively!
+	check(ObjectBeingSerialized == Object);
+	ObjectBeingSerialized = nullptr;
 
-			#if UE_POSE_SEARCH_DERIVED_DATA_LOGGING
-			--Indentation;
-			#endif
+	if (!bAnyAssetNotReady)
+	{
+		LocalCachedHash.Object = Object;
+		LocalCachedHash.Hash = Hasher.Finalize();
+
+		if (DebugPartialKeyHashesMode != EDebugPartialKeyHashesMode::DoNotUse && PartialKeyHashes)
+		{
+			PartialKeyHashes->Add(LocalCachedHash.Object, LocalCachedHash.Hash, ObjectBeingSerializedDependencies);
+
+			if (bIsValidTestPartialKeyHash)
+			{
+				check(TestEntry.CheckDependencies(ObjectBeingSerializedDependencies));
+				check(LocalCachedHash.Hash == TestEntry.Hash);
+			}
 		}
 	}
-
-	return *this;
 }
 
 FString FKeyBuilder::GetArchiveName() const
@@ -234,8 +329,17 @@ bool FKeyBuilder::AnyAssetNotReady() const
 FIoHash FKeyBuilder::Finalize() const
 {
 	check(!bAnyAssetNotReady); // otherwise key can be non deterministic
+
+	HashBuilderType FinalizeHasher;
+
+	for (const FLocalPartialKeyHash& LocalCachedHash : LocalPartialKeyHashes)
+	{
+		const HashDigestType::ByteArray& LocalHashData = LocalCachedHash.Hash.GetBytes();
+		FinalizeHasher.Update(LocalHashData, sizeof(HashDigestType::ByteArray));
+	}
+
 	// Stores a BLAKE3-160 hash, taken from the first 20 bytes of a BLAKE3-256 hash
-	return FIoHash(Hasher.Finalize());
+	return FIoHash(FinalizeHasher.Finalize());
 }
 
 const TSet<const UObject*>& FKeyBuilder::GetDependencies() const
@@ -264,18 +368,60 @@ bool FKeyBuilder::IsAddNameOnlyType(class UObject* Object)
 		nullptr != Cast<UStreamableRenderAsset>(Object);
 }
 
-#if UE_POSE_SEARCH_DERIVED_DATA_LOGGING
-FString FKeyBuilder::GetIndentation() const
+bool FKeyBuilder::ValidateAgainst(const FKeyBuilder& Other) const
 {
-	FString IndentationString;
-	for (int32 i = 0; i < Indentation; ++i)
+	if (bAnyAssetNotReady != Other.bAnyAssetNotReady)
 	{
-		IndentationString.Append(" ");
+		return false;
 	}
-	return IndentationString;
+
+	if (Dependencies.Num() != Other.Dependencies.Num())
+	{
+		return false;
+	}
+
+	for (TSet<const UObject*>::TConstIterator Iter = Dependencies.CreateConstIterator(); Iter; ++Iter)
+	{
+		if (!Other.Dependencies.Contains(*Iter))
+		{
+			return false;
+		}
+	}
+
+	if (ObjectsToSerialize.Num() != Other.ObjectsToSerialize.Num())
+	{
+		return false;
+	}
+	
+	for (int32 Index = 0; Index < ObjectsToSerialize.Num(); ++Index)
+	{
+		if (ObjectsToSerialize[Index] != Other.ObjectsToSerialize[Index])
+		{
+			return false;
+		}
+	}
+
+	if (LocalPartialKeyHashes.Num() != Other.LocalPartialKeyHashes.Num())
+	{
+		return false;
+	}
+	
+	for (int32 Index = 0; Index < LocalPartialKeyHashes.Num(); ++Index)
+	{
+		if (LocalPartialKeyHashes[Index].Hash != Other.LocalPartialKeyHashes[Index].Hash)
+		{
+			return false;
+		}
+
+		if (LocalPartialKeyHashes[Index].Object != Other.LocalPartialKeyHashes[Index].Object)
+		{
+			return false;
+		}
+	}
+
+	return true;
 }
-#endif
 
-} // namespace UE::PoseSearch
+} // namespace UE::PoseSearch	
 
 #endif // WITH_EDITOR
\ No newline at end of file
diff --git a/Engine/Plugins/Animation/PoseSearch/Source/Runtime/Public/PoseSearch/PoseSearchDerivedData.h b/Engine/Plugins/Animation/PoseSearch/Source/Runtime/Public/PoseSearch/PoseSearchDerivedData.h
--- a/Engine/Plugins/Animation/PoseSearch/Source/Runtime/Public/PoseSearch/PoseSearchDerivedData.h
+++ b/Engine/Plugins/Animation/PoseSearch/Source/Runtime/Public/PoseSearch/PoseSearchDerivedData.h
@@ -3,6 +3,7 @@
 #pragma once
 
 #if WITH_EDITOR
+#include "PoseSearch/PoseSearchDerivedDataKey.h"
 #include "TickableEditorObject.h"
 #include "UObject/GCObject.h"
 #include "UObject/ObjectSaveContext.h"
@@ -86,6 +87,9 @@ namespace UE::PoseSearch
 		FDelegateHandle OnPreObjectPropertyChangedHandle;
 		FDelegateHandle OnObjectPropertyChangedHandle;
 
+		// Experimental, this feature might be removed without warning, not for production use
+		FPartialKeyHashes PartialKeyHashes;
+
 		static FCriticalSection Mutex;
 	};
 } // namespace UE::PoseSearch
diff --git a/Engine/Plugins/Animation/PoseSearch/Source/Runtime/Public/PoseSearch/PoseSearchDerivedDataKey.h b/Engine/Plugins/Animation/PoseSearch/Source/Runtime/Public/PoseSearch/PoseSearchDerivedDataKey.h
--- a/Engine/Plugins/Animation/PoseSearch/Source/Runtime/Public/PoseSearch/PoseSearchDerivedDataKey.h
+++ b/Engine/Plugins/Animation/PoseSearch/Source/Runtime/Public/PoseSearch/PoseSearchDerivedDataKey.h
@@ -12,15 +12,7 @@
 namespace UE::PoseSearch
 {
 
-// log properties and UObjects names
-#ifndef UE_POSE_SEARCH_DERIVED_DATA_LOGGING
-	#define UE_POSE_SEARCH_DERIVED_DATA_LOGGING 0
-#endif
-
-// log properties data
-#ifndef UE_POSE_SEARCH_DERIVED_DATA_LOGGING_VERBOSE
-	#define UE_POSE_SEARCH_DERIVED_DATA_LOGGING_VERBOSE 0
-#endif
+class FPartialKeyHashes;
 
 class POSESEARCH_API FKeyBuilder : public FArchiveUObject
 {
@@ -35,6 +27,20 @@ public:
 	
 	FKeyBuilder();
 	FKeyBuilder(const UObject* Object, bool bUseDataVer, bool bPerformConditionalPostLoadIfRequired);
+	
+	// Experimental, this feature might be removed without warning, not for production use
+	enum EDebugPartialKeyHashesMode
+	{
+		Use,
+		DoNotUse,
+		Validate
+	};
+
+	// Experimental, this feature might be removed without warning, not for production use
+	FKeyBuilder(const UObject* Object, bool bUseDataVer, bool bPerformConditionalPostLoadIfRequired, FPartialKeyHashes* InPartialKeyHashes, EDebugPartialKeyHashesMode InDebugPartialKeyHashesMode);
+
+	// Experimental, this feature might be removed without warning, not for production use
+	bool ValidateAgainst(const FKeyBuilder& Other) const;
 
 	using Super::IsSaving;
 	using Super::operator<<;
@@ -59,24 +65,156 @@ protected:
 	// to keep the key generation lightweight, we hash only the full names for these types. Object(s) will be added to Dependencies
 	static bool IsAddNameOnlyType(class UObject* Object);
 
-#if UE_POSE_SEARCH_DERIVED_DATA_LOGGING
-	FString GetIndentation() const;
-	int32 Indentation = 0;
-#endif
-
 	HashBuilderType Hasher;
 
 	// Set of objects that have already been serialized
 	TSet<const UObject*> Dependencies;
 
 	// Object currently being serialized
-	const UObject* ObjectBeingSerialized = nullptr;
+	UObject* ObjectBeingSerialized = nullptr;
 
 	// true if some dependent assets are not ready (fully loaded)
 	bool bAnyAssetNotReady = false;
 
 	// if true ConditionalPostLoad will be performed on the dependant assets requiring it
 	bool bPerformConditionalPostLoad = false;
+
+private:
+	// Experimental, this feature might be removed without warning, not for production use
+	void SerializeObjectInternal(UObject* Object);
+
+	// Experimental, this feature might be removed without warning, not for production use
+	FArchive& TryAddDependency(UObject* Object, bool bAddToPartialKeyHashes);
+
+	// Experimental, this feature might be removed without warning, not for production use
+	TArray<UObject*> ObjectsToSerialize;
+	
+	// Experimental, this feature might be removed without warning, not for production use
+	TArray<UObject*> ObjectBeingSerializedDependencies;
+
+	// Experimental, this feature might be removed without warning, not for production use
+	struct FLocalPartialKeyHash
+	{
+		UObject* Object = nullptr;
+		HashDigestType Hash;
+	};
+	
+	// Experimental, this feature might be removed without warning, not for production use
+	TArray<FLocalPartialKeyHash> LocalPartialKeyHashes;
+	
+	// Experimental, this feature might be removed without warning, not for production use
+	FPartialKeyHashes* PartialKeyHashes = nullptr;
+
+	// Experimental, this feature might be removed without warning, not for production use
+	EDebugPartialKeyHashesMode DebugPartialKeyHashesMode = EDebugPartialKeyHashesMode::Use;
+};
+
+// Experimental, this feature might be removed without warning, not for production use
+class FPartialKeyHashes
+{
+public:
+	struct FEntry
+	{
+		FKeyBuilder::HashDigestType Hash;
+		TArray<TWeakObjectPtr<>> Dependencies;
+
+		bool CheckDependencies(TArrayView<UObject*> OtherDependencies) const
+		{
+			#if DO_CHECK
+			if (Dependencies.Num() != OtherDependencies.Num())
+			{
+				return false;
+			}
+			
+			for (int32 DependencyIndex = 0; DependencyIndex < Dependencies.Num(); ++DependencyIndex)
+			{
+				if (!OtherDependencies[DependencyIndex])
+				{
+					return false;
+				}
+
+				// we could have lost a weak pointer here...
+				if (Dependencies[DependencyIndex].IsValid())
+				{
+					if (Dependencies[DependencyIndex].Get() != OtherDependencies[DependencyIndex])
+					{
+						return false;
+					}
+				}
+			}
+			#endif // DO_CHECK
+
+			return true;
+		}
+
+	};
+
+	void Reset()
+	{
+		Entries.Reset();
+	}
+
+	void Remove(UObject* Object)
+	{
+		Entries.Remove(Object);
+	}
+
+	void Add(UObject* Object, const FKeyBuilder::HashDigestType& Hash, TArrayView<UObject*> Dependencies)
+	{
+		check(Object);
+
+		check(!Hash.IsZero());
+
+		if (FEntry* OldEntry = Entries.Find(Object))
+		{
+			check(OldEntry->Hash == Hash);
+			check(OldEntry->CheckDependencies(Dependencies))
+		}
+		else
+		{
+			FEntry& NewEntry = Entries.Add(Object);
+			NewEntry.Hash = Hash;
+			NewEntry.Dependencies = Dependencies;
+		}
+	}
+
+	const FEntry* Find(UObject* Object)
+	{
+		check(Object);
+		if (const TPair<TWeakObjectPtr<>, FEntry>* Pair = Entries.FindPair(Object))
+		{
+			// making sure all the TWeakObjectPtr are still valid
+			if (!Pair->Key.IsValid())
+			{
+				Entries.Remove(Object);
+				return nullptr;
+			}
+			
+			for (const TWeakObjectPtr<> Dependency : Pair->Value.Dependencies)
+			{
+				if (!Dependency.IsValid())
+				{
+					Entries.Remove(Object);
+					return nullptr;
+				}
+			}
+
+			return &Pair->Value;
+		}
+
+		return nullptr;
+	}
+
+private:
+	struct FEntries : public TMap<TWeakObjectPtr<>, FEntry>
+	{
+		const TPair<TWeakObjectPtr<>, FEntry>* FindPair(TWeakObjectPtr<> Key) const
+		{
+			return Pairs.Find(Key);
+		}
+	};
+
+	FEntries Entries;
 };
 
 } // namespace UE::PoseSearch
