pose search - set DisplayName for UAnimNotifyState_PoseSearchBlockTransition as Pose Search: Block Transition In

diff --git a/Engine/Plugins/Animation/PoseSearch/Source/Runtime/Public/PoseSearch/PoseSearchAnimNotifies.h b/Engine/Plugins/Animation/PoseSearch/Source/Runtime/Public/PoseSearch/PoseSearchAnimNotifies.h
--- a/Engine/Plugins/Animation/PoseSearch/Source/Runtime/Public/PoseSearch/PoseSearchAnimNotifies.h
+++ b/Engine/Plugins/Animation/PoseSearch/Source/Runtime/Public/PoseSearch/PoseSearchAnimNotifies.h
@@ -25,7 +25,7 @@ class POSESEARCH_API UAnimNotifyState_PoseSearchExcludeFromDatabase : public UAn
 
 // A pose search search will not return results that overlap this notify, but the animation segment can still play
 // if a previous search result advances into it.
-UCLASS(Blueprintable, meta = (DisplayName = "Pose Search: Block Transition"))
+UCLASS(Blueprintable, meta = (DisplayName = "Pose Search: Block Transition In"))
 class POSESEARCH_API UAnimNotifyState_PoseSearchBlockTransition : public UAnimNotifyState_PoseSearchBase
 {
 	GENERATED_BODY()
