pose search - adding BlueprintThreadSafe UPoseSearchLibrary::GetDatabaseTags

diff --git a/Engine/Plugins/Animation/PoseSearch/Source/Runtime/Private/PoseSearchLibrary.cpp b/Engine/Plugins/Animation/PoseSearch/Source/Runtime/Private/PoseSearchLibrary.cpp
--- a/Engine/Plugins/Animation/PoseSearch/Source/Runtime/Private/PoseSearchLibrary.cpp
+++ b/Engine/Plugins/Animation/PoseSearch/Source/Runtime/Private/PoseSearchLibrary.cpp
@@ -714,6 +714,18 @@ void UPoseSearchLibrary::IsAnimationAssetLooping(const UObject* Asset, bool& bIs
 	}
 }
 
+void UPoseSearchLibrary::GetDatabaseTags(const UPoseSearchDatabase* Database, TArray<FName>& Tags)
+{
+	if (Database)
+	{
+		Tags = Database->Tags;
+	}
+	else
+	{
+		Tags.Reset();
+	}
+}
+
 void UPoseSearchLibrary::MotionMatch(
 	UAnimInstance* AnimInstance,
 	TArray<UObject*> AssetsToSearch,
diff --git a/Engine/Plugins/Animation/PoseSearch/Source/Runtime/Public/PoseSearch/PoseSearchLibrary.h b/Engine/Plugins/Animation/PoseSearch/Source/Runtime/Public/PoseSearch/PoseSearchLibrary.h
--- a/Engine/Plugins/Animation/PoseSearch/Source/Runtime/Public/PoseSearch/PoseSearchLibrary.h
+++ b/Engine/Plugins/Animation/PoseSearch/Source/Runtime/Public/PoseSearch/PoseSearchLibrary.h
@@ -288,5 +288,8 @@ public:
 
 	UFUNCTION(BlueprintPure, Category = "Animation|Pose Search|Experimental", meta = (BlueprintThreadSafe))
 	static void IsAnimationAssetLooping(const UObject* Asset, bool& bIsAssetLooping);
+
+	UFUNCTION(BlueprintPure, Category = "Animation|Pose Search|Experimental", meta = (BlueprintThreadSafe))
+	static void GetDatabaseTags(const UPoseSearchDatabase* Database, TArray<FName>& Tags);
 };
 
