pose search - pose search debugger support for AnimNext

diff --git a/Engine/Plugins/Animation/PoseSearch/Source/Editor/Private/PoseSearchDatabaseViewModel.cpp b/Engine/Plugins/Animation/PoseSearch/Source/Editor/Private/PoseSearchDatabaseViewModel.cpp
--- a/Engine/Plugins/Animation/PoseSearch/Source/Editor/Private/PoseSearchDatabaseViewModel.cpp
+++ b/Engine/Plugins/Animation/PoseSearch/Source/Editor/Private/PoseSearchDatabaseViewModel.cpp
@@ -208,7 +208,7 @@ bool FDatabasePreviewActor::DrawPreviewActors(TConstArrayView<FDatabasePreviewAc
 	int32 CommonIndexAssetIndex = INDEX_NONE;
 #endif // DO_CHECK
 
-	TArray<const USkinnedMeshComponent*, TInlineAllocator<PreallocatedRolesNum>> Meshes;
+	TArray<const UObject*, TInlineAllocator<PreallocatedRolesNum>> AnimContexts;
 	FRoleToIndex RoleToIndex;
 	TArray<FArchivedPoseHistory, TInlineAllocator<PreallocatedRolesNum>> ArchivedPoseHistories;
 	TArray<const IPoseHistory*, TInlineAllocator<PreallocatedRolesNum>> PoseHistories;
@@ -217,7 +217,7 @@ bool FDatabasePreviewActor::DrawPreviewActors(TConstArrayView<FDatabasePreviewAc
 	if (NumPreviewActors > PreallocatedRolesNum)
 	{
 		// reserve the needed amount of memory for containers
-		Meshes.Reserve(NumPreviewActors);
+		AnimContexts.Reserve(NumPreviewActors);
 		RoleToIndex.Reserve(NumPreviewActors);
 		ArchivedPoseHistories.Reserve(NumPreviewActors);
 		PoseHistories.Reserve(NumPreviewActors);
@@ -274,8 +274,8 @@ bool FDatabasePreviewActor::DrawPreviewActors(TConstArrayView<FDatabasePreviewAc
 		}
 #endif // DO_CHECK
 
-		RoleToIndex.Add(PreviewActor.ActorRole) = Meshes.Num();
-		Meshes.Add(Mesh);
+		RoleToIndex.Add(PreviewActor.ActorRole) = AnimContexts.Num();
+		AnimContexts.Add(Mesh);
 
 		FArchivedPoseHistory& ArchivedPoseHistory = ArchivedPoseHistories.AddDefaulted_GetRef();
 		ArchivedPoseHistory.Trajectory = PreviewActor.Trajectory;
@@ -303,7 +303,7 @@ bool FDatabasePreviewActor::DrawPreviewActors(TConstArrayView<FDatabasePreviewAc
 		PoseHistories.Add(&ArchivedPoseHistory);
 	}
 
-	UE::PoseSearch::FDebugDrawParams DrawParams(Meshes, PoseHistories, RoleToIndex, PoseSearchDatabase);
+	UE::PoseSearch::FDebugDrawParams DrawParams(AnimContexts, PoseHistories, RoleToIndex, PoseSearchDatabase);
 	DrawParams.DrawFeatureVector(CommonCurrentPoseIndex);
 
 	if (!QueryVector.IsEmpty())
diff --git a/Engine/Plugins/Animation/PoseSearch/Source/Editor/Private/PoseSearchDebugger.h b/Engine/Plugins/Animation/PoseSearch/Source/Editor/Private/PoseSearchDebugger.h
--- a/Engine/Plugins/Animation/PoseSearch/Source/Editor/Private/PoseSearchDebugger.h
+++ b/Engine/Plugins/Animation/PoseSearch/Source/Editor/Private/PoseSearchDebugger.h
@@ -120,4 +120,10 @@ private:
 	virtual int32 GetSortOrderPriorityInternal() const override { return 10; };
 };
 
+class FAnimNextDebuggerTrackCreator : public FDebuggerTrackCreator
+{
+private:
+	virtual FName GetTargetTypeNameInternal() const override { return "AnimNextComponent"; }
+};
+
 } // namespace UE::PoseSearch
diff --git a/Engine/Plugins/Animation/PoseSearch/Source/Editor/Private/PoseSearchDebuggerView.cpp b/Engine/Plugins/Animation/PoseSearch/Source/Editor/Private/PoseSearchDebuggerView.cpp
--- a/Engine/Plugins/Animation/PoseSearch/Source/Editor/Private/PoseSearchDebuggerView.cpp
+++ b/Engine/Plugins/Animation/PoseSearch/Source/Editor/Private/PoseSearchDebuggerView.cpp
@@ -225,13 +225,14 @@ void SDebuggerView::Tick(const FGeometry& AllottedGeometry, const double InCurre
 	if (const FTraceMotionMatchingStateMessage* State = Model->GetMotionMatchingState())
 	{
 		FRoleToIndex RoleToIndex;
-		TArray<const USkinnedMeshComponent*> Meshes;
+		TArray<const UObject*> AnimContexts;
 		TArray<const IPoseHistory*> PoseHistories;
+		UWorld* World = nullptr;
 
 		const int32 NumRoles = State->Roles.Num();
 
 		RoleToIndex.Reserve(NumRoles);
-		Meshes.SetNum(NumRoles);
+		AnimContexts.SetNum(NumRoles);
 		PoseHistories.SetNum(NumRoles);
 
 		for (int32 RoleIndex = 0; RoleIndex < NumRoles; ++RoleIndex)
@@ -247,8 +248,9 @@ void SDebuggerView::Tick(const FGeometry& AllottedGeometry, const double InCurre
 					{
 						if (UPoseSearchMeshComponent* PoseSearchMeshComponent = Cast<UPoseSearchMeshComponent>(ActorComponent))
 						{
+							World = PoseSearchMeshComponent->GetWorld();
 							RoleToIndex.Add(Role) = RoleIndex;
-							Meshes[RoleIndex] = PoseSearchMeshComponent;
+							AnimContexts[RoleIndex] = PoseSearchMeshComponent;
 							PoseHistories[RoleIndex] = &State->PoseHistories[RoleIndex];
 							break;
 						}
@@ -272,7 +274,6 @@ void SDebuggerView::Tick(const FGeometry& AllottedGeometry, const double InCurre
 		{
 			for (int32 RoleIndex = 0; RoleIndex < NumRoles; ++RoleIndex)
 			{
-				UWorld* World = Meshes[RoleIndex]->GetWorld();
 				const FPoseSearchQueryTrajectory& Trajectory = State->PoseHistories[RoleIndex].Trajectory;
 				
 				if (bDrawTrajectory)
@@ -300,7 +301,7 @@ void SDebuggerView::Tick(const FGeometry& AllottedGeometry, const double InCurre
 						EAsyncBuildIndexResult::Success == FAsyncPoseSearchDatabasesManagement::RequestAsyncBuildIndex(CurrentDatabase, ERequestAsyncBuildFlag::ContinueRequest) &&
 						DbEntry.QueryVector.Num() == Database->Schema->SchemaCardinality)
 					{
-						FDebugDrawParams DrawParams(Meshes, PoseHistories, RoleToIndex, CurrentDatabase, EDebugDrawFlags::DrawQuery);
+						FDebugDrawParams DrawParams(AnimContexts, PoseHistories, RoleToIndex, CurrentDatabase, EDebugDrawFlags::DrawQuery);
 						DrawParams.DrawFeatureVector(DbEntry.QueryVector);
 						break;
 					}
@@ -322,7 +323,7 @@ void SDebuggerView::Tick(const FGeometry& AllottedGeometry, const double InCurre
 			const UPoseSearchDatabase* RowDatabase = Row->SharedData->SourceDatabase.Get();
 			if (EAsyncBuildIndexResult::Success == FAsyncPoseSearchDatabasesManagement::RequestAsyncBuildIndex(RowDatabase, ERequestAsyncBuildFlag::ContinueRequest))
 			{
-				FDebugDrawParams DrawParams(Meshes, PoseHistories, RoleToIndex, RowDatabase);
+				FDebugDrawParams DrawParams(AnimContexts, PoseHistories, RoleToIndex, RowDatabase);
 				DrawParams.DrawFeatureVector(Row->PoseIdx);
 			}
 		}
@@ -339,7 +340,7 @@ void SDebuggerView::Tick(const FGeometry& AllottedGeometry, const double InCurre
 			if (EAsyncBuildIndexResult::Success == FAsyncPoseSearchDatabasesManagement::RequestAsyncBuildIndex(Database, ERequestAsyncBuildFlag::ContinueRequest))
 			{
 				// Use the motion-matching state's pose idx, as the active row may be update-throttled at this point
-				FDebugDrawParams DrawParams(Meshes, PoseHistories, RoleToIndex, Database);
+				FDebugDrawParams DrawParams(AnimContexts, PoseHistories, RoleToIndex, Database);
 				DrawParams.DrawFeatureVector(ActiveRows[0]->PoseIdx);
 			}
 		}
@@ -356,7 +357,7 @@ void SDebuggerView::Tick(const FGeometry& AllottedGeometry, const double InCurre
 			const UPoseSearchDatabase* Database = ContinuingRows[0]->SharedData->SourceDatabase.Get();
 			if (EAsyncBuildIndexResult::Success == FAsyncPoseSearchDatabasesManagement::RequestAsyncBuildIndex(Database, ERequestAsyncBuildFlag::ContinueRequest))
 			{
-				FDebugDrawParams DrawParams(Meshes, PoseHistories, RoleToIndex, Database);
+				FDebugDrawParams DrawParams(AnimContexts, PoseHistories, RoleToIndex, Database);
 				DrawParams.DrawFeatureVector(ContinuingRows[0]->PoseIdx);
 			}
 		}
diff --git a/Engine/Plugins/Animation/PoseSearch/Source/Editor/Private/PoseSearchDebuggerView.h b/Engine/Plugins/Animation/PoseSearch/Source/Editor/Private/PoseSearchDebuggerView.h
--- a/Engine/Plugins/Animation/PoseSearch/Source/Editor/Private/PoseSearchDebuggerView.h
+++ b/Engine/Plugins/Animation/PoseSearch/Source/Editor/Private/PoseSearchDebuggerView.h
@@ -145,6 +145,7 @@ private:
 	/** Contains the switcher, the entire debugger view */
 	TSharedPtr<SVerticalBox> DebuggerView;
 
+	// @todo: rename it to AnimContextId
 	/** AnimInstance this view was created for */
 	uint64 AnimInstanceId = 0;
 
diff --git a/Engine/Plugins/Animation/PoseSearch/Source/Editor/Private/PoseSearchDebuggerViewModel.h b/Engine/Plugins/Animation/PoseSearch/Source/Editor/Private/PoseSearchDebuggerViewModel.h
--- a/Engine/Plugins/Animation/PoseSearch/Source/Editor/Private/PoseSearchDebuggerViewModel.h
+++ b/Engine/Plugins/Animation/PoseSearch/Source/Editor/Private/PoseSearchDebuggerViewModel.h
@@ -61,6 +61,7 @@ private:
 	/** Pointer to the active rewind debugger in the scene */
 	TAttribute<const IRewindDebugger*> RewindDebugger;
 
+	// @todo: rename it to AnimContextId
 	/** Anim Instance associated with this debugger instance */
 	uint64 AnimInstanceId = 0;
 
diff --git a/Engine/Plugins/Animation/PoseSearch/Source/Editor/Private/PoseSearchEditor.cpp b/Engine/Plugins/Animation/PoseSearch/Source/Editor/Private/PoseSearchEditor.cpp
--- a/Engine/Plugins/Animation/PoseSearch/Source/Editor/Private/PoseSearchEditor.cpp
+++ b/Engine/Plugins/Animation/PoseSearch/Source/Editor/Private/PoseSearchEditor.cpp
@@ -40,6 +40,8 @@ private:
 	
 	/** Creates the track for the Rewind Debugger */
 	TSharedPtr<FDebuggerTrackCreator> DebuggerTrackCreator;
+	TSharedPtr<FAnimNextDebuggerTrackCreator> AnimNextDebuggerTrackCreator;
+
 	/** Enables dedicated PoseSearch trace module */
 	TSharedPtr<FTraceModule> TraceModule;
 	
@@ -65,8 +67,10 @@ void FEditorModule::StartupModule()
 		FDebugger::Initialize();
 		TraceModule = MakeShared<FTraceModule>();
 		DebuggerTrackCreator = MakeShared<FDebuggerTrackCreator>();
-
+		AnimNextDebuggerTrackCreator = MakeShared<FAnimNextDebuggerTrackCreator>();
+		
 		IModularFeatures::Get().RegisterModularFeature(FDebuggerTrackCreator::ModularFeatureName, DebuggerTrackCreator.Get());
+		IModularFeatures::Get().RegisterModularFeature(FDebuggerTrackCreator::ModularFeatureName, AnimNextDebuggerTrackCreator.Get());
 		IModularFeatures::Get().RegisterModularFeature(TraceServices::ModuleFeatureName, TraceModule.Get());
 		
 		// Register Ed Modes used by PoseSearchDatabase and PoseSearchInteractionAsset
diff --git a/Engine/Plugins/Animation/PoseSearch/Source/Runtime/Private/PoseSearchContext.cpp b/Engine/Plugins/Animation/PoseSearch/Source/Runtime/Private/PoseSearchContext.cpp
--- a/Engine/Plugins/Animation/PoseSearch/Source/Runtime/Private/PoseSearchContext.cpp
+++ b/Engine/Plugins/Animation/PoseSearch/Source/Runtime/Private/PoseSearchContext.cpp
@@ -3,6 +3,7 @@
 #include "PoseSearch/PoseSearchContext.h"
 #include "Animation/AnimInstanceProxy.h"
 #include "AnimationRuntime.h"
+#include "Engine/SkeletalMesh.h"
 #include "PoseSearch/PoseSearchDatabase.h"
 #include "PoseSearch/PoseSearchHistory.h"
 #include "PoseSearch/PoseSearchSchema.h"
@@ -15,35 +16,82 @@ namespace UE::PoseSearch
 {
 	
 #if ENABLE_DRAW_DEBUG
+
+class UAnimInstanceProxyProvider : public UAnimInstance
+{
+public:
+	static FAnimInstanceProxy* GetAnimInstanceProxy(const UAnimInstance* AnimInstance)
+	{
+		if (AnimInstance)
+		{
+			// const cast is ok, since UAnimInstance::AnimInstanceProxy is mutable
+			return &static_cast<UAnimInstanceProxyProvider*>(const_cast<UAnimInstance*>(AnimInstance))->GetProxyOnAnyThread<FAnimInstanceProxy>();
+		}
+		return nullptr;
+	}
+};
+
+static FAnimInstanceProxy* GetAnimInstanceProxy(TConstArrayView<const UObject*> AnimContexts)
+{
+	if (!AnimContexts.IsEmpty())
+	{
+		return UAnimInstanceProxyProvider::GetAnimInstanceProxy(Cast<UAnimInstance>(AnimContexts[0]));
+	}
+	return nullptr;
+}
+
+static const USkinnedMeshComponent* GetMesh(TConstArrayView<const UObject*> AnimContexts, int32 RoleIndex = 0)
+{
+	if (AnimContexts.IsValidIndex(RoleIndex))
+	{
+		return Cast<USkinnedMeshComponent>(AnimContexts[RoleIndex]);
+	}
+	return nullptr;
+}
+
+static const UWorld* GetWorld(TConstArrayView<const UObject*> AnimContexts)
+{
+	if (const USkinnedMeshComponent* SkinnedMeshComponent = GetMesh(AnimContexts))
+	{
+		return SkinnedMeshComponent->GetWorld();
+	}
+	return nullptr;
+}
+
 //////////////////////////////////////////////////////////////////////////
 // FDebugDrawParams
 FDebugDrawParams::FDebugDrawParams(TArrayView<FAnimInstanceProxy*> InAnimInstanceProxies, TConstArrayView<const IPoseHistory*> InPoseHistories, const FRoleToIndex& InRoleToIndex, const UPoseSearchDatabase* InDatabase, EDebugDrawFlags InFlags)
-: AnimInstanceProxies(InAnimInstanceProxies)
-, Meshes()
+: AnimContexts()
 , PoseHistories(InPoseHistories)
 , RoleToIndex(InRoleToIndex)
 , Database(InDatabase)
 , Flags(InFlags)
 {
-	check(RoleToIndex.Num() == InAnimInstanceProxies.Num() && RoleToIndex.Num() == PoseHistories.Num());
-	check(IsValid(RoleToIndex));
 }
 
 FDebugDrawParams::FDebugDrawParams(TArrayView<const USkinnedMeshComponent*> InMeshes, TConstArrayView<const IPoseHistory*> InPoseHistories, const FRoleToIndex& InRoleToIndex, const UPoseSearchDatabase* InDatabase, EDebugDrawFlags InFlags)
-: AnimInstanceProxies()
-, Meshes(InMeshes)
+: AnimContexts()
 , PoseHistories(InPoseHistories)
 , RoleToIndex(InRoleToIndex)
 , Database(InDatabase)
 , Flags(InFlags)
 {
-	check(RoleToIndex.Num() == Meshes.Num() && RoleToIndex.Num() == PoseHistories.Num());
+}
+
+FDebugDrawParams::FDebugDrawParams(TConstArrayView<const UObject*> InAnimContexts, TConstArrayView<const IPoseHistory*> InPoseHistories, const FRoleToIndex& InRoleToIndex, const UPoseSearchDatabase* InDatabase, EDebugDrawFlags InFlags)
+: AnimContexts(InAnimContexts)
+, PoseHistories(InPoseHistories)
+, RoleToIndex(InRoleToIndex)
+, Database(InDatabase)
+, Flags(InFlags)
+{
+	check(RoleToIndex.Num() == PoseHistories.Num());
 	check(IsValid(RoleToIndex));
 }
 
 bool FDebugDrawParams::CanDraw() const
 {
-	return (!AnimInstanceProxies.IsEmpty() || !Meshes.IsEmpty()) && !RoleToIndex.IsEmpty() && Database && Database->Schema;
+	return !AnimContexts.IsEmpty() && !RoleToIndex.IsEmpty() && Database && Database->Schema;
 }
 
 const FSearchIndex* FDebugDrawParams::GetSearchIndex() const
@@ -126,9 +174,12 @@ FVector FDebugDrawParams::ExtractPosition(TConstArrayView<float> PoseVector, flo
 				}
 			}
 
-			if (SchemaBoneIdx > RootSchemaBoneIdx && !Meshes.IsEmpty())
+			if (SchemaBoneIdx > RootSchemaBoneIdx)
 			{
-				return Meshes[*RoleIndex]->GetSocketTransform(Schema->GetBoneReferences(Role)[SchemaBoneIdx].BoneName).GetTranslation();
+				if (const USkinnedMeshComponent* Mesh = GetMesh(AnimContexts, *RoleIndex))
+				{
+					return Mesh->GetSocketTransform(Schema->GetBoneReferences(Role)[SchemaBoneIdx].BoneName).GetTranslation();
+				}
 			}
 		}
 	}
@@ -262,9 +313,12 @@ FQuat FDebugDrawParams::ExtractRotation(TConstArrayView<float> PoseVector, float
 				}
 			}
 
-			if (SchemaBoneIdx > RootSchemaBoneIdx && !Meshes.IsEmpty())
+			if (SchemaBoneIdx > RootSchemaBoneIdx)
 			{
-				return Meshes[*RoleIndex]->GetSocketTransform(Schema->GetBoneReferences(Role)[SchemaBoneIdx].BoneName).GetRotation();
+				if (const USkinnedMeshComponent* Mesh = GetMesh(AnimContexts, *RoleIndex))
+				{
+					return Mesh->GetSocketTransform(Schema->GetBoneReferences(Role)[SchemaBoneIdx].BoneName).GetRotation();
+				}
 			}
 		}
 	}
@@ -286,15 +340,14 @@ void FDebugDrawParams::DrawLine(const FVector& LineStart, const FVector& LineEnd
 {
 	if (Color.A > 0)
 	{
-		if (!AnimInstanceProxies.IsEmpty())
+		if (FAnimInstanceProxy* AnimInstanceProxy = GetAnimInstanceProxy(AnimContexts))
 		{
-			// any AnimInstanceProxy is fine to draw
-			AnimInstanceProxies[0]->AnimDrawDebugLine(LineStart, LineEnd, Color, false, 0.f, Thickness, SDPG_Foreground);
+			AnimInstanceProxy->AnimDrawDebugLine(LineStart, LineEnd, Color, false, 0.f, Thickness, SDPG_Foreground);
 		}
-		else if (!Meshes.IsEmpty())
+		else if (const UWorld* World = GetWorld(AnimContexts))
 		{
 			// any Mesh is fine to draw
-			DrawDebugLine(Meshes[0]->GetWorld(), LineStart, LineEnd, Color, false, 0.f, SDPG_Foreground, Thickness);
+			DrawDebugLine(World, LineStart, LineEnd, Color, false, 0.f, SDPG_Foreground, Thickness);
 		}
 	}
 }
@@ -303,15 +356,14 @@ void FDebugDrawParams::DrawPoint(const FVector& Position, const FColor& Color, f
 {
 	if (Color.A > 0)
 	{
-		if (!AnimInstanceProxies.IsEmpty())
+		if (FAnimInstanceProxy* AnimInstanceProxy = GetAnimInstanceProxy(AnimContexts))
 		{
-			// any AnimInstanceProxy is fine to draw
-			AnimInstanceProxies[0]->AnimDrawDebugPoint(Position, Thickness, Color, false, 0.f, SDPG_Foreground);
+			AnimInstanceProxy->AnimDrawDebugPoint(Position, Thickness, Color, false, 0.f, SDPG_Foreground);
 		}
-		else if (!Meshes.IsEmpty())
+		else if (const UWorld* World = GetWorld(AnimContexts))
 		{
 			// any Mesh is fine to draw
-			DrawDebugPoint(Meshes[0]->GetWorld(), Position, Thickness, Color, false, 0.f, SDPG_Foreground);
+			DrawDebugPoint(World, Position, Thickness, Color, false, 0.f, SDPG_Foreground);
 		}
 	}
 }
@@ -320,15 +372,14 @@ void FDebugDrawParams::DrawCircle(const FMatrix& TransformMatrix, float Radius,
 {
 	if (Color.A > 0)
 	{
-		if (!AnimInstanceProxies.IsEmpty())
+		if (FAnimInstanceProxy* AnimInstanceProxy = GetAnimInstanceProxy(AnimContexts))
 		{
-			// any AnimInstanceProxy is fine to draw
-			AnimInstanceProxies[0]->AnimDrawDebugCircle(TransformMatrix.GetOrigin(), Radius, Segments, Color, TransformMatrix.GetScaledAxis(EAxis::X), false, 0.f, SDPG_Foreground, Thickness);
+			AnimInstanceProxy->AnimDrawDebugCircle(TransformMatrix.GetOrigin(), Radius, Segments, Color, TransformMatrix.GetScaledAxis(EAxis::X), false, 0.f, SDPG_Foreground, Thickness);
 		}
-		else if (!Meshes.IsEmpty())
+		else if (const UWorld* World = GetWorld(AnimContexts))
 		{
 			// any Mesh is fine to draw
-			DrawDebugCircle(Meshes[0]->GetWorld(), TransformMatrix, Radius, Segments, Color, false, 0.f, SDPG_Foreground, Thickness);
+			DrawDebugCircle(World, TransformMatrix, Radius, Segments, Color, false, 0.f, SDPG_Foreground, Thickness);
 		}
 	}
 }
@@ -444,7 +495,7 @@ FCachedQuery::FCachedQuery(const UPoseSearchSchema* InSchema)
 // FSearchContext
 FSearchContext::FSearchContext(float InDesiredPermutationTimeOffset, const FPoseIndicesHistory* InPoseIndicesHistory,
 	const FSearchResult& InCurrentResult, const FFloatInterval& InPoseJumpThresholdTime, bool bInUseCachedChannelData)
-: AnimInstances()
+: AnimContexts()
 , PoseHistories()
 , RoleToIndex()
 , AssetsToConsider()
@@ -461,11 +512,11 @@ FSearchContext::FSearchContext(float InDesiredPermutationTimeOffset, const FPose
 #endif // UE_POSE_SEARCH_TRACE_ENABLED
 }
 
-void FSearchContext::AddRole(const FRole& Role, const UAnimInstance* AnimInstance, const IPoseHistory* PoseHistory)
+void FSearchContext::AddRole(const FRole& Role, const UObject* AnimContext, const IPoseHistory* PoseHistory)
 {
-	check(RoleToIndex.Num() == AnimInstances.Num() && RoleToIndex.Num() == PoseHistories.Num());
+	check(RoleToIndex.Num() == AnimContexts.Num() && RoleToIndex.Num() == PoseHistories.Num());
 
-	AnimInstances.Add(AnimInstance);
+	AnimContexts.Add(AnimContext);
 	PoseHistories.Add(PoseHistory);
 	RoleToIndex.Add(Role) = RoleToIndex.Num();
 
@@ -585,12 +636,30 @@ FTransform FSearchContext::GetWorldRootBoneTransformAtTime(float SampleTime, con
 		}
 	}
 
-	const UAnimInstance* AnimInstance = GetAnimInstance(SampleRole);
-	if (AnimInstance && AnimInstance->CurrentSkeleton)
+	const UObject* AnimContext = GetAnimContext(SampleRole);
+	if (const UAnimInstance* AnimInstance = Cast<UAnimInstance>(AnimContext))
 	{
-		const FTransform& RootBoneTransform = AnimInstance->CurrentSkeleton->GetReferenceSkeleton().GetRefBonePose()[RootSchemaBoneIdx];
-		const FTransform& ComponentToWorldTransform = AnimInstance->GetSkelMeshComponent()->GetComponentTransform();
-		return RootBoneTransform * ComponentToWorldTransform;
+		if (AnimInstance && AnimInstance->CurrentSkeleton)
+		{
+			const FTransform& RootBoneTransform = AnimInstance->CurrentSkeleton->GetReferenceSkeleton().GetRefBonePose()[RootSchemaBoneIdx];
+			const FTransform& ComponentToWorldTransform = AnimInstance->GetSkelMeshComponent()->GetComponentTransform();
+			return RootBoneTransform * ComponentToWorldTransform;
+		}
+	}
+	else if (const UActorComponent* ActorComponent = Cast<UActorComponent>(AnimContext))
+	{
+		const AActor* Actor = ActorComponent->GetOwner();
+		check(Actor);
+		// @todo: this code depends on how AnimNext gather its context object, and will likely change in the future
+		if (const USkeletalMeshComponent* SkeletalMeshComponent = Actor->GetComponentByClass<USkeletalMeshComponent>())
+		{
+			if (const USkeletalMesh* SkeletalMesh = SkeletalMeshComponent->GetSkeletalMeshAsset())
+			{
+				const FTransform& RootBoneTransform = SkeletalMesh->GetSkeleton()->GetReferenceSkeleton().GetRefBonePose()[RootSchemaBoneIdx];
+				const FTransform& ComponentToWorldTransform = SkeletalMeshComponent->GetComponentTransform();
+				return RootBoneTransform * ComponentToWorldTransform;
+			}
+		}
 	}
 
 	return FTransform::Identity;
diff --git a/Engine/Plugins/Animation/PoseSearch/Source/Runtime/Private/PoseSearchFeatureChannel_Curve.cpp b/Engine/Plugins/Animation/PoseSearch/Source/Runtime/Private/PoseSearchFeatureChannel_Curve.cpp
--- a/Engine/Plugins/Animation/PoseSearch/Source/Runtime/Private/PoseSearchFeatureChannel_Curve.cpp
+++ b/Engine/Plugins/Animation/PoseSearch/Source/Runtime/Private/PoseSearchFeatureChannel_Curve.cpp
@@ -35,8 +35,16 @@ void UPoseSearchFeatureChannel_Curve::BuildQuery(UE::PoseSearch::FSearchContext&
 	float CurveValue = 0.0f;
 	if (bUseBlueprintQueryOverride)
 	{
-		CurveValue = BP_GetCurveValue(SearchContext.GetAnimInstance(SampleRole));
-		FFeatureVectorHelper::EncodeFloat(SearchContext.EditFeatureVector(), ChannelDataOffset, CurveValue);
+		if (const UAnimInstance* AnimInstance = Cast<UAnimInstance>(SearchContext.GetAnimContext(SampleRole)))
+		{
+			CurveValue = BP_GetCurveValue(AnimInstance);
+			FFeatureVectorHelper::EncodeFloat(SearchContext.EditFeatureVector(), ChannelDataOffset, CurveValue);
+		}
+		else
+		{
+			// @todo: support non UAnimInstance anim contexts for AnimNext
+			UE_LOG(LogPoseSearch, Warning, TEXT("UPoseSearchFeatureChannel_Curve::BuildQuery - unsupported null UAnimInstance: WIP support for AnimNext!"));
+		}
 		return;
 	}
 	
diff --git a/Engine/Plugins/Animation/PoseSearch/Source/Runtime/Private/PoseSearchFeatureChannel_Heading.cpp b/Engine/Plugins/Animation/PoseSearch/Source/Runtime/Private/PoseSearchFeatureChannel_Heading.cpp
--- a/Engine/Plugins/Animation/PoseSearch/Source/Runtime/Private/PoseSearchFeatureChannel_Heading.cpp
+++ b/Engine/Plugins/Animation/PoseSearch/Source/Runtime/Private/PoseSearchFeatureChannel_Heading.cpp
@@ -103,9 +103,17 @@ void UPoseSearchFeatureChannel_Heading::BuildQuery(UE::PoseSearch::FSearchContex
 	const bool bIsRootBone = SchemaBoneIdx == RootSchemaBoneIdx;
 	if (bUseBlueprintQueryOverride)
 	{
-		const FQuat BoneRotationWorld = BP_GetWorldRotation(SearchContext.GetAnimInstance(SampleRole));
-		const FQuat BoneRotation = SearchContext.GetSampleRotation(SampleTimeOffset, OriginTimeOffset, SchemaBoneIdx, SchemaOriginBoneIdx, SampleRole, OriginRole, EPermutationTimeType::UseSampleTime, &BoneRotationWorld);
-		FFeatureVectorHelper::EncodeVector(SearchContext.EditFeatureVector(), ChannelDataOffset, GetAxis(BoneRotation), ComponentStripping, true);
+		if (const UAnimInstance* AnimInstance = Cast<UAnimInstance>(SearchContext.GetAnimContext(SampleRole)))
+		{
+			const FQuat BoneRotationWorld = BP_GetWorldRotation(AnimInstance);
+			const FQuat BoneRotation = SearchContext.GetSampleRotation(SampleTimeOffset, OriginTimeOffset, SchemaBoneIdx, SchemaOriginBoneIdx, SampleRole, OriginRole, EPermutationTimeType::UseSampleTime, &BoneRotationWorld);
+			FFeatureVectorHelper::EncodeVector(SearchContext.EditFeatureVector(), ChannelDataOffset, GetAxis(BoneRotation), ComponentStripping, true);
+		}
+		else
+		{
+			// @todo: support non UAnimInstance anim contexts for AnimNext
+			UE_LOG(LogPoseSearch, Warning, TEXT("UPoseSearchFeatureChannel_Curve::BuildQuery - unsupported null UAnimInstance: WIP support for AnimNext!"));
+		}
 		return;
 	}
 	
diff --git a/Engine/Plugins/Animation/PoseSearch/Source/Runtime/Private/PoseSearchFeatureChannel_Position.cpp b/Engine/Plugins/Animation/PoseSearch/Source/Runtime/Private/PoseSearchFeatureChannel_Position.cpp
--- a/Engine/Plugins/Animation/PoseSearch/Source/Runtime/Private/PoseSearchFeatureChannel_Position.cpp
+++ b/Engine/Plugins/Animation/PoseSearch/Source/Runtime/Private/PoseSearchFeatureChannel_Position.cpp
@@ -83,9 +83,17 @@ void UPoseSearchFeatureChannel_Position::BuildQuery(UE::PoseSearch::FSearchConte
 
 	if (bUseBlueprintQueryOverride)
 	{
-		const FVector BonePositionWorld = BP_GetWorldPosition(SearchContext.GetAnimInstance(SampleRole));
-		const FVector BonePosition = SearchContext.GetSamplePosition(SampleTimeOffset, OriginTimeOffset, SchemaBoneIdx, SchemaOriginBoneIdx, SampleRole, OriginRole, PermutationTimeType, &BonePositionWorld);
-  		FFeatureVectorHelper::EncodeVector(SearchContext.EditFeatureVector(), ChannelDataOffset, BonePosition, ComponentStripping, false);
+		if (const UAnimInstance* AnimInstance = Cast<UAnimInstance>(SearchContext.GetAnimContext(SampleRole)))
+		{
+			const FVector BonePositionWorld = BP_GetWorldPosition(AnimInstance);
+			const FVector BonePosition = SearchContext.GetSamplePosition(SampleTimeOffset, OriginTimeOffset, SchemaBoneIdx, SchemaOriginBoneIdx, SampleRole, OriginRole, PermutationTimeType, &BonePositionWorld);
+  			FFeatureVectorHelper::EncodeVector(SearchContext.EditFeatureVector(), ChannelDataOffset, BonePosition, ComponentStripping, false);
+		}
+		else
+		{
+			// @todo: support non UAnimInstance anim contexts for AnimNext
+			UE_LOG(LogPoseSearch, Warning, TEXT("UPoseSearchFeatureChannel_Curve::BuildQuery - unsupported null UAnimInstance: WIP support for AnimNext!"));
+		}
 		return;
 	}
 
diff --git a/Engine/Plugins/Animation/PoseSearch/Source/Runtime/Private/PoseSearchFeatureChannel_TimeToEvent.cpp b/Engine/Plugins/Animation/PoseSearch/Source/Runtime/Private/PoseSearchFeatureChannel_TimeToEvent.cpp
--- a/Engine/Plugins/Animation/PoseSearch/Source/Runtime/Private/PoseSearchFeatureChannel_TimeToEvent.cpp
+++ b/Engine/Plugins/Animation/PoseSearch/Source/Runtime/Private/PoseSearchFeatureChannel_TimeToEvent.cpp
@@ -30,13 +30,20 @@ void UPoseSearchFeatureChannel_TimeToEvent::BuildQuery(UE::PoseSearch::FSearchCo
 		return;
 	}
 
-	if (SearchContext.GetAnimInstances().IsEmpty())
+	if (SearchContext.GetAnimContexts().IsEmpty())
 	{
-		UE_LOG(LogPoseSearch, Error, TEXT("UPoseSearchFeatureChannel_Position::BuildQuery - no provided anim instance!"));
+		UE_LOG(LogPoseSearch, Error, TEXT("UPoseSearchFeatureChannel_Position::BuildQuery - no provided anim contexts!"));
 	}
 	
-	const float TimeToEvent = BP_GetTimeToEvent(SearchContext.GetAnimInstances()[0]);
-	FFeatureVectorHelper::EncodeFloat(SearchContext.EditFeatureVector(), ChannelDataOffset, TimeToEvent);
+	if (const UAnimInstance* AnimInstance = Cast<UAnimInstance>(SearchContext.GetAnimContexts()[0]))
+	{
+		const float TimeToEvent = BP_GetTimeToEvent(AnimInstance);
+		FFeatureVectorHelper::EncodeFloat(SearchContext.EditFeatureVector(), ChannelDataOffset, TimeToEvent);
+	}
+	else
+	{
+		UE_LOG(LogPoseSearch, Warning, TEXT("UPoseSearchFeatureChannel_TimeToEvent::BuildQuery - UPoseSearchFeatureChannel_TimeToEvent is designed to work only as BluePrint overridden class. The query TimeToEvent value will be defaulted to 0"));
+	}
 }
 
 #if WITH_EDITOR
diff --git a/Engine/Plugins/Animation/PoseSearch/Source/Runtime/Private/PoseSearchFeatureChannel_Velocity.cpp b/Engine/Plugins/Animation/PoseSearch/Source/Runtime/Private/PoseSearchFeatureChannel_Velocity.cpp
--- a/Engine/Plugins/Animation/PoseSearch/Source/Runtime/Private/PoseSearchFeatureChannel_Velocity.cpp
+++ b/Engine/Plugins/Animation/PoseSearch/Source/Runtime/Private/PoseSearchFeatureChannel_Velocity.cpp
@@ -47,13 +47,21 @@ void UPoseSearchFeatureChannel_Velocity::BuildQuery(UE::PoseSearch::FSearchConte
 	const bool bIsRootBone = SchemaBoneIdx == RootSchemaBoneIdx;
 	if (bUseBlueprintQueryOverride)
 	{
-		const FVector LinearVelocityWorld = BP_GetWorldVelocity(SearchContext.GetAnimInstance(SampleRole));
-		FVector LinearVelocity = SearchContext.GetSampleVelocity(SampleTimeOffset, OriginTimeOffset, SchemaBoneIdx, SchemaOriginBoneIdx, SampleRole, OriginRole, bUseCharacterSpaceVelocities, EPermutationTimeType::UseSampleTime, &LinearVelocityWorld);
-		if (bNormalize)
+		if (const UAnimInstance* AnimInstance = Cast<UAnimInstance>(SearchContext.GetAnimContext(SampleRole)))
 		{
-			LinearVelocity = LinearVelocity.GetClampedToMaxSize(1.f);
+			const FVector LinearVelocityWorld = BP_GetWorldVelocity(AnimInstance);
+			FVector LinearVelocity = SearchContext.GetSampleVelocity(SampleTimeOffset, OriginTimeOffset, SchemaBoneIdx, SchemaOriginBoneIdx, SampleRole, OriginRole, bUseCharacterSpaceVelocities, EPermutationTimeType::UseSampleTime, &LinearVelocityWorld);
+			if (bNormalize)
+			{
+				LinearVelocity = LinearVelocity.GetClampedToMaxSize(1.f);
+			}
+			FFeatureVectorHelper::EncodeVector(SearchContext.EditFeatureVector(), ChannelDataOffset, LinearVelocity, ComponentStripping, false);
+		}
+		else
+		{
+			// @todo: support non UAnimInstance anim contexts for AnimNext
+			UE_LOG(LogPoseSearch, Warning, TEXT("UPoseSearchFeatureChannel_Curve::BuildQuery - unsupported null UAnimInstance: WIP support for AnimNext!"));
 		}
-		FFeatureVectorHelper::EncodeVector(SearchContext.EditFeatureVector(), ChannelDataOffset, LinearVelocity, ComponentStripping, false);
 		return;
 	}
 	
diff --git a/Engine/Plugins/Animation/PoseSearch/Source/Runtime/Private/PoseSearchInteractionAsset.cpp b/Engine/Plugins/Animation/PoseSearch/Source/Runtime/Private/PoseSearchInteractionAsset.cpp
--- a/Engine/Plugins/Animation/PoseSearch/Source/Runtime/Private/PoseSearchInteractionAsset.cpp
+++ b/Engine/Plugins/Animation/PoseSearch/Source/Runtime/Private/PoseSearchInteractionAsset.cpp
@@ -227,119 +227,132 @@ void UPoseSearchInteractionAsset::CalculateWarpTransforms(float Time, const TArr
 	check(FullAlignedActorRootBoneTransforms.Num() == GetNumRoles());
 
 	const int32 ItemsNum = Items.Num();
-	if (ItemsNum == 0)
+	if (ItemsNum < 2)
 	{
-		return;
+		for (int32 ItemIndex = 0; ItemIndex < ItemsNum; ++ItemIndex)
+		{
+			FullAlignedActorRootBoneTransforms[ItemIndex] = ActorRootBoneTransforms[ItemIndex];
+		}
 	}
-
-	TArray<FTransform, TInlineAllocator<UE::PoseSearch::PreallocatedRolesNum>> AssetRootBoneTransforms;
-	AssetRootBoneTransforms.SetNum(ItemsNum);
-
-	// ItemIndex is the RoleIndex and Role = Item.Role
-	for (int32 ItemIndex = 0; ItemIndex < ItemsNum; ++ItemIndex)
+	else
 	{
-		const FPoseSearchInteractionAssetItem& Item = Items[ItemIndex];
+		TArray<FTransform, TInlineAllocator<UE::PoseSearch::PreallocatedRolesNum>> AssetRootBoneTransforms;
+		AssetRootBoneTransforms.SetNum(ItemsNum);
 
-		// sampling the AnimationAsset to extract the current time transform and the initial (time of 0) transform
-		const UE::PoseSearch::FAnimationAssetSampler Sampler(Item.Animation, Item.Origin);
-		AssetRootBoneTransforms[ItemIndex] = Sampler.ExtractRootTransform(Time);
-
-#if DO_CHECK
-		// array containing the bone index of the root bone (0)
-		TArray<uint16, TInlineAllocator<1>> BoneIndices;
-		BoneIndices.SetNumZeroed(1);
-
-		// extracting the pose, containing only the root bone from the Sampler 
-		FMemMark Mark(FMemStack::Get());
-		FCompactPose Pose;
-		FBoneContainer BoneContainer;
-		BoneContainer.InitializeTo(BoneIndices, UE::Anim::FCurveFilterSettings(UE::Anim::ECurveFilterMode::DisallowAll), *Items[ItemIndex].Animation->GetSkeleton());
-		Pose.SetBoneContainer(&BoneContainer);
-		Sampler.ExtractPose(Time, Pose);
-
-		// making sure the animation root bone transform is Identity, so we can confuse the root with the root BONE transform and preserve performances!
-		const FTransform& RootBoneTransform = Pose.GetBones()[0];
-		check(RootBoneTransform.Equals(FTransform::Identity));
-#endif // DO_CHECK
-	}
+		// ItemIndex is the RoleIndex and Role = Item.Role
+		for (int32 ItemIndex = 0; ItemIndex < ItemsNum; ++ItemIndex)
+		{
+			const FPoseSearchInteractionAssetItem& Item = Items[ItemIndex];
 
-	TArray<int32, TInlineAllocator<UE::PoseSearch::PreallocatedRolesNum>> SortedByWarpingWeightRotationItemIndex;
-	TArray<float, TInlineAllocator<UE::PoseSearch::PreallocatedRolesNum>> NormalizedWarpingWeightTranslation;
-	SortedByWarpingWeightRotationItemIndex.SetNum(ItemsNum);
-	NormalizedWarpingWeightTranslation.SetNum(ItemsNum);
+			// sampling the AnimationAsset to extract the current time transform and the initial (time of 0) transform
+			const UE::PoseSearch::FAnimationAssetSampler Sampler(Item.Animation, Item.Origin);
+			AssetRootBoneTransforms[ItemIndex] = Sampler.ExtractRootTransform(Time);
 
-	float WarpingWeightTranslationSum = 0.f;
-	float WarpingWeightRotationSum = 0.f;
-	for (int32 ItemIndex = 0; ItemIndex < ItemsNum; ++ItemIndex)
-	{
-		SortedByWarpingWeightRotationItemIndex[ItemIndex] = ItemIndex;
-		WarpingWeightTranslationSum += Items[ItemIndex].WarpingWeightTranslation;
-		WarpingWeightRotationSum += Items[ItemIndex].WarpingWeightRotation;
-	}
+#if ENABLE_ANIM_DEBUG && WITH_EDITOR
+			if (Items[ItemIndex].Animation)
+			{
+				// array containing the bone index of the root bone (0)
+				TArray<uint16, TInlineAllocator<1>> BoneIndices;
+				BoneIndices.SetNumZeroed(1);
+
+				// extracting the pose, containing only the root bone from the Sampler 
+				FMemMark Mark(FMemStack::Get());
+				FCompactPose Pose;
+				FBoneContainer BoneContainer;
+				BoneContainer.InitializeTo(BoneIndices, UE::Anim::FCurveFilterSettings(UE::Anim::ECurveFilterMode::DisallowAll), *Items[ItemIndex].Animation->GetSkeleton());
+				Pose.SetBoneContainer(&BoneContainer);
+				Sampler.ExtractPose(Time, Pose);
+
+				// making sure the animation root bone transform is Identity, so we can confuse the root with the root BONE transform and preserve performances!
+				const FTransform& RootBoneTransform = Pose.GetBones()[0];
+				if (!RootBoneTransform.Equals(FTransform::Identity))
+				{
+					const FVector Pos = RootBoneTransform.GetLocation();
+					const FRotator Rot(RootBoneTransform.GetRotation());
+					UE_LOG(LogPoseSearch, Error, TEXT("UPoseSearchInteractionAsset::CalculateWarpTransforms unsupported non identity root bone in %s at time %f Pos(%f, %f, %f), Rot(%f, %f, %f)"), *Items[ItemIndex].Animation->GetName(), Time, Pos.X, Pos.Y, Pos.Z, Rot.Pitch, Rot.Yaw, Rot.Roll);
+				}
+			}
+#endif // ENABLE_ANIM_DEBUG && WITH_EDITOR
+		}
 
-	const float NormalizedHomogeneousWeight = 1.f / ItemsNum;
-	if (WarpingWeightTranslationSum > UE_KINDA_SMALL_NUMBER)
-	{
+		TArray<int32, TInlineAllocator<UE::PoseSearch::PreallocatedRolesNum>> SortedByWarpingWeightRotationItemIndex;
+		TArray<float, TInlineAllocator<UE::PoseSearch::PreallocatedRolesNum>> NormalizedWarpingWeightTranslation;
+		SortedByWarpingWeightRotationItemIndex.SetNum(ItemsNum);
+		NormalizedWarpingWeightTranslation.SetNum(ItemsNum);
+
+		float WarpingWeightTranslationSum = 0.f;
+		float WarpingWeightRotationSum = 0.f;
 		for (int32 ItemIndex = 0; ItemIndex < ItemsNum; ++ItemIndex)
 		{
-			NormalizedWarpingWeightTranslation[ItemIndex] = Items[ItemIndex].WarpingWeightTranslation / WarpingWeightTranslationSum;
+			SortedByWarpingWeightRotationItemIndex[ItemIndex] = ItemIndex;
+			WarpingWeightTranslationSum += Items[ItemIndex].WarpingWeightTranslation;
+			WarpingWeightRotationSum += Items[ItemIndex].WarpingWeightRotation;
 		}
-	}
-	else
-	{
-		for (int32 ItemIndex = 0; ItemIndex < ItemsNum; ++ItemIndex)
+
+		const float NormalizedHomogeneousWeight = 1.f / ItemsNum;
+		if (WarpingWeightTranslationSum > UE_KINDA_SMALL_NUMBER)
 		{
-			NormalizedWarpingWeightTranslation[ItemIndex] = NormalizedHomogeneousWeight;
+			for (int32 ItemIndex = 0; ItemIndex < ItemsNum; ++ItemIndex)
+			{
+				NormalizedWarpingWeightTranslation[ItemIndex] = Items[ItemIndex].WarpingWeightTranslation / WarpingWeightTranslationSum;
+			}
 		}
-	}
-
-	if (WarpingWeightRotationSum > UE_KINDA_SMALL_NUMBER)
-	{
-		SortedByWarpingWeightRotationItemIndex.Sort([this](const int32 A, const int32 B)
+		else
+		{
+			for (int32 ItemIndex = 0; ItemIndex < ItemsNum; ++ItemIndex)
 			{
-				return Items[A].WarpingWeightRotation < Items[B].WarpingWeightRotation;
-			});
-	}
+				NormalizedWarpingWeightTranslation[ItemIndex] = NormalizedHomogeneousWeight;
+			}
+		}
 
-	const FQuat AssetReferenceOrientation = FindReferenceOrientation(AssetRootBoneTransforms, SortedByWarpingWeightRotationItemIndex);
-	const FQuat ActorsReferenceOrientation = FindReferenceOrientation(ActorRootBoneTransforms, SortedByWarpingWeightRotationItemIndex);
-	
-	FQuat WeightedActorsReferenceOrientation = ActorsReferenceOrientation;
-	if (WarpingWeightRotationSum > UE_KINDA_SMALL_NUMBER)
-	{
-		// ItemIndex are in order of WarpingWeightRotation. the last one is the one with the highest most WarpingWeightRotation, the most "important"
-		for (int32 ItemIndex : SortedByWarpingWeightRotationItemIndex)
+		if (WarpingWeightRotationSum > UE_KINDA_SMALL_NUMBER)
 		{
-			const FPoseSearchInteractionAssetItem& Item = Items[ItemIndex];
-			const float NormalizedWarpingWeightRotation = Item.WarpingWeightRotation / WarpingWeightRotationSum;
-			if (NormalizedWarpingWeightRotation > NormalizedHomogeneousWeight)
+			SortedByWarpingWeightRotationItemIndex.Sort([this](const int32 A, const int32 B)
+				{
+					return Items[A].WarpingWeightRotation < Items[B].WarpingWeightRotation;
+				});
+		}
+
+		const FQuat AssetReferenceOrientation = FindReferenceOrientation(AssetRootBoneTransforms, SortedByWarpingWeightRotationItemIndex);
+		const FQuat ActorsReferenceOrientation = FindReferenceOrientation(ActorRootBoneTransforms, SortedByWarpingWeightRotationItemIndex);
+
+		FQuat WeightedActorsReferenceOrientation = ActorsReferenceOrientation;
+		if (WarpingWeightRotationSum > UE_KINDA_SMALL_NUMBER)
+		{
+			// ItemIndex are in order of WarpingWeightRotation. the last one is the one with the highest most WarpingWeightRotation, the most "important"
+			for (int32 ItemIndex : SortedByWarpingWeightRotationItemIndex)
 			{
-				// NormalizedHomogeneousWeight is one only if ItemsNum is one, 
-				// BUT NormalizedWarpingWeightRotation > NormalizedHomogeneousWeight should always be false
-				check(!FMath::IsNearlyEqual(NormalizedHomogeneousWeight, 1.f));
+				const FPoseSearchInteractionAssetItem& Item = Items[ItemIndex];
+				const float NormalizedWarpingWeightRotation = Item.WarpingWeightRotation / WarpingWeightRotationSum;
+				if (NormalizedWarpingWeightRotation > NormalizedHomogeneousWeight)
+				{
+					// NormalizedHomogeneousWeight is one only if ItemsNum is one, 
+					// BUT NormalizedWarpingWeightRotation > NormalizedHomogeneousWeight should always be false
+					check(!FMath::IsNearlyEqual(NormalizedHomogeneousWeight, 1.f));
 
-				// how much this item wants to reorient the ReferenceOrientation from the homogeneous "fair" value
-				const float SlerpParam = (NormalizedWarpingWeightRotation - NormalizedHomogeneousWeight) / (1.f - NormalizedHomogeneousWeight);
+					// how much this item wants to reorient the ReferenceOrientation from the homogeneous "fair" value
+					const float SlerpParam = (NormalizedWarpingWeightRotation - NormalizedHomogeneousWeight) / (1.f - NormalizedHomogeneousWeight);
 
-				// calculating the reference orientation relative to the character
-				// AssetReferenceOrientation in actor world orientation
-				const FQuat ActorAssetReferenceOrientation = ActorRootBoneTransforms[ItemIndex].GetRotation() * (AssetRootBoneTransforms[ItemIndex].GetRotation().Inverse() * AssetReferenceOrientation);
+					// calculating the reference orientation relative to the character
+					// AssetReferenceOrientation in actor world orientation
+					const FQuat ActorAssetReferenceOrientation = ActorRootBoneTransforms[ItemIndex].GetRotation() * (AssetRootBoneTransforms[ItemIndex].GetRotation().Inverse() * AssetReferenceOrientation);
 
-				WeightedActorsReferenceOrientation = FQuat::Slerp(WeightedActorsReferenceOrientation, ActorAssetReferenceOrientation, SlerpParam);
+					WeightedActorsReferenceOrientation = FQuat::Slerp(WeightedActorsReferenceOrientation, ActorAssetReferenceOrientation, SlerpParam);
+				}
 			}
 		}
-	}
 
-	const FVector AssetReferencePosition = FindReferencePosition(AssetRootBoneTransforms, NormalizedWarpingWeightTranslation);
-	const FVector ActorsReferencePosition = FindReferencePosition(ActorRootBoneTransforms, NormalizedWarpingWeightTranslation);
+		const FVector AssetReferencePosition = FindReferencePosition(AssetRootBoneTransforms, NormalizedWarpingWeightTranslation);
+		const FVector ActorsReferencePosition = FindReferencePosition(ActorRootBoneTransforms, NormalizedWarpingWeightTranslation);
 
-	// aligning all the actors to ActorsReferencePosition, WeightedActorsReferenceOrientation
-	const FTransform AssetReferenceTransform(AssetReferenceOrientation, AssetReferencePosition);
-	const FTransform ActorsReferenceTransform(WeightedActorsReferenceOrientation, ActorsReferencePosition);
-	const FTransform AssetReferenceInverseTransform = AssetReferenceTransform.Inverse();
-	
-	for (int32 ItemIndex = 0; ItemIndex < ItemsNum; ++ItemIndex)
-	{
-		FullAlignedActorRootBoneTransforms[ItemIndex] = (AssetRootBoneTransforms[ItemIndex] * AssetReferenceInverseTransform) * ActorsReferenceTransform;
+		// aligning all the actors to ActorsReferencePosition, WeightedActorsReferenceOrientation
+		const FTransform AssetReferenceTransform(AssetReferenceOrientation, AssetReferencePosition);
+		const FTransform ActorsReferenceTransform(WeightedActorsReferenceOrientation, ActorsReferencePosition);
+		const FTransform AssetReferenceInverseTransform = AssetReferenceTransform.Inverse();
+
+		for (int32 ItemIndex = 0; ItemIndex < ItemsNum; ++ItemIndex)
+		{
+			FullAlignedActorRootBoneTransforms[ItemIndex] = (AssetRootBoneTransforms[ItemIndex] * AssetReferenceInverseTransform) * ActorsReferenceTransform;
+		}
 	}
 }
\ No newline at end of file
diff --git a/Engine/Plugins/Animation/PoseSearch/Source/Runtime/Private/PoseSearchInteractionIsland.cpp b/Engine/Plugins/Animation/PoseSearch/Source/Runtime/Private/PoseSearchInteractionIsland.cpp
--- a/Engine/Plugins/Animation/PoseSearch/Source/Runtime/Private/PoseSearchInteractionIsland.cpp
+++ b/Engine/Plugins/Animation/PoseSearch/Source/Runtime/Private/PoseSearchInteractionIsland.cpp
@@ -458,7 +458,7 @@ bool FInteractionIsland::DoSearch_AnyThread(UObject* AnimInstance, const FPoseSe
 	{
 		FMemMark Mark(FMemStack::Get());
 
-		TArray<UAnimInstance*, TInlineAllocator<UE::PoseSearch::PreallocatedRolesNum, TMemStackAllocator<>>> AnimInstances;
+		TArray<const UObject*, TInlineAllocator<UE::PoseSearch::PreallocatedRolesNum, TMemStackAllocator<>>> AnimContexts;
 		TArray<const UE::PoseSearch::IPoseHistory*, TInlineAllocator<UE::PoseSearch::PreallocatedRolesNum, TMemStackAllocator<>>> PoseHistories;
 		TArray<UE::PoseSearch::FSearchResult, TMemStackAllocator<>> PoseSearchResults;
 
@@ -481,7 +481,8 @@ bool FInteractionIsland::DoSearch_AnyThread(UObject* AnimInstance, const FPoseSe
 				return false;
 			}
 
-			AnimInstances.Reset();
+			AnimContexts.Reset();
+			AnimContexts.Reserve(SearchContext.AnimInstances.Num());
 			for (const TWeakObjectPtr<UAnimInstance>& AnimInstancePtr : SearchContext.AnimInstances)
 			{
 				UAnimInstance* SearchContextAnimInstance = AnimInstancePtr.Get();
@@ -491,7 +492,7 @@ bool FInteractionIsland::DoSearch_AnyThread(UObject* AnimInstance, const FPoseSe
 					return false;
 				}
 
-				AnimInstances.Add(SearchContextAnimInstance);
+				AnimContexts.Add(SearchContextAnimInstance);
 			}
 
 			PoseHistories.Reset();
@@ -536,7 +537,7 @@ bool FInteractionIsland::DoSearch_AnyThread(UObject* AnimInstance, const FPoseSe
 			//}
 
 			// @todo: we could perform multiple UPoseSearchLibrary::MotionMatch in parallel!
-			const UE::PoseSearch::FSearchResult PoseSearchResult = UPoseSearchLibrary::MotionMatch(AnimInstances, SearchContext.Roles,
+			const UE::PoseSearch::FSearchResult PoseSearchResult = UPoseSearchLibrary::MotionMatch(AnimContexts, SearchContext.Roles,
 				PoseHistories, AssetsToSearch, ContinuingPropertiesToUse, PoseSearchFutureProperties);
 
 			if (PoseSearchResult.PoseCost.GetTotalCost() < SearchContext.MaxCost)
diff --git a/Engine/Plugins/Animation/PoseSearch/Source/Runtime/Private/PoseSearchInteractionSubsystem.cpp b/Engine/Plugins/Animation/PoseSearch/Source/Runtime/Private/PoseSearchInteractionSubsystem.cpp
--- a/Engine/Plugins/Animation/PoseSearch/Source/Runtime/Private/PoseSearchInteractionSubsystem.cpp
+++ b/Engine/Plugins/Animation/PoseSearch/Source/Runtime/Private/PoseSearchInteractionSubsystem.cpp
@@ -421,7 +421,7 @@ void UPoseSearchInteractionSubsystem::Tick(float DeltaSeconds)
 	// FScopeLock Lock(&AvailabilityRequestsMapMutex); is not necessary since UPoseSearchInteractionSubsystem gets ticked outside the parallel animation jobs
 	FMemMark Mark(FMemStack::Get());
 	
-	// processing the AvailabilityRequestsMap to fill up a map of Databases pointing to an array of all the AnimInstances with related roles
+	// processing the AvailabilityRequestsMap to fill up a map of Databases pointing to an array of all the AnimContexts with related roles
 	FDatabaseToRoledAnimInstances DatabaseToRoledAnimInstances;
 	ProcessAvailabilityRequestsMap(AvailabilityRequestsMap, GetWorld(), DatabaseToRoledAnimInstances);
 
diff --git a/Engine/Plugins/Animation/PoseSearch/Source/Runtime/Private/PoseSearchLibrary.cpp b/Engine/Plugins/Animation/PoseSearch/Source/Runtime/Private/PoseSearchLibrary.cpp
--- a/Engine/Plugins/Animation/PoseSearch/Source/Runtime/Private/PoseSearchLibrary.cpp
+++ b/Engine/Plugins/Animation/PoseSearch/Source/Runtime/Private/PoseSearchLibrary.cpp
@@ -212,21 +212,6 @@ namespace UE::PoseSearch
 
 		return false;
 	}
-
-#if ENABLE_DRAW_DEBUG && ENABLE_ANIM_DEBUG
-	class UAnimInstanceProxyProvider : public UAnimInstance
-	{
-	public:
-		static FAnimInstanceProxy* GetAnimInstanceProxy(UAnimInstance* AnimInstance)
-		{
-			if (AnimInstance)
-			{
-				return &static_cast<UAnimInstanceProxyProvider*>(AnimInstance)->GetProxyOnAnyThread<FAnimInstanceProxy>();
-			}
-			return nullptr;
-		}
-	};
-#endif //ENABLE_DRAW_DEBUG && ENABLE_ANIM_DEBUG
 }
 
 //////////////////////////////////////////////////////////////////////////
@@ -332,24 +317,40 @@ void UPoseSearchLibrary::TraceMotionMatching(UE::PoseSearch::FSearchContext& Sea
 	FSearchResult& CurrentResult = CurrentState.CurrentSearchResult;
 	const float ElapsedPoseSearchTime = CurrentState.ElapsedPoseSearchTime;
 
-	const int32 AnimInstancesNum = SearchContext.GetAnimInstances().Num();
-	TraceState.SkeletalMeshComponentIds.SetNum(AnimInstancesNum);
+	const int32 AnimContextsNum = SearchContext.GetAnimContexts().Num();
+	TraceState.SkeletalMeshComponentIds.SetNum(AnimContextsNum);
 
-	for (int32 AnimInstanceIndex = 0; AnimInstanceIndex < AnimInstancesNum; ++AnimInstanceIndex)
+	for (int32 AnimInstanceIndex = 0; AnimInstanceIndex < AnimContextsNum; ++AnimInstanceIndex)
 	{
-		if (const UAnimInstance* AnimInstance = SearchContext.GetAnimInstances()[AnimInstanceIndex])
+		if (const UObject* AnimContext = SearchContext.GetAnimContexts()[AnimInstanceIndex])
 		{
-			const UObject* SkeletalMeshComponent = AnimInstance->GetOuter();
+			if (const UAnimInstance* AnimInstance = Cast<UAnimInstance>(AnimContext))
+			{
+				const UObject* SkeletalMeshComponent = AnimInstance->GetOuter();
 
-			TRACE_OBJECT(AnimInstance);
+				TRACE_OBJECT(AnimInstance);
 
-			TraceState.SkeletalMeshComponentIds[AnimInstanceIndex] = FObjectTrace::GetObjectId(SkeletalMeshComponent);
+				TraceState.SkeletalMeshComponentIds[AnimInstanceIndex] = FObjectTrace::GetObjectId(SkeletalMeshComponent);
 
-			SearchId = HashCombineFast(SearchId, GetTypeHash(FObjectTrace::GetObjectId(AnimInstance)));
+				SearchId = HashCombineFast(SearchId, GetTypeHash(FObjectTrace::GetObjectId(AnimInstance)));
+			}
+			else if (const UActorComponent* ActorComponent = Cast<UActorComponent>(AnimContext))
+			{
+				const AActor* Actor = ActorComponent->GetOwner();
+				check(Actor);
+				
+				const UObject* SkeletalMeshComponent = Actor->GetComponentByClass<USkeletalMeshComponent>();
+
+				TRACE_OBJECT(ActorComponent);
+
+				TraceState.SkeletalMeshComponentIds[AnimInstanceIndex] = FObjectTrace::GetObjectId(SkeletalMeshComponent);
+
+				SearchId = HashCombineFast(SearchId, GetTypeHash(FObjectTrace::GetObjectId(ActorComponent)));
+			}
 		}
 	}
 
-	TraceState.Roles.SetNum(AnimInstancesNum);
+	TraceState.Roles.SetNum(AnimContextsNum);
 	for (const FRoleToIndexPair& RoleToIndexPair : SearchContext.GetRoleToIndex())
 	{
 		TraceState.Roles[RoleToIndexPair.Value] = RoleToIndexPair.Key;
@@ -358,8 +359,8 @@ void UPoseSearchLibrary::TraceMotionMatching(UE::PoseSearch::FSearchContext& Sea
 	SearchId = HashCombineFast(SearchId, GetTypeHash(TraceState.Roles));
 
 	// @todo: do we need to hash pose history names in SearchId as well?
-	TraceState.PoseHistories.SetNum(AnimInstancesNum);
-	for (int32 AnimInstanceIndex = 0; AnimInstanceIndex < AnimInstancesNum; ++AnimInstanceIndex)
+	TraceState.PoseHistories.SetNum(AnimContextsNum);
+	for (int32 AnimInstanceIndex = 0; AnimInstanceIndex < AnimContextsNum; ++AnimInstanceIndex)
 	{
 		TraceState.PoseHistories[AnimInstanceIndex].InitFrom(SearchContext.GetPoseHistories()[AnimInstanceIndex]);
 	}
@@ -475,10 +476,10 @@ void UPoseSearchLibrary::TraceMotionMatching(UE::PoseSearch::FSearchContext& Sea
 
 	TraceState.Cycle = FPlatformTime::Cycles64();
 
-	// @todo: avoid publishing duplicated TraceState in ALL the AnimInstances! -currently necessary for multi character-
-	for (const UAnimInstance* AnimInstance : SearchContext.GetAnimInstances())
+	// @todo: avoid publishing duplicated TraceState in ALL the AnimContexts! -currently necessary for multi character-
+	for (const UObject* AnimContext : SearchContext.GetAnimContexts())
 	{
-		TraceState.AnimInstanceId = FObjectTrace::GetObjectId(AnimInstance);
+		TraceState.AnimInstanceId = FObjectTrace::GetObjectId(AnimContext);
 		TraceState.NodeId = SearchId;
 		TraceState.Output();
 	}
@@ -675,18 +676,15 @@ void UPoseSearchLibrary::UpdateMotionMatchingState(
 		if (EAsyncBuildIndexResult::Success == FAsyncPoseSearchDatabasesManagement::RequestAsyncBuildIndex(CurResultDatabase, ERequestAsyncBuildFlag::ContinueRequest))
 #endif // WITH_EDITOR
 		{
-			FAnimInstanceProxy* AnimInstanceProxy = Context.AnimInstanceProxy;
-			const TArrayView<FAnimInstanceProxy*> AnimInstanceProxies = MakeArrayView(&AnimInstanceProxy, 1);
-			
 			if (bDebugDrawCurResult)
 			{
-				FDebugDrawParams DrawParams(AnimInstanceProxies, SearchContext.GetPoseHistories(), SearchContext.GetRoleToIndex(), CurResultDatabase);
+				FDebugDrawParams DrawParams(SearchContext.GetAnimContexts(), SearchContext.GetPoseHistories(), SearchContext.GetRoleToIndex(), CurResultDatabase);
 				DrawParams.DrawFeatureVector(CurResult.PoseIdx);
 			}
 
 			if (bDebugDrawQuery)
 			{
-				FDebugDrawParams DrawParams(AnimInstanceProxies, SearchContext.GetPoseHistories(), SearchContext.GetRoleToIndex(), CurResultDatabase, EDebugDrawFlags::DrawQuery);
+				FDebugDrawParams DrawParams(SearchContext.GetAnimContexts(), SearchContext.GetPoseHistories(), SearchContext.GetRoleToIndex(), CurResultDatabase, EDebugDrawFlags::DrawQuery);
 				DrawParams.DrawFeatureVector(SearchContext.GetOrBuildQuery(CurResultDatabase->Schema));
 			}
 		}
@@ -793,12 +791,15 @@ void UPoseSearchLibrary::MotionMatch(
 	FMemMark Mark(FMemStack::Get());
 
 	TArray<const UE::PoseSearch::IPoseHistory*, TInlineAllocator<PreallocatedRolesNum, TMemStackAllocator<>>> PoseHistories;
+	TArray<const UObject*, TInlineAllocator<PreallocatedRolesNum, TMemStackAllocator<>>> AnimContexts;
 	for (UAnimInstance* AnimInstance : AnimInstances)
 	{
 		if (const FAnimNode_PoseSearchHistoryCollector_Base* PoseHistoryNode = FindPoseHistoryNode(PoseHistoryName, AnimInstance))
 		{
 			PoseHistories.Add(&PoseHistoryNode->GetPoseHistory());
 		}
+
+		AnimContexts.Add(AnimInstance);
 	}
 
 	if (PoseHistories.Num() != AnimInstances.Num())
@@ -807,7 +808,7 @@ void UPoseSearchLibrary::MotionMatch(
 		return;
 	}
 
-	const FSearchResult SearchResult = MotionMatch(AnimInstances, Roles, PoseHistories, AssetsToSearch, ContinuingProperties, Future);
+	const FSearchResult SearchResult = MotionMatch(AnimContexts, Roles, PoseHistories, AssetsToSearch, ContinuingProperties, Future);
 	if (SearchResult.IsValid())
 	{
 		const UPoseSearchDatabase* Database = SearchResult.Database.Get();
@@ -843,14 +844,14 @@ void UPoseSearchLibrary::MotionMatch(
 }
 
 UE::PoseSearch::FSearchResult UPoseSearchLibrary::MotionMatch(
-	const TArrayView<UAnimInstance*> AnimInstances,
+	const TArrayView<const UObject*> AnimContexts,
 	const TArrayView<const UE::PoseSearch::FRole> Roles,
 	const TArrayView<const UE::PoseSearch::IPoseHistory*> PoseHistories, 
 	const TArrayView<const UObject*> AssetsToSearch,
 	const FPoseSearchContinuingProperties& ContinuingProperties,
 	const FPoseSearchFutureProperties& Future)
 {
-	check(!AnimInstances.IsEmpty() && AnimInstances.Num() == Roles.Num() && AnimInstances.Num() == PoseHistories.Num());
+	check(!AnimContexts.IsEmpty() && AnimContexts.Num() == Roles.Num() && AnimContexts.Num() == PoseHistories.Num());
 
 	using namespace UE::PoseSearch;
 
@@ -884,57 +885,57 @@ UE::PoseSearch::FSearchResult UPoseSearchLibrary::MotionMatch(
 
 		for (int32 RoleIndex = 0; RoleIndex < Roles.Num(); ++RoleIndex)
 		{
-			MemStackPoseHistories[RoleIndex].Init(InternalPoseHistories[RoleIndex]);
+			if (const UAnimInstance* AnimInstance = Cast<UAnimInstance>(AnimContexts[RoleIndex]))
+			{
+				MemStackPoseHistories[RoleIndex].Init(InternalPoseHistories[RoleIndex]);
 
-			// extracting 2 poses to be able to calculate velocities
-			FCSPose<FCompactPose> ComponentSpacePose;
-			FCompactPose Pose;
-			FBlendedCurve Curves;
-			Pose.SetBoneContainer(&AnimInstances[RoleIndex]->GetRequiredBonesOnAnyThread());
+				// extracting 2 poses to be able to calculate velocities
+				FCSPose<FCompactPose> ComponentSpacePose;
+				FCompactPose Pose;
+				FBlendedCurve Curves;
+				Pose.SetBoneContainer(&AnimInstance->GetRequiredBonesOnAnyThread());
 
-			// @todo: add input BlendParameters to support sampling FutureAnimation blendspaces and support for multi character
-			const UAnimationAsset* AnimationAsset = Cast<UAnimationAsset>(Future.Animation);
-			if (!AnimationAsset)
-			{
-				if (const UMultiAnimAsset* MultiAnimAsset = Cast<UMultiAnimAsset>(Future.Animation))
+				// @todo: add input BlendParameters to support sampling FutureAnimation blendspaces and support for multi character
+				const UAnimationAsset* AnimationAsset = Cast<UAnimationAsset>(Future.Animation);
+				if (!AnimationAsset)
 				{
-					AnimationAsset = MultiAnimAsset->GetAnimationAsset(Roles[RoleIndex]);
+					if (const UMultiAnimAsset* MultiAnimAsset = Cast<UMultiAnimAsset>(Future.Animation))
+					{
+						AnimationAsset = MultiAnimAsset->GetAnimationAsset(Roles[RoleIndex]);
+					}
+					else
+					{
+						checkNoEntry();
+					}
 				}
-				else
+
+				const FAnimationAssetSampler Sampler(AnimationAsset);
+				for (int32 i = 0; i < 2; ++i)
 				{
-					checkNoEntry();
-				}
-			}
+					const float FuturePoseExtractionTime = FutureAnimationTime + (i - 1) * FiniteDelta;
+					const float FuturePoseAnimationTime = FutureIntervalTime + (i - 1) * FiniteDelta;
 
-			const FAnimationAssetSampler Sampler(AnimationAsset);
-			for (int32 i = 0; i < 2; ++i)
-			{
-				const float FuturePoseExtractionTime = FutureAnimationTime + (i - 1) * FiniteDelta;
-				const float FuturePoseAnimationTime = FutureIntervalTime + (i - 1) * FiniteDelta;
+					Sampler.ExtractPose(FuturePoseExtractionTime, Pose, Curves);
+					ComponentSpacePose.InitPose(Pose);
+					MemStackPoseHistories[RoleIndex].AddFuturePose(FuturePoseAnimationTime, ComponentSpacePose);
+				}
 
-				Sampler.ExtractPose(FuturePoseExtractionTime, Pose, Curves);
-				ComponentSpacePose.InitPose(Pose);
-				MemStackPoseHistories[RoleIndex].AddFuturePose(FuturePoseAnimationTime, ComponentSpacePose);
+				InternalPoseHistories[RoleIndex] = MemStackPoseHistories[RoleIndex].GetThisOrPoseHistory();
 			}
-
-#if ENABLE_DRAW_DEBUG && ENABLE_ANIM_DEBUG
-			if (FAnimInstanceProxy* AnimInstanceProxy = UAnimInstanceProxyProvider::GetAnimInstanceProxy(AnimInstances[RoleIndex]))
+			else
 			{
-				MemStackPoseHistories[RoleIndex].DebugDraw(*AnimInstanceProxy, FColor::Orange);
+				// @todo: implement FPoseSearchFutureProperties support for AnimNext
+				UE_LOG(LogPoseSearch, Error, TEXT("UPoseSearchLibrary::MotionMatch - FPoseSearchFutureProperties is supported only for UAnimInstance AnimContexts! (Unimplemented)"));
 			}
-#endif // ENABLE_DRAW_DEBUG && ENABLE_ANIM_DEBUG
-
-			InternalPoseHistories[RoleIndex] = MemStackPoseHistories[RoleIndex].GetThisOrPoseHistory();
-		}
+		}		
 	}
 
 	FSearchResult ReconstructedPreviousSearchResult;
 	FSearchContext SearchContext(FutureIntervalTime, nullptr, ReconstructedPreviousSearchResult);
 
-	// @todo: all assets in AssetsToSearch should have a consistent Roles requirements, or else the search will throw an error!
 	for (int32 RoleIndex = 0; RoleIndex < Roles.Num(); ++RoleIndex)
 	{
-		SearchContext.AddRole(Roles[RoleIndex], AnimInstances[RoleIndex], InternalPoseHistories[RoleIndex]);
+		SearchContext.AddRole(Roles[RoleIndex], AnimContexts[RoleIndex], InternalPoseHistories[RoleIndex]);
 	}
 
 	TAssetsToSearchPerDatabaseMap AssetsToSearchPerDatabaseMap;
@@ -942,7 +943,12 @@ UE::PoseSearch::FSearchResult UPoseSearchLibrary::MotionMatch(
 	bool bAsyncBuildIndexInProgress = false;
 
 	// collecting all the possible continuing pose search (it could be multiple searches, but most likely only one)
-	const float DeltaSeconds = AnimInstances[0] ? AnimInstances[0]->GetDeltaSeconds() : FiniteDelta;
+	float DeltaSeconds = FiniteDelta;
+	if (const UAnimInstance* AnimInstance = Cast<UAnimInstance>(AnimContexts[0]))
+	{
+		DeltaSeconds = AnimInstance->GetDeltaSeconds();
+	}
+
 	if (const UObject* PlayingAnimationAsset = ContinuingProperties.PlayingAsset.Get())
 	{
 		// checking if PlayingAnimationAsset has an associated database
@@ -1084,23 +1090,15 @@ UE::PoseSearch::FSearchResult UPoseSearchLibrary::MotionMatch(
 
 		if (bDrawMatch || bDrawquery)
 		{
-			TArray<FAnimInstanceProxy*, TInlineAllocator<PreallocatedRolesNum, TMemStackAllocator<>>> AnimInstanceProxies;
-			AnimInstanceProxies.SetNum(Roles.Num());
-			
-			for (int32 RoleIndex = 0; RoleIndex < Roles.Num(); ++RoleIndex)
-			{
-				AnimInstanceProxies[RoleIndex] = UAnimInstanceProxyProvider::GetAnimInstanceProxy(AnimInstances[RoleIndex]);
-			}
-
 			if (bDrawMatch)
 			{
-				FDebugDrawParams DrawParams(AnimInstanceProxies, SearchContext.GetPoseHistories(), SearchContext.GetRoleToIndex(), SearchResult.Database.Get());
+				FDebugDrawParams DrawParams(SearchContext.GetAnimContexts(), SearchContext.GetPoseHistories(), SearchContext.GetRoleToIndex(), SearchResult.Database.Get());
 				DrawParams.DrawFeatureVector(SearchResult.PoseIdx);
 			}
 
 			if (bDrawquery)
 			{
-				FDebugDrawParams DrawParams(AnimInstanceProxies, SearchContext.GetPoseHistories(), SearchContext.GetRoleToIndex(), SearchResult.Database.Get(), EDebugDrawFlags::DrawQuery);
+				FDebugDrawParams DrawParams(SearchContext.GetAnimContexts(), SearchContext.GetPoseHistories(), SearchContext.GetRoleToIndex(), SearchResult.Database.Get(), EDebugDrawFlags::DrawQuery);
 				DrawParams.DrawFeatureVector(SearchContext.GetOrBuildQuery(SearchResult.Database->Schema));
 			}
 		}
@@ -1111,7 +1109,12 @@ UE::PoseSearch::FSearchResult UPoseSearchLibrary::MotionMatch(
 	FMotionMatchingState MotionMatchingState;
 	MotionMatchingState.CurrentSearchResult = SearchResult;
 	MotionMatchingState.ElapsedPoseSearchTime = 0.0f;
-	TraceMotionMatching(SearchContext, MotionMatchingState, DeltaSeconds, true, FObjectTrace::GetWorldElapsedTime(AnimInstances[0] ? AnimInstances[0]->GetWorld() : nullptr));
+	double WorldElapsedTime = 0;
+	if (const UAnimInstance* AnimInstance = Cast<UAnimInstance>(AnimContexts[0]))
+	{
+		WorldElapsedTime = FObjectTrace::GetWorldElapsedTime(AnimInstance->GetWorld());
+	}
+	TraceMotionMatching(SearchContext, MotionMatchingState, DeltaSeconds, true, WorldElapsedTime);
 #endif // UE_POSE_SEARCH_TRACE_ENABLED
 
 	return SearchResult;
@@ -1162,6 +1165,25 @@ const FAnimNode_PoseSearchHistoryCollector_Base* UPoseSearchLibrary::FindPoseHis
 
 ///////////////////////////////////////////////////////////////////////////////////////////
 // Begin deprecated signatures
+UE::PoseSearch::FSearchResult UPoseSearchLibrary::MotionMatch(
+	const TArrayView<UAnimInstance*> AnimInstances,
+	const TArrayView<const UE::PoseSearch::FRole> Roles,
+	const TArrayView<const UE::PoseSearch::IPoseHistory*> PoseHistories,
+	const TArrayView<const UObject*> AssetsToSearch,
+	const FPoseSearchContinuingProperties& ContinuingProperties,
+	const FPoseSearchFutureProperties& Future)
+{
+	TArray<const UObject*, TInlineAllocator<UE::PoseSearch::PreallocatedRolesNum>> AnimContexts;
+	AnimContexts.Reserve(AnimInstances.Num());
+
+	for (UAnimInstance* AnimInstance : AnimInstances)
+	{
+		AnimContexts.Add(AnimInstance);
+	}
+
+	return MotionMatch(AnimContexts, Roles, PoseHistories, AssetsToSearch, ContinuingProperties, Future);
+}
+
 void UPoseSearchLibrary::MotionMatch(
 	const TArrayView<UAnimInstance*> AnimInstances,
 	const TArrayView<const UE::PoseSearch::FRole> Roles,
@@ -1184,7 +1206,15 @@ UE::PoseSearch::FSearchResult UPoseSearchLibrary::MotionMatch(
 	const FPoseSearchFutureProperties& Future,
 	const int32 DebugSessionUniqueIdentifier)
 {
-	return MotionMatch(AnimInstances, Roles, PoseHistories, AssetsToSearch, ContinuingProperties, Future);
+	TArray<const UObject*, TInlineAllocator<UE::PoseSearch::PreallocatedRolesNum>> AnimContexts;
+	AnimContexts.Reserve(AnimInstances.Num());
+
+	for (UAnimInstance* AnimInstance : AnimInstances)
+	{
+		AnimContexts.Add(AnimInstance);
+	}
+
+	return MotionMatch(AnimContexts, Roles, PoseHistories, AssetsToSearch, ContinuingProperties, Future);
 }
 
 UE::PoseSearch::FSearchResult UPoseSearchLibrary::MotionMatch(
@@ -1200,10 +1230,8 @@ UE::PoseSearch::FSearchResult UPoseSearchLibrary::MotionMatch(
 		PoseHistory = &PoseHistoryProvider->GetPoseHistory();
 	}
 
-	UAnimInstance* AnimInstance = Cast<UAnimInstance>(Context.AnimInstanceProxy->GetAnimInstanceObject());
-	check(AnimInstance);
-
-	return MotionMatch(MakeArrayView(&AnimInstance, 1), MakeArrayView(&DefaultRole, 1), MakeArrayView(&PoseHistory, 1),
+	const UObject* AnimContext = Context.AnimInstanceProxy->GetAnimInstanceObject();
+	return MotionMatch(MakeArrayView(&AnimContext, 1), MakeArrayView(&DefaultRole, 1), MakeArrayView(&PoseHistory, 1),
 		AssetsToSearch, ContinuingProperties, FPoseSearchFutureProperties());
 }
 		
@@ -1216,7 +1244,15 @@ UE::PoseSearch::FSearchResult UPoseSearchLibrary::MotionMatch(
 	const int32 DebugSessionUniqueIdentifier,
 	float DesiredPermutationTimeOffset)
 {
-	return MotionMatch(AnimInstances, Roles, PoseHistories, AssetsToSearch, ContinuingProperties, FPoseSearchFutureProperties());
+	TArray<const UObject*, TInlineAllocator<UE::PoseSearch::PreallocatedRolesNum>> AnimContexts;
+	AnimContexts.Reserve(AnimInstances.Num());
+
+	for (UAnimInstance* AnimInstance : AnimInstances)
+	{
+		AnimContexts.Add(AnimInstance);
+	}
+
+	return MotionMatch(AnimContexts, Roles, PoseHistories, AssetsToSearch, ContinuingProperties, FPoseSearchFutureProperties());
 }
 
 // End deprecated signatures
diff --git a/Engine/Plugins/Animation/PoseSearch/Source/Runtime/Private/PoseSearchModule.cpp b/Engine/Plugins/Animation/PoseSearch/Source/Runtime/Private/PoseSearchModule.cpp
--- a/Engine/Plugins/Animation/PoseSearch/Source/Runtime/Private/PoseSearchModule.cpp
+++ b/Engine/Plugins/Animation/PoseSearch/Source/Runtime/Private/PoseSearchModule.cpp
@@ -63,10 +63,8 @@ public:
 			PoseHistory = &PoseHistoryProvider->GetPoseHistory();
 		}
 
-		UAnimInstance* AnimInstance = Cast<UAnimInstance>(GraphContext.AnimInstanceProxy->GetAnimInstanceObject());
-		check(AnimInstance);
-		
-		const UE::PoseSearch::FSearchResult SearchResult = UPoseSearchLibrary::MotionMatch(MakeArrayView(&AnimInstance, 1),
+		const UObject* AnimContext = GraphContext.AnimInstanceProxy->GetAnimInstanceObject();
+		const UE::PoseSearch::FSearchResult SearchResult = UPoseSearchLibrary::MotionMatch(MakeArrayView(&AnimContext, 1),
 			MakeArrayView(&DefaultRole, 1), MakeArrayView(&PoseHistory, 1), AssetsToSearch, ContinuingProperties, Future);
 
 		UE::Anim::IPoseSearchProvider::FSearchResult ProviderResult;
diff --git a/Engine/Plugins/Animation/PoseSearch/Source/Runtime/Public/PoseSearch/PoseSearchContext.h b/Engine/Plugins/Animation/PoseSearch/Source/Runtime/Public/PoseSearch/PoseSearchContext.h
--- a/Engine/Plugins/Animation/PoseSearch/Source/Runtime/Public/PoseSearch/PoseSearchContext.h
+++ b/Engine/Plugins/Animation/PoseSearch/Source/Runtime/Public/PoseSearch/PoseSearchContext.h
@@ -53,9 +53,12 @@ ENUM_CLASS_FLAGS(EPoseCandidateFlags);
 #if ENABLE_DRAW_DEBUG
 struct POSESEARCH_API FDebugDrawParams
 {
+	UE_DEPRECATED(5.6, "Use other FDebugDrawParams signatures instead")
 	FDebugDrawParams(TArrayView<FAnimInstanceProxy*> InAnimInstanceProxies, TConstArrayView<const IPoseHistory*> InPoseHistories, const FRoleToIndex& InRoleToIndex, const UPoseSearchDatabase* InDatabase, EDebugDrawFlags InFlags = EDebugDrawFlags::None);
+	UE_DEPRECATED(5.6, "Use other FDebugDrawParams signatures instead")
 	FDebugDrawParams(TArrayView<const USkinnedMeshComponent*> InMeshes, TConstArrayView<const IPoseHistory*> InPoseHistories, const FRoleToIndex& InRoleToIndex, const UPoseSearchDatabase* InDatabase, EDebugDrawFlags InFlags = EDebugDrawFlags::None);
-
+	FDebugDrawParams(TConstArrayView<const UObject*> InAnimContexts, TConstArrayView<const IPoseHistory*> InPoseHistories, const FRoleToIndex& InRoleToIndex, const UPoseSearchDatabase* InDatabase, EDebugDrawFlags InFlags = EDebugDrawFlags::None);
+	
 	const FSearchIndex* GetSearchIndex() const;
 	const UPoseSearchSchema* GetSchema() const;
 
@@ -75,8 +78,7 @@ struct POSESEARCH_API FDebugDrawParams
 private:
 	bool CanDraw() const;
 
-	const TArrayView<FAnimInstanceProxy*> AnimInstanceProxies;
-	const TArrayView<const USkinnedMeshComponent*> Meshes;
+	const TConstArrayView<const UObject*> AnimContexts;
 	const TConstArrayView<const IPoseHistory*> PoseHistories;
 
 	// NoTe: mapping Role to the index of the associated asset that this FDebugDrawParams is drawing.
@@ -126,7 +128,7 @@ struct POSESEARCH_API FSearchContext
 	FSearchContext& operator=(const FSearchContext& Other) = delete;
 	FSearchContext& operator=(FSearchContext&& Other) = delete;
 
-	void AddRole(const FRole& Role, const UAnimInstance* AnimInstance, const IPoseHistory* PoseHistory);
+	void AddRole(const FRole& Role, const UObject* AnimContext, const IPoseHistory* PoseHistory);
 
 	// Returns the curve value of name CurveName at an offset time of SampleTimeOffset.
 	// If the curve is not found, assume value of 0, which is consistent with curve behavior in the animation update.
@@ -170,8 +172,22 @@ struct POSESEARCH_API FSearchContext
 	const IPoseHistory* GetPoseHistory(const FRole& Role) const { return PoseHistories[RoleToIndex[Role]]; }
 
 	float GetDesiredPermutationTimeOffset() const { return DesiredPermutationTimeOffset; }
-	const UAnimInstance* GetAnimInstance(const FRole& Role) const { return AnimInstances[RoleToIndex[Role]]; }
-	const TArray<const UAnimInstance*, TInlineAllocator<PreallocatedRolesNum, TMemStackAllocator<>>>& GetAnimInstances() const { return AnimInstances; }
+	
+	UE_DEPRECATED(5.6, "Use GetAnimContext instead")
+	const UAnimInstance* GetAnimInstance(const FRole& Role) const { return Cast<UAnimInstance>(AnimContexts[RoleToIndex[Role]]); }
+	const UObject* GetAnimContext(const FRole& Role) const { return AnimContexts[RoleToIndex[Role]]; }
+
+	UE_DEPRECATED(5.6, "Use GetAnimContexts instead")
+	const TArray<const UAnimInstance*, TInlineAllocator<PreallocatedRolesNum, TMemStackAllocator<>>> GetAnimInstances() const 
+	{
+		TArray<const UAnimInstance*, TInlineAllocator<PreallocatedRolesNum, TMemStackAllocator<>>> AnimInstances;
+		for (const UObject* AnimContext : AnimContexts)
+		{
+			AnimInstances.Add(Cast<UAnimInstance>(AnimContext));
+		}
+		return AnimInstances;
+	}
+	TConstArrayView<const UObject*> GetAnimContexts() const { return AnimContexts; }
 
 	const FRoleToIndex& GetRoleToIndex() const { return RoleToIndex; }
 
@@ -201,7 +217,7 @@ private:
 	FQuat GetSampleRotationInternal(float SampleTime, float OriginTime, int8 SchemaSampleBoneIdx, int8 SchemaOriginBoneIdx, const FRole& SampleRole, const FRole& OriginRole, const FQuat* SampleBoneRotationWorldOverride = nullptr);
 	FTransform GetWorldRootBoneTransformAtTime(float SampleTime, const FRole& SampleRole) const;
 	
-	TArray<const UAnimInstance*, TInlineAllocator<PreallocatedRolesNum, TMemStackAllocator<>>> AnimInstances;
+	TArray<const UObject*, TInlineAllocator<PreallocatedRolesNum, TMemStackAllocator<>>> AnimContexts;
 	TArray<const IPoseHistory*, TInlineAllocator<PreallocatedRolesNum, TMemStackAllocator<>>> PoseHistories;
 	FRoleToIndex RoleToIndex;
 	
diff --git a/Engine/Plugins/Animation/PoseSearch/Source/Runtime/Public/PoseSearch/PoseSearchInteractionIsland.h b/Engine/Plugins/Animation/PoseSearch/Source/Runtime/Public/PoseSearch/PoseSearchInteractionIsland.h
--- a/Engine/Plugins/Animation/PoseSearch/Source/Runtime/Public/PoseSearch/PoseSearchInteractionIsland.h
+++ b/Engine/Plugins/Animation/PoseSearch/Source/Runtime/Public/PoseSearch/PoseSearchInteractionIsland.h
@@ -15,7 +15,7 @@ struct FInteractionIsland;
 // Experimental, this feature might be removed without warning, not for production use
 struct FInteractionSearchContext
 {
-	// @todo: since AnimInstances, HistoryCollectors, Roles, and BroadPhaseRadiuses have the same cardinality,
+	// @todo: since AnimContexts, HistoryCollectors, Roles, and BroadPhaseRadiuses have the same cardinality,
 	// add an array of structs instead of many arrays
 	TArray<TWeakObjectPtr<UAnimInstance>> AnimInstances;
 	// @todo: perhaps use the animnode of the pose history collector instead
diff --git a/Engine/Plugins/Animation/PoseSearch/Source/Runtime/Public/PoseSearch/PoseSearchLibrary.h b/Engine/Plugins/Animation/PoseSearch/Source/Runtime/Public/PoseSearch/PoseSearchLibrary.h
--- a/Engine/Plugins/Animation/PoseSearch/Source/Runtime/Public/PoseSearch/PoseSearchLibrary.h
+++ b/Engine/Plugins/Animation/PoseSearch/Source/Runtime/Public/PoseSearch/PoseSearchLibrary.h
@@ -237,6 +237,15 @@ public:
 		const FPoseSearchFutureProperties& Future,
 		FPoseSearchBlueprintResult& Result);
 
+	static UE::PoseSearch::FSearchResult MotionMatch(
+		const TArrayView<const UObject*> AnimContexts,
+		const TArrayView<const UE::PoseSearch::FRole> Roles,
+		const TArrayView<const UE::PoseSearch::IPoseHistory*> PoseHistories, 
+		const TArrayView<const UObject*> AssetsToSearch,
+		const FPoseSearchContinuingProperties& ContinuingProperties,
+		const FPoseSearchFutureProperties& Future);
+		
+	UE_DEPRECATED(5.5, "Use other MotionMatch signatures instead")
 	static UE::PoseSearch::FSearchResult MotionMatch(
 		const TArrayView<UAnimInstance*> AnimInstances,
 		const TArrayView<const UE::PoseSearch::FRole> Roles,
diff --git a/Engine/Plugins/Animation/PoseSearch/Source/Runtime/Public/PoseSearch/Trace/PoseSearchTraceLogger.h b/Engine/Plugins/Animation/PoseSearch/Source/Runtime/Public/PoseSearch/Trace/PoseSearchTraceLogger.h
--- a/Engine/Plugins/Animation/PoseSearch/Source/Runtime/Public/PoseSearch/Trace/PoseSearchTraceLogger.h
+++ b/Engine/Plugins/Animation/PoseSearch/Source/Runtime/Public/PoseSearch/Trace/PoseSearchTraceLogger.h
@@ -25,6 +25,7 @@ struct POSESEARCH_API FTraceMessage
 {
 	uint64 Cycle = 0;
 
+	// @todo: rename it to AnimContextId
 	uint64 AnimInstanceId = 0;
 
 	// motion matching Search Id associated with this message
diff --git a/Engine/Plugins/Experimental/AnimNextPoseSearch/Source/AnimNextPoseSearch/Private/MotionMatchingTrait.cpp b/Engine/Plugins/Experimental/AnimNextPoseSearch/Source/AnimNextPoseSearch/Private/MotionMatchingTrait.cpp
--- a/Engine/Plugins/Experimental/AnimNextPoseSearch/Source/AnimNextPoseSearch/Private/MotionMatchingTrait.cpp
+++ b/Engine/Plugins/Experimental/AnimNextPoseSearch/Source/AnimNextPoseSearch/Private/MotionMatchingTrait.cpp
@@ -93,13 +93,18 @@ namespace UE::AnimNext
 			}
 		}
 
-		TArray<UAnimInstance*> AnimInstances; AnimInstances.Add(nullptr);
+		const UObject* AnimContext = nullptr;
+		if (const FAnimNextModuleInstance* ModuleInstance = Context.GetRootGraphInstance().GetModuleInstance())
+		{
+			AnimContext = ModuleInstance->Object;
+		}
+
+		const TArrayView<const UObject*> AnimContexts = MakeArrayView(&AnimContext, 1);
 		const TArrayView<const FRole> Roles = MakeArrayView(&DefaultRole, 1);
 		const TArrayView<const UE::PoseSearch::IPoseHistory*> PoseHistories = MakeArrayView(&PoseHistory, 1);
 		const FPoseSearchFutureProperties FutureProperties;
 
-		const FSearchResult SearchResult = UPoseSearchLibrary::MotionMatch(AnimInstances, Roles, PoseHistories,  AssetsToSearch, ContinuingProperties, FutureProperties);
-
+		const FSearchResult SearchResult = UPoseSearchLibrary::MotionMatch(AnimContexts, Roles, PoseHistories,  AssetsToSearch, ContinuingProperties, FutureProperties);
 		if (!SearchResult.bIsContinuingPoseSearch)
 		{
 			if (const UE::PoseSearch::FSearchIndexAsset* SearchIndexAsset = SearchResult.GetSearchIndexAsset())
