pose search - MM interactions improvements

diff --git a/Engine/Plugins/Animation/PoseSearch/Source/Runtime/Private/PoseSearchInteractionIsland.cpp b/Engine/Plugins/Animation/PoseSearch/Source/Runtime/Private/PoseSearchInteractionIsland.cpp
--- a/Engine/Plugins/Animation/PoseSearch/Source/Runtime/Private/PoseSearchInteractionIsland.cpp
+++ b/Engine/Plugins/Animation/PoseSearch/Source/Runtime/Private/PoseSearchInteractionIsland.cpp
@@ -34,7 +34,7 @@ static FInteractionSearchResult InitSearchResult(const FSearchResult& SearchResu
 		{
 			if (const UMultiAnimAsset* MultiAnimAsset = Cast<UMultiAnimAsset>(DatabaseAnimationAssetBase->GetAnimationAsset()))
 			{
-				check(MultiAnimAsset->GetNumRoles() == SearchContext.AnimInstances.Num());
+				check(MultiAnimAsset->GetNumRoles() == SearchContext.AnimContexts.Num());
 				check(MultiAnimAsset->GetNumRoles() == SearchContext.HistoryCollectors.Num());
 				check(MultiAnimAsset->GetNumRoles() == SearchContext.Roles.Num());
 
@@ -52,10 +52,17 @@ static FInteractionSearchResult InitSearchResult(const FSearchResult& SearchResu
 					const UE::PoseSearch::FRole& InteractionAssetRole = MultiAnimAsset->GetRole(MultiAnimAssetRoleIndex);
 					const int32 SearchContextIndex = SearchContextRoleToIndex[InteractionAssetRole];
 
-					if (UAnimInstance* AnimInstance = SearchContext.AnimInstances[SearchContextIndex].Get())
+					if (const UObject* AnimContext = SearchContext.AnimContexts[SearchContextIndex].Get())
 					{
-						const FTransform RootBoneTransform = AnimInstance->GetSkelMeshComponent()->GetBoneTransform(0);
-						ActorRootBoneTransforms.Add(RootBoneTransform);
+						if (const UAnimInstance* AnimInstance = Cast<UAnimInstance>(AnimContext))
+						{
+							const FTransform RootBoneTransform = AnimInstance->GetSkelMeshComponent()->GetBoneTransform(0);
+							ActorRootBoneTransforms.Add(RootBoneTransform);
+						}
+						else
+						{
+							unimplemented();
+						}
 					}
 				}
 
@@ -82,14 +89,21 @@ static FInteractionSearchResult InitSearchResult(const FSearchResult& SearchResu
 			else
 			{
 				// support for non UMultiAnimAsset to be backward compatible with regular motion matching searches
-				check(SearchContext.AnimInstances.Num() == 1);
+				check(SearchContext.AnimContexts.Num() == 1);
 
-				if (UAnimInstance* AnimInstance = SearchContext.AnimInstances[0].Get())
+				if (const UObject* AnimContext = SearchContext.AnimContexts[0].Get())
 				{
-					const FTransform RootBoneTransform = AnimInstance->GetSkelMeshComponent()->GetBoneTransform(0);
-					InteractionSearchResult.FullAlignedActorRootBoneTransforms.SetNum(1, EAllowShrinking::No);
-					// @todo: should it be RootBoneTransform + root motion transform?
-					InteractionSearchResult.FullAlignedActorRootBoneTransforms[0] = RootBoneTransform;
+					if (const UAnimInstance* AnimInstance = Cast<UAnimInstance>(AnimContext))
+					{
+						const FTransform RootBoneTransform = AnimInstance->GetSkelMeshComponent()->GetBoneTransform(0);
+						InteractionSearchResult.FullAlignedActorRootBoneTransforms.SetNum(1, EAllowShrinking::No);
+						// @todo: should it be RootBoneTransform + root motion transform?
+						InteractionSearchResult.FullAlignedActorRootBoneTransforms[0] = RootBoneTransform;
+					}
+					else
+					{
+						unimplemented();
+					}
 				}
 			}
 		}
@@ -139,7 +153,7 @@ bool FInteractionSearchContext::IsValid() const
 		return false;
 	}
 
-	const int32 Num = AnimInstances.Num();
+	const int32 Num = AnimContexts.Num();
 	if (Num < 1)
 	{
 		return false;
@@ -157,14 +171,14 @@ bool FInteractionSearchContext::IsValid() const
 
 	for (int32 IndexA = 0; IndexA < Num; ++IndexA)
 	{
-		if (AnimInstances[IndexA] == nullptr)
+		if (AnimContexts[IndexA] == nullptr)
 		{
 			return false;
 		}
 
 		for (int32 IndexB = IndexA + 1; IndexB < Num; ++IndexB)
 		{
-			if (AnimInstances[IndexA] == AnimInstances[IndexB])
+			if (AnimContexts[IndexA] == AnimContexts[IndexB])
 			{
 				return false;
 			}
@@ -208,28 +222,28 @@ bool FInteractionSearchContext::IsEquivalent(const FInteractionSearchContext& Ot
 		return false;
 	}
 
-	const int32 Num = AnimInstances.Num();
-	if (Num != Other.AnimInstances.Num())
+	const int32 Num = AnimContexts.Num();
+	if (Num != Other.AnimContexts.Num())
 	{
 		return false;
 	}
 
-	int32 CommonRoledAnimInstances = 0;
+	int32 CommonRoledAnimContexts = 0;
 	for (int32 IndexThis = 0; IndexThis < Num; ++IndexThis)
 	{
 		for (int32 IndexOther = 0; IndexOther < Num; ++IndexOther)
 		{
-			if (AnimInstances[IndexThis] == Other.AnimInstances[IndexOther] &&
+			if (AnimContexts[IndexThis] == Other.AnimContexts[IndexOther] &&
 				HistoryCollectors[IndexThis] == Other.HistoryCollectors[IndexOther] &&
 				Roles[IndexThis] == Other.Roles[IndexOther])
 			{
-				++CommonRoledAnimInstances;
+				++CommonRoledAnimContexts;
 			}
 		}
 	}
 
-	// using >= in case there are duplicated animinstances in this or Other (this->IsValid() should be false!)
-	if (CommonRoledAnimInstances >= Num)
+	// using >= in case there are duplicated AnimContexts in this or Other (this->IsValid() should be false!)
+	if (CommonRoledAnimContexts >= Num)
 	{
 		return true;
 	}
@@ -254,14 +268,21 @@ void FInteractionIsland::FPreTickFunction::ExecuteTick(float DeltaTime, enum ELe
 	// generating trajectories before running any of the skeletal mesh component ticks
 	for (const FInteractionSearchContext& SearchContext : Island->SearchContexts)
 	{
-		for (int32 Index = 0; Index < SearchContext.AnimInstances.Num(); ++Index)
+		for (int32 Index = 0; Index < SearchContext.AnimContexts.Num(); ++Index)
 		{
-			if (const UAnimInstance* AnimInstance = SearchContext.AnimInstances[Index].Get())
+			if (const UObject* AnimContext = SearchContext.AnimContexts[Index].Get())
 			{
-				// since FInteractionIsland has a tick dependency with the USkeletalMeshComponent it's safe modify the FAnimNode_PoseSearchHistoryCollector_Base
-				FAnimNode_PoseSearchHistoryCollector_Base* HistoryCollector = const_cast<FAnimNode_PoseSearchHistoryCollector_Base*>(SearchContext.HistoryCollectors[Index]);
-				check(HistoryCollector);
-				HistoryCollector->GenerateTrajectory(AnimInstance);
+				if (const UAnimInstance* AnimInstance = Cast<UAnimInstance>(AnimContext))
+				{
+					// since FInteractionIsland has a tick dependency with the USkeletalMeshComponent it's safe modify the FAnimNode_PoseSearchHistoryCollector_Base
+					FAnimNode_PoseSearchHistoryCollector_Base* HistoryCollector = const_cast<FAnimNode_PoseSearchHistoryCollector_Base*>(SearchContext.HistoryCollectors[Index]);
+					check(HistoryCollector);
+					HistoryCollector->GenerateTrajectory(AnimInstance);
+				}
+				else
+				{
+					unimplemented();
+				}
 			}
 		}
 	}
@@ -310,13 +331,20 @@ void FInteractionIsland::DebugDraw(const FColor& Color) const
 	{
 		for (const FInteractionSearchContext& SearchContext : SearchContexts)
 		{
-			for (int32 Index = 0; Index < SearchContext.AnimInstances.Num(); ++Index)
+			for (int32 Index = 0; Index < SearchContext.AnimContexts.Num(); ++Index)
 			{
-				if (const UAnimInstance* AnimInstance = SearchContext.AnimInstances[Index].Get())
+				if (const UObject* AnimContext = SearchContext.AnimContexts[Index].Get())
 				{
-					const FVector Position = AnimInstance->GetSkelMeshComponent()->GetComponentLocation();
-					const float BroadPhaseRadius = SearchContext.BroadPhaseRadiuses[Index];
-					DrawDebugCircle(AnimInstance->GetWorld(), Position, BroadPhaseRadius, 40, Color, false, 0.f, SDPG_Foreground, 0.f, FVector::XAxisVector, FVector::YAxisVector, false);
+					if (const UAnimInstance* AnimInstance = Cast<UAnimInstance>(AnimContext))
+					{
+						const FVector Position = AnimInstance->GetSkelMeshComponent()->GetComponentLocation();
+						const float BroadPhaseRadius = SearchContext.BroadPhaseRadiuses[Index];
+						DrawDebugCircle(AnimInstance->GetWorld(), Position, BroadPhaseRadius, 40, Color, false, 0.f, SDPG_Foreground, 0.f, FVector::XAxisVector, FVector::YAxisVector, false);
+					}
+					else
+					{
+						unimplemented();
+					}
 				}
 			}
 		}
@@ -324,52 +352,88 @@ void FInteractionIsland::DebugDraw(const FColor& Color) const
 #endif // ENABLE_DRAW_DEBUG && ENABLE_ANIM_DEBUG
 }
 
-void FInteractionIsland::InjectToActor(AActor* Actor)
+UActorComponent* FInteractionIsland::FindPreTickComponent(const AActor* Actor)
+{
+	if (UActorComponent* CharacterMovementComponent = Actor->GetComponentByClass<UCharacterMovementComponent>())
+	{
+		return CharacterMovementComponent;
+	}
+	
+	FName Tag = "MotionMatchingInteractionPre";
+	if (UActorComponent* PreTickComponent = Actor->FindComponentByTag<UActorComponent>(Tag))
+	{
+		return PreTickComponent;
+	}
+
+	UE_LOG(LogPoseSearch, Error, TEXT("FInteractionIsland::FindPreTickComponent requires AActor %s to have a UCharacterMovementComponent, or a component tagged as 'MotionMatchingInteractionPre' to work!"), *Actor->GetName());
+	return nullptr;
+}
+
+UActorComponent* FInteractionIsland::FindPostTickComponent(const AActor* Actor)
+{
+	if (UActorComponent* SkeletalMeshComponent = Actor->GetComponentByClass<USkeletalMeshComponent>())
+	{
+		return SkeletalMeshComponent;
+	}
+
+	FName Tag = "MotionMatchingInteractionPost";
+	if (UActorComponent* PostTickComponent = Actor->FindComponentByTag<UActorComponent>(Tag))
+	{
+		return PostTickComponent;
+
+	}
+
+	UE_LOG(LogPoseSearch, Error, TEXT("FInteractionIsland::FindPostTickComponent requires AActor %s to have a USkeletalMeshComponent, or a component tagged as 'MotionMatchingInteractionPost' to work!"), *Actor->GetName());
+	return nullptr;
+}
+
+void FInteractionIsland::InjectToActor(const UObject* AnimContext)
 {
 	check(IsInGameThread());
 
 	// Called by UPoseSearchInteractionSubsystem::Tick when there aren't animation jobs flying. No need to FScopeLock Lock(&Mutex);
-	if (Actor)
+	if (AnimContext)
 	{
-		if (UCharacterMovementComponent* CharacterMovementComponent = Actor->GetComponentByClass<UCharacterMovementComponent>())
+		if (const UAnimInstance* AnimInstance = Cast<UAnimInstance>(AnimContext))
 		{
-			if (USkeletalMeshComponent* SkeletalMeshComponent = Actor->GetComponentByClass<USkeletalMeshComponent>())
+			AActor* Actor = AnimInstance->GetOwningActor();
+
+			if (UActorComponent* PreTickComponent = FindPreTickComponent(Actor))
 			{
-				const bool bIsFirstInjectedActor = IsUninjected();
-
-				//	tick order: 
-				//		CharacterMovementComponent(s) ->
-				//			Island.PreTickFunction ->
-				//				first injected actor SkeletalMeshComponent ->
-				//					Island.PostTickFunction ->
-				//						other SkeletalMeshComponent(s)
-				CharacterMovementComponents.AddUnique(CharacterMovementComponent);
-				SkeletalMeshComponents.AddUnique(SkeletalMeshComponent);
-	
-				// making sure that if we add a unique CharacterMovementComponent, we add as well a unique SkeletalMeshComponent
-				// (so we can remove them later on in a consistent fashion)
-				check(CharacterMovementComponents.Num() == SkeletalMeshComponents.Num());
+				if (UActorComponent* PostTickComponent = FindPostTickComponent(Actor))
+				{
+					const bool bIsFirstInjectedActor = IsUninjected();
 
-				PreTickFunction.AddPrerequisite(Actor, CharacterMovementComponent->PrimaryComponentTick);
-				SkeletalMeshComponent->PrimaryComponentTick.AddPrerequisite(Actor, PreTickFunction);
+					//	tick order: 
+					//		PreTickComponents (UCharacterMovementComponent, or UCharacterMoverComponent) ->
+					//			Island.PreTickFunction ->
+					//				first injected actor PostTickComponent (USkeletalMeshComponent, or WIP UAnimNextComponent) ->
+					//					Island.PostTickFunction ->
+					//						other PostTickComponent(s)
+					PreTickComponents.AddUnique(PreTickComponent);
+					PostTickComponents.AddUnique(PostTickComponent);
 
-				if (bIsFirstInjectedActor)
-				{
-					PostTickFunction.AddPrerequisite(Actor, SkeletalMeshComponent->PrimaryComponentTick);
-				}
-				else 
-				{
-					SkeletalMeshComponent->PrimaryComponentTick.AddPrerequisite(Actor, PostTickFunction);
+					// making sure that if we add a unique PreTickComponent, we add as well a unique PostTickComponent
+					// (so we can remove them later on in a consistent fashion)
+					check(PreTickComponents.Num() == PostTickComponents.Num());
+
+					PreTickFunction.AddPrerequisite(Actor, PreTickComponent->PrimaryComponentTick);
+					PostTickComponent->PrimaryComponentTick.AddPrerequisite(Actor, PreTickFunction);
+
+					if (bIsFirstInjectedActor)
+					{
+						PostTickFunction.AddPrerequisite(Actor, PostTickComponent->PrimaryComponentTick);
+					}
+					else
+					{
+						PostTickComponent->PrimaryComponentTick.AddPrerequisite(Actor, PostTickFunction);
+					}
 				}
 			}
-			else
-			{
-				UE_LOG(LogPoseSearch, Error, TEXT("FInteractionIsland::InjectToActor requires AActor %s to have a USkeletalMeshComponent to work!"), *Actor->GetName());
-			}
 		}
 		else
 		{
-			UE_LOG(LogPoseSearch, Error, TEXT("FInteractionIsland::InjectToActor requires AActor %s to have a UCharacterMovementComponent to work!"), *Actor->GetName());
+			unimplemented();
 		}
 	}
 }
@@ -393,41 +457,41 @@ void FInteractionIsland::Uninject()
 	// Called by UPoseSearchInteractionSubsystem::Tick when there aren't animation jobs flying. No need to FScopeLock Lock(&Mutex);
 	check(IsInGameThread());
 
-	check(CharacterMovementComponents.Num() == SkeletalMeshComponents.Num());
+	check(PreTickComponents.Num() == PostTickComponents.Num());
 
-	for (int32 ActorIndex = 0; ActorIndex < CharacterMovementComponents.Num(); ++ActorIndex)
+	for (int32 ActorIndex = 0; ActorIndex < PreTickComponents.Num(); ++ActorIndex)
 	{
-		UCharacterMovementComponent* CharacterMovementComponent = CharacterMovementComponents[ActorIndex].Get();
-		USkeletalMeshComponent* SkeletalMeshComponent = SkeletalMeshComponents[ActorIndex].Get();
+		UActorComponent* PreTickComponent = PreTickComponents[ActorIndex].Get();
+		UActorComponent* PostTickComponent = PostTickComponents[ActorIndex].Get();
 
-		if (CharacterMovementComponent)
+		if (PreTickComponent)
 		{
-			check(SkeletalMeshComponent);
+			check(PostTickComponent);
 
-			AActor* Actor = CharacterMovementComponent->GetOwner();
+			AActor* Actor = PreTickComponent->GetOwner();
 			check(Actor);
 
-			PreTickFunction.RemovePrerequisite(Actor, CharacterMovementComponent->PrimaryComponentTick);
-			SkeletalMeshComponent->PrimaryComponentTick.RemovePrerequisite(Actor, PreTickFunction);
+			PreTickFunction.RemovePrerequisite(Actor, PreTickComponent->PrimaryComponentTick);
+			PostTickComponent->PrimaryComponentTick.RemovePrerequisite(Actor, PreTickFunction);
 
 			const bool bIsFirstInjectedActor = ActorIndex == 0;
 			if (bIsFirstInjectedActor)
 			{
-				PostTickFunction.RemovePrerequisite(Actor, SkeletalMeshComponent->PrimaryComponentTick);
+				PostTickFunction.RemovePrerequisite(Actor, PostTickComponent->PrimaryComponentTick);
 			}
 			else
 			{
-				SkeletalMeshComponent->PrimaryComponentTick.RemovePrerequisite(Actor, PostTickFunction);
+				PostTickComponent->PrimaryComponentTick.RemovePrerequisite(Actor, PostTickFunction);
 			}
 		}
 		else
 		{
-			check(!SkeletalMeshComponent);
+			check(!PostTickComponent);
 		}
 	}
 
-	CharacterMovementComponents.Reset();
-	SkeletalMeshComponents.Reset();
+	PreTickComponents.Reset();
+	PostTickComponents.Reset();
 
 	SearchContexts.Reset();
 	SearchResults.Reset();
@@ -436,16 +500,16 @@ void FInteractionIsland::Uninject()
 
 bool FInteractionIsland::IsUninjected() const
 {
-	return SkeletalMeshComponents.IsEmpty();
+	return PostTickComponents.IsEmpty();
 }
 
-bool FInteractionIsland::DoSearch_AnyThread(UObject* AnimInstance, const FPoseSearchContinuingProperties& ContinuingProperties, FPoseSearchInteractionBlueprintResult& Result)
+bool FInteractionIsland::DoSearch_AnyThread(UObject* AnimContext, const FPoseSearchContinuingProperties& ContinuingProperties, FPoseSearchInteractionBlueprintResult& Result)
 {
 	bool bDoPerfomSearch = false;
 	{
 		// thread safety note!
 		// goal:	avoiding deadlock between SearchResultsMutex lock and waiting for UAnimInstance::HandleExistingParallelEvaluationTask.
-		// why:		UPoseSearchLibrary::MotionMatch could call via AnimInstance GetProxyOnAnyThread<FAnimInstanceProxy>() that, if on GameThread, 
+		// why:		UPoseSearchLibrary::MotionMatch could call via AnimContext GetProxyOnAnyThread<FAnimInstanceProxy>() that, if on GameThread, 
 		//			could call UAnimInstance::HandleExistingParallelEvaluationTask 
 		// fix:		avoid UPoseSearchLibrary::MotionMatch calls wrapped by any lock, at the cost of eventually (by design should be NEVER) performing the searches twice
 		//			By design we should inject ticks dependencies (added by FInteractionIsland::InjectToActor via AddTickPrerequisiteComponent), so concurrently 
@@ -482,17 +546,16 @@ bool FInteractionIsland::DoSearch_AnyThread(UObject* AnimInstance, const FPoseSe
 			}
 
 			AnimContexts.Reset();
-			AnimContexts.Reserve(SearchContext.AnimInstances.Num());
-			for (const TWeakObjectPtr<UAnimInstance>& AnimInstancePtr : SearchContext.AnimInstances)
+			for (const TWeakObjectPtr<const UObject>& AnimContextPtr : SearchContext.AnimContexts)
 			{
-				UAnimInstance* SearchContextAnimInstance = AnimInstancePtr.Get();
-				if (!SearchContextAnimInstance)
+				const UObject* SearchContextAnimContext = AnimContextPtr.Get();
+				if (!SearchContextAnimContext)
 				{
-					UE_LOG(LogPoseSearch, Error, TEXT("FInteractionIsland::DoSearch_AnyThread null anim instance"));
+					UE_LOG(LogPoseSearch, Error, TEXT("FInteractionIsland::DoSearch_AnyThread null anim context"));
 					return false;
 				}
 
-				AnimContexts.Add(SearchContextAnimInstance);
+				AnimContexts.Add(SearchContextAnimContext);
 			}
 
 			PoseHistories.Reset();
@@ -573,25 +636,25 @@ bool FInteractionIsland::DoSearch_AnyThread(UObject* AnimInstance, const FPoseSe
 			}
 	
 			// calling this funtion within this scope since we already locked SearchResultsMutex
-			return GetResult_AnyThread(AnimInstance, Result);
+			return GetResult_AnyThread(AnimContext, Result);
 		}
 	}
 
-	return GetResult_AnyThread(AnimInstance, Result);
+	return GetResult_AnyThread(AnimContext, Result);
 }
 
-bool FInteractionIsland::GetResult_AnyThread(UObject* AnimInstance, FPoseSearchInteractionBlueprintResult& Result)
+bool FInteractionIsland::GetResult_AnyThread(UObject* AnimContext, FPoseSearchInteractionBlueprintResult& Result)
 {
 	// locking to read SearchResults
 	FScopeLock Lock(&SearchResultsMutex);
 
-	// looking for AnimInstance in SearchResults to fill up Result
+	// looking for AnimContext in SearchResults to fill up Result
 	for (const FInteractionSearchResult& SearchResult : SearchResults)
 	{
 		const FInteractionSearchContext& SearchContext = SearchContexts[SearchResult.SearchIndex];
-		for (int32 AnimInstanceIndex = 0; AnimInstanceIndex < SearchContext.AnimInstances.Num(); ++AnimInstanceIndex)
+		for (int32 AnimContextIndex = 0; AnimContextIndex < SearchContext.AnimContexts.Num(); ++AnimContextIndex)
 		{
-			if (SearchContext.AnimInstances[AnimInstanceIndex].Get() == AnimInstance)
+			if (SearchContext.AnimContexts[AnimContextIndex].Get() == AnimContext)
 			{
 				const UPoseSearchDatabase* Database = SearchResult.Database.Get();
 				check(Database);
@@ -602,7 +665,7 @@ bool FInteractionIsland::GetResult_AnyThread(UObject* AnimInstance, FPoseSearchI
 				const FPoseSearchDatabaseAnimationAssetBase* DatabaseAnimationAssetBase = Database->GetDatabaseAnimationAsset<FPoseSearchDatabaseAnimationAssetBase>(*SearchIndexAsset);
 				check(DatabaseAnimationAssetBase);
 
-				const UE::PoseSearch::FRole Role = SearchContext.Roles[AnimInstanceIndex];
+				const UE::PoseSearch::FRole Role = SearchContext.Roles[AnimContextIndex];
 
 				Result.SelectedAnimation = DatabaseAnimationAssetBase->GetAnimationAsset();
 				Result.SelectedTime = SearchResult.AssetTime;
@@ -614,9 +677,9 @@ bool FInteractionIsland::GetResult_AnyThread(UObject* AnimInstance, FPoseSearchI
 				Result.SearchCost = SearchResult.PoseCost.GetTotalCost();
 				Result.Role = Role;
 
-				if (SearchResult.FullAlignedActorRootBoneTransforms.IsValidIndex(AnimInstanceIndex))
+				if (SearchResult.FullAlignedActorRootBoneTransforms.IsValidIndex(AnimContextIndex))
 				{
-					Result.FullAlignedActorRootBoneTransform = SearchResult.FullAlignedActorRootBoneTransforms[AnimInstanceIndex];
+					Result.FullAlignedActorRootBoneTransform = SearchResult.FullAlignedActorRootBoneTransforms[AnimContextIndex];
 				}
 
 				//Result.bIsFromContinuingPlaying = SearchResult.bIsContinuingPoseSearch;
@@ -637,7 +700,7 @@ bool FInteractionIsland::GetResult_AnyThread(UObject* AnimInstance, FPoseSearchI
 				//	}
 				//}
 
-				// we found our AnimInstance: we can stop searching
+				// we found our AnimContext: we can stop searching
 				return true;
 			}
 		}
@@ -656,7 +719,7 @@ const FInteractionSearchResult* FInteractionIsland::FindSearchResult(const FInte
 	{
 		const FInteractionSearchContext& LocalSearchContext = SearchContexts[SearchResult.SearchIndex];
 		if (LocalSearchContext.Database == SearchContext.Database &&
-			LocalSearchContext.AnimInstances == SearchContext.AnimInstances &&
+			LocalSearchContext.AnimContexts == SearchContext.AnimContexts &&
 			LocalSearchContext.Roles == SearchContext.Roles)
 		{
 			return &SearchResult;
diff --git a/Engine/Plugins/Animation/PoseSearch/Source/Runtime/Private/PoseSearchInteractionSubsystem.cpp b/Engine/Plugins/Animation/PoseSearch/Source/Runtime/Private/PoseSearchInteractionSubsystem.cpp
--- a/Engine/Plugins/Animation/PoseSearch/Source/Runtime/Private/PoseSearchInteractionSubsystem.cpp
+++ b/Engine/Plugins/Animation/PoseSearch/Source/Runtime/Private/PoseSearchInteractionSubsystem.cpp
@@ -18,22 +18,29 @@ namespace UE::PoseSearch
 		return PoseHistoryName.ToString();
 	}
 	
-	const FAnimNode_PoseSearchHistoryCollector_Base* FInteractionAvailabilityEx::GetHistoryCollector(const UAnimInstance* AnimInstance) const
+	const FAnimNode_PoseSearchHistoryCollector_Base* FInteractionAvailabilityEx::GetHistoryCollector(const UObject* AnimContext) const
 	{
 		if (HistoryCollector)
 		{
 			return HistoryCollector;
 		}
-		return UPoseSearchLibrary::FindPoseHistoryNode(PoseHistoryName, AnimInstance);
+		
+		if (const UAnimInstance* AnimInstance = Cast<UAnimInstance>(AnimContext))
+		{
+			return UPoseSearchLibrary::FindPoseHistoryNode(PoseHistoryName, AnimInstance);
+		}
+
+		unimplemented();
+		return nullptr;
 	}
 
 	struct FInteractionSearchContextGroup
 	{
 		bool Contains(const FInteractionSearchContext& SearchContext) const
 		{
-			for (const TWeakObjectPtr<UAnimInstance>& AnimInstance : SearchContext.AnimInstances)
+			for (const TWeakObjectPtr<const UObject>& AnimContext : SearchContext.AnimContexts)
 			{
-				if (AnimInstances.Find(AnimInstance.Get()))
+				if (AnimContexts.Find(AnimContext.Get()))
 				{
 					return true;
 				}
@@ -43,11 +50,11 @@ namespace UE::PoseSearch
 
 		void Add(const FInteractionSearchContext& SearchContext, int32 SearchContextIndex)
 		{
-			for (const TWeakObjectPtr<UAnimInstance>& AnimInstance : SearchContext.AnimInstances)
+			for (const TWeakObjectPtr<const UObject>& AnimContext : SearchContext.AnimContexts)
 			{
-				if (AnimInstance.Get())
+				if (AnimContext.Get())
 				{
-					AnimInstances.Add(AnimInstance.Get());
+					AnimContexts.Add(AnimContext.Get());
 				}
 			}
 
@@ -56,10 +63,10 @@ namespace UE::PoseSearch
 
 		void Merge(const FInteractionSearchContextGroup& SearchContextGroup)
 		{
-			for (const UAnimInstance* AnimInstance : SearchContextGroup.AnimInstances)
+			for (const UObject* AnimContext : SearchContextGroup.AnimContexts)
 			{
-				check(AnimInstance);
-				AnimInstances.Add(AnimInstance);
+				check(AnimContext);
+				AnimContexts.Add(AnimContext);
 			}
 
 			for (int32 SearchContextsIndex : SearchContextGroup.SearchContextsIndices)
@@ -68,13 +75,13 @@ namespace UE::PoseSearch
 			}
 		}
 
-		TSet<const UAnimInstance*, DefaultKeyFuncs<const UAnimInstance*>, TInlineSetAllocator<UE::PoseSearch::PreallocatedRolesNum>> AnimInstances;
+		TSet<const UObject*, DefaultKeyFuncs<const UObject*>, TInlineSetAllocator<UE::PoseSearch::PreallocatedRolesNum>> AnimContexts;
 		TArray<int32, TInlineAllocator<PreallocatedSearchesNum>> SearchContextsIndices;
 	};
 
-	struct FRoledAnimInstance
+	struct FRoledAnimContext
 	{
-		UAnimInstance* AnimInstance = nullptr;
+		const UObject* AnimContext = nullptr;
 		UE::PoseSearch::FRole Role = UE::PoseSearch::DefaultRole;
 		const FAnimNode_PoseSearchHistoryCollector_Base* HistoryCollector = nullptr;
 		float BroadPhaseRadius = 0.f;
@@ -82,84 +89,84 @@ namespace UE::PoseSearch
 	};
 
 	// @todo: inherit from a TSet if it gets too slow 
-	// sorted array of FRoledAnimInstance(s). Sorted by UAnimInstance to prevent "alternating" tuples between different frames like
-	// ([AnimInstanceA, RoleA], [AnimInstanceB, RoleB]) vs ([AnimInstanceB, RoleB], [AnimInstanceA, RoleA]) 
+	// sorted array of FRoledAnimContext(s). Sorted by UAnimContext to prevent "alternating" tuples between different frames like
+	// ([AnimContextA, RoleA], [AnimContextB, RoleB]) vs ([AnimContextB, RoleB], [AnimContextA, RoleA]) 
 	// that ultimately represent the same search, but endign up with different SearchIDs
-	struct FRoledAnimInstances
+	struct FRoledAnimContexts
 	{
-		void AddRoledAnimInstace(UAnimInstance* AnimInstance, UE::PoseSearch::FRole Role, const FAnimNode_PoseSearchHistoryCollector_Base* HistoryCollector, float BroadPhaseRadius, float MaxCost)
+		void AddRoledAnimContext(const UObject* AnimContext, UE::PoseSearch::FRole Role, const FAnimNode_PoseSearchHistoryCollector_Base* HistoryCollector, float BroadPhaseRadius, float MaxCost)
 		{
-			check(AnimInstance && HistoryCollector);
+			check(AnimContext && HistoryCollector);
 
-			for (FRoledAnimInstance& RoledAnimInstance : RoledAnimInstances)
+			for (FRoledAnimContext& RoledAnimContext : RoledAnimContexts)
 			{
-				if (RoledAnimInstance.AnimInstance == AnimInstance && RoledAnimInstance.HistoryCollector == HistoryCollector && RoledAnimInstance.Role == Role)
+				if (RoledAnimContext.AnimContext == AnimContext && RoledAnimContext.HistoryCollector == HistoryCollector && RoledAnimContext.Role == Role)
 				{
-					// we already found the FRoledAnimInstances. Updating the BroadPhaseRadius and MaxCost accordingly
-					RoledAnimInstance.BroadPhaseRadius = FMath::Max(RoledAnimInstance.BroadPhaseRadius, BroadPhaseRadius);
+					// we already found the FRoledAnimContexts. Updating the BroadPhaseRadius and MaxCost accordingly
+					RoledAnimContext.BroadPhaseRadius = FMath::Max(RoledAnimContext.BroadPhaseRadius, BroadPhaseRadius);
 					
 					// MaxCost is valid only if greater than zero
 					if (MaxCost > 0.f)
 					{
-						RoledAnimInstance.MaxCost = FMath::Min(RoledAnimInstance.MaxCost, MaxCost);
+						RoledAnimContext.MaxCost = FMath::Min(RoledAnimContext.MaxCost, MaxCost);
 					}
 					return;
 				}
 			}
 
-			// we couldn't find the FRoledAnimInstances. Let's add a new one
-			FRoledAnimInstance& NewRoledAnimInstance = RoledAnimInstances.AddDefaulted_GetRef();
-			NewRoledAnimInstance.AnimInstance = AnimInstance;
-			NewRoledAnimInstance.HistoryCollector = HistoryCollector;
-			NewRoledAnimInstance.Role = Role;
-			NewRoledAnimInstance.BroadPhaseRadius = BroadPhaseRadius;
+			// we couldn't find the FRoledAnimContexts. Let's add a new one
+			FRoledAnimContext& NewRoledAnimContext = RoledAnimContexts.AddDefaulted_GetRef();
+			NewRoledAnimContext.AnimContext = AnimContext;
+			NewRoledAnimContext.HistoryCollector = HistoryCollector;
+			NewRoledAnimContext.Role = Role;
+			NewRoledAnimContext.BroadPhaseRadius = BroadPhaseRadius;
 			if (MaxCost > 0.f)
 			{
-				NewRoledAnimInstance.MaxCost = MaxCost;
+				NewRoledAnimContext.MaxCost = MaxCost;
 			}
 
-			// @todo: use a better datastructure to keep things FRoledAnimInstance(s) sorted? Or use insertion sort?
-			RoledAnimInstances.Sort([](const FRoledAnimInstance& A, const FRoledAnimInstance& B)
+			// @todo: use a better datastructure to keep things FRoledAnimContext(s) sorted? Or use insertion sort?
+			RoledAnimContexts.Sort([](const FRoledAnimContext& A, const FRoledAnimContext& B)
 			{
-				// we're sorting by AnimInstance address in memory, since searches neeed to be consistent 
+				// we're sorting by AnimContext address in memory, since searches neeed to be consistent 
 				// between frames on the same machine, not between replicated machines
-				return B.AnimInstance < A.AnimInstance;
+				return B.AnimContext < A.AnimContext;
 			});
 		}
 
-		TArrayView<const FRoledAnimInstance> GetDataView() const
+		TArrayView<const FRoledAnimContext> GetDataView() const
 		{
-			return MakeArrayView(RoledAnimInstances.GetData(), RoledAnimInstances.Num());
+			return MakeArrayView(RoledAnimContexts.GetData(), RoledAnimContexts.Num());
 		}
 
 	private:
-		TArray<FRoledAnimInstance, TInlineAllocator<UE::PoseSearch::PreallocatedRolesNum, TMemStackAllocator<>>> RoledAnimInstances;
+		TArray<FRoledAnimContext, TInlineAllocator<UE::PoseSearch::PreallocatedRolesNum, TMemStackAllocator<>>> RoledAnimContexts;
 	};
 
-	typedef TMap<const UPoseSearchDatabase*, FRoledAnimInstances, TInlineSetAllocator<PreallocatedSearchesNum, UE::PoseSearch::TMemStackSetAllocator<>>> FDatabaseToRoledAnimInstances;
-	typedef TPair<const UPoseSearchDatabase*, FRoledAnimInstances> FDatabaseToRoledAnimInstancesPair;
+	typedef TMap<const UPoseSearchDatabase*, FRoledAnimContexts, TInlineSetAllocator<PreallocatedSearchesNum, UE::PoseSearch::TMemStackSetAllocator<>>> FDatabaseToRoledAnimContexts;
+	typedef TPair<const UPoseSearchDatabase*, FRoledAnimContexts> FDatabaseToRoledAnimContextsPair;
 
-	static void ProcessAvailabilityRequestsMap(const FAvailabilityRequestsMap& AvailabilityRequestsMap, UWorld* SubsystemWorld, FDatabaseToRoledAnimInstances& DatabaseToRoledAnimInstances)
+	static void ProcessAvailabilityRequestsMap(const FAvailabilityRequestsMap& AvailabilityRequestsMap, UWorld* SubsystemWorld, FDatabaseToRoledAnimContexts& DatabaseToRoledAnimContexts)
 	{
 		for (FAvailabilityRequestsMap::TConstIterator Iter = AvailabilityRequestsMap.CreateConstIterator(); Iter; ++Iter)
 		{
 			// making sure the weak pointer is still valid
-			UAnimInstance* AnimInstance = Cast<UAnimInstance>(Iter->Key.Get());
-			if (!AnimInstance)
+			UObject* AnimContext = Iter->Key.Get();
+			if (!AnimContext)
 			{
-				UE_LOG(LogPoseSearch, Log, TEXT("ProcessAvailabilityRequestsMap null anim instance. The associated character got removed from the world"));
+				UE_LOG(LogPoseSearch, Log, TEXT("ProcessAvailabilityRequestsMap null anim context. The associated character got removed from the world"));
 				continue;
 			}
 
-			if (!AnimInstance->GetWorld())
+			if (!AnimContext->GetWorld())
 			{
-				UE_LOG(LogPoseSearch, Error, TEXT("ProcessAvailabilityRequestsMap AnimInstance %s is not in a world!"), *AnimInstance->GetName());
+				UE_LOG(LogPoseSearch, Error, TEXT("ProcessAvailabilityRequestsMap AnimContext %s is not in a world!"), *AnimContext->GetName());
 				continue;
 			}
 
-			if (SubsystemWorld != AnimInstance->GetWorld())
+			if (SubsystemWorld != AnimContext->GetWorld())
 			{
-				UE_LOG(LogPoseSearch, Error, TEXT("ProcessAvailabilityRequestsMap AnimInstance %s is from World %s, and supposed to be in World @s!"), *AnimInstance->GetName(), *AnimInstance->GetWorld()->GetName(), *SubsystemWorld->GetName());
+				UE_LOG(LogPoseSearch, Error, TEXT("ProcessAvailabilityRequestsMap AnimContext %s is from World %s, and supposed to be in World @s!"), *AnimContext->GetName(), *AnimContext->GetWorld()->GetName(), *SubsystemWorld->GetName());
 				continue;
 			}
 
@@ -177,20 +184,20 @@ namespace UE::PoseSearch
 					continue;
 				}
 
-				const FAnimNode_PoseSearchHistoryCollector_Base* HistoryCollector = AvailabilityRequest.GetHistoryCollector(AnimInstance);
+				const FAnimNode_PoseSearchHistoryCollector_Base* HistoryCollector = AvailabilityRequest.GetHistoryCollector(AnimContext);
 				if (!HistoryCollector)
 				{
-					UE_LOG(LogPoseSearch, Error, TEXT("ProcessAvailabilityRequestsMap couldn't find PoseHistory %s for AnimInstance %s"), *AvailabilityRequest.GetPoseHistoryName(), *AnimInstance->GetName());
+					UE_LOG(LogPoseSearch, Error, TEXT("ProcessAvailabilityRequestsMap couldn't find PoseHistory %s for AnimContext %s"), *AvailabilityRequest.GetPoseHistoryName(), *AnimContext->GetName());
 					continue;
 				}
 
-				FRoledAnimInstances& RoledAnimInstances = DatabaseToRoledAnimInstances.FindOrAdd(AvailabilityRequest.Database);
+				FRoledAnimContexts& RoledAnimContexts = DatabaseToRoledAnimContexts.FindOrAdd(AvailabilityRequest.Database);
 				if (AvailabilityRequest.RolesFilter.IsEmpty())
 				{
 					// adding ALL the possible roles from the database:
 					for (const FPoseSearchRoledSkeleton& RoledSkeleton : AvailabilityRequest.Database->Schema->GetRoledSkeletons())
 					{
-						RoledAnimInstances.AddRoledAnimInstace(AnimInstance, RoledSkeleton.Role, HistoryCollector, AvailabilityRequest.BroadPhaseRadius, AvailabilityRequest.MaxCost);
+						RoledAnimContexts.AddRoledAnimContext(AnimContext, RoledSkeleton.Role, HistoryCollector, AvailabilityRequest.BroadPhaseRadius, AvailabilityRequest.MaxCost);
 					}
 				}
 				else
@@ -199,7 +206,7 @@ namespace UE::PoseSearch
 					{
 						if (AvailabilityRequest.Database->Schema->GetRoledSkeleton(Role))
 						{
-							RoledAnimInstances.AddRoledAnimInstace(AnimInstance, Role, HistoryCollector, AvailabilityRequest.BroadPhaseRadius, AvailabilityRequest.MaxCost);
+							RoledAnimContexts.AddRoledAnimContext(AnimContext, Role, HistoryCollector, AvailabilityRequest.BroadPhaseRadius, AvailabilityRequest.MaxCost);
 						}
 						else
 						{
@@ -288,24 +295,24 @@ bool UPoseSearchInteractionSubsystem::ValidateAllIslands() const
 #if DO_CHECK
 	using namespace UE::PoseSearch;
 
-	TSet<TWeakObjectPtr<UCharacterMovementComponent>> CharacterMovementComponents;
-	TSet<TWeakObjectPtr<USkeletalMeshComponent>> SkeletalMeshComponents;
+	TSet<TWeakObjectPtr<UActorComponent>> PreTickComponents;
+	TSet<TWeakObjectPtr<UActorComponent>> PostTickComponents;
 
 	bool bAlreadyInSet = false;
 	for (const FInteractionIsland* Island : Islands)
 	{
-		for (const TWeakObjectPtr<UCharacterMovementComponent>& CharacterMovementComponent : Island->GetCharacterMovementComponents())
+		for (const TWeakObjectPtr<UActorComponent>& PreTickComponent : Island->GetPreTickComponents())
 		{
-			CharacterMovementComponents.Add(CharacterMovementComponent, &bAlreadyInSet);
+			PreTickComponents.Add(PreTickComponent, &bAlreadyInSet);
 			if (bAlreadyInSet)
 			{
 				return false;
 			}
 		}
 
-		for (const TWeakObjectPtr<USkeletalMeshComponent>& SkeletalMeshComponent : Island->GetSkeletalMeshComponents())
+		for (const TWeakObjectPtr<UActorComponent>& PostTickComponent : Island->GetPostTickComponents())
 		{
-			SkeletalMeshComponents.Add(SkeletalMeshComponent, &bAlreadyInSet);
+			PostTickComponents.Add(PostTickComponent, &bAlreadyInSet);
 			if (bAlreadyInSet)
 			{
 				return false;
@@ -344,23 +351,23 @@ void UPoseSearchInteractionSubsystem::PopulateContinuingProperties(UE::PoseSearc
 	}
 }
 
-UE::PoseSearch::FInteractionIsland* UPoseSearchInteractionSubsystem::FindIsland(UObject* InAnimInstance)
+UE::PoseSearch::FInteractionIsland* UPoseSearchInteractionSubsystem::FindIsland(const UObject* AnimContext)
 {
 	using namespace UE::PoseSearch;
 
-	if (InAnimInstance)
+	if (AnimContext)
 	{
-		if (UAnimInstance* AnimInstance = Cast<UAnimInstance>(InAnimInstance))
+		if (const UAnimInstance* AnimInstance = Cast<UAnimInstance>(AnimContext))
 		{
 			if (AActor* Actor = AnimInstance->GetOwningActor())
 			{
-				if (USkeletalMeshComponent* SkeletalMeshComponent = Actor->GetComponentByClass<USkeletalMeshComponent>())
+				if (UActorComponent* PostTickComponent = FInteractionIsland::FindPostTickComponent(Actor))
 				{
 					for (FInteractionIsland* Island : Islands)
 					{
-						for (const TWeakObjectPtr<USkeletalMeshComponent>& IslandSkeletalMeshComponent : Island->GetSkeletalMeshComponents())
+						for (const TWeakObjectPtr<UActorComponent>& IslandPostTickComponent : Island->GetPostTickComponents())
 						{
-							if (IslandSkeletalMeshComponent.Get() == SkeletalMeshComponent)
+							if (IslandPostTickComponent.Get() == PostTickComponent)
 							{
 								return Island;
 							}
@@ -369,15 +376,19 @@ UE::PoseSearch::FInteractionIsland* UPoseSearchInteractionSubsystem::FindIsland(
 				}
 			}
 		}
+		else
+		{
+			unimplemented();
+		}
 	}
 	return nullptr;
 }
 
-UPoseSearchInteractionSubsystem* UPoseSearchInteractionSubsystem::GetSubsystem_AnyThread(UObject* AnimInstance)
+UPoseSearchInteractionSubsystem* UPoseSearchInteractionSubsystem::GetSubsystem_AnyThread(const UObject* AnimContext)
 {
-	if (AnimInstance)
+	if (AnimContext)
 	{
-		if (UWorld* World = AnimInstance->GetWorld())
+		if (UWorld* World = AnimContext->GetWorld())
 		{
 			FScopeLock Lock(&RetrieveSubsystemMutex);
 
@@ -403,6 +414,17 @@ void UPoseSearchInteractionSubsystem::Deinitialize()
 	Super::Deinitialize();
 }
 
+static FVector GetContextLocation(const UObject* AnimContext)
+{
+	if (const UAnimInstance* AnimInstance = Cast<UAnimInstance>(AnimContext))
+	{
+		return AnimInstance->GetOwningActor()->GetActorLocation();
+	}
+
+	unimplemented();
+	return FVector::ZeroVector;
+}
+
 void UPoseSearchInteractionSubsystem::Tick(float DeltaSeconds)
 {
 	using namespace UE::PoseSearch;
@@ -422,8 +444,8 @@ void UPoseSearchInteractionSubsystem::Tick(float DeltaSeconds)
 	FMemMark Mark(FMemStack::Get());
 	
 	// processing the AvailabilityRequestsMap to fill up a map of Databases pointing to an array of all the AnimContexts with related roles
-	FDatabaseToRoledAnimInstances DatabaseToRoledAnimInstances;
-	ProcessAvailabilityRequestsMap(AvailabilityRequestsMap, GetWorld(), DatabaseToRoledAnimInstances);
+	FDatabaseToRoledAnimContexts DatabaseToRoledAnimContexts;
+	ProcessAvailabilityRequestsMap(AvailabilityRequestsMap, GetWorld(), DatabaseToRoledAnimContexts);
 
 	// for each database now we try to create all the possible permutations of the roled anim instances
 	// for example, given a database set up with assets for 2 characters interactions with roles RoleA and RoleB
@@ -444,37 +466,37 @@ void UPoseSearchInteractionSubsystem::Tick(float DeltaSeconds)
 	//
 	// CharB/RoleA - CharB/RoleB -> invalid because of same CharB
 
-	TSet<const UAnimInstance*, DefaultKeyFuncs<const UAnimInstance*>, TInlineSetAllocator<UE::PoseSearch::PreallocatedRolesNum>> AnimInstances;
+	TSet<const UObject*, DefaultKeyFuncs<const UObject*>, TInlineSetAllocator<UE::PoseSearch::PreallocatedRolesNum>> AnimContexts;
 	TSet<UE::PoseSearch::FRole, DefaultKeyFuncs<UE::PoseSearch::FRole>, TInlineSetAllocator<UE::PoseSearch::PreallocatedRolesNum>> CoveredRoles;
 	TArray<FInteractionSearchContext, TInlineAllocator<PreallocatedSearchesNum, TMemStackAllocator<>>> SearchContexts;
-	for (const FDatabaseToRoledAnimInstancesPair& DatabaseToRoledAnimInstancesPair : DatabaseToRoledAnimInstances)
+	for (const FDatabaseToRoledAnimContextsPair& DatabaseToRoledAnimContextsPair : DatabaseToRoledAnimContexts)
 	{
-		const UPoseSearchDatabase* Database = DatabaseToRoledAnimInstancesPair.Key;
+		const UPoseSearchDatabase* Database = DatabaseToRoledAnimContextsPair.Key;
 		const TArray<FPoseSearchRoledSkeleton>& RoledSkeletons = Database->Schema->GetRoledSkeletons();
 		const int32 CombinationCardinality = RoledSkeletons.Num();
 		
 		// @todo: naive implementation: optimize this code!
-		GeneratePermutations(DatabaseToRoledAnimInstancesPair.Value.GetDataView(), CombinationCardinality, 
-			[&AnimInstances, &CoveredRoles, Database, &RoledSkeletons, &SearchContexts](const TArrayView<const FRoledAnimInstance> RoledAnimInstances, const TArrayView<int32> Combination)
+		GeneratePermutations(DatabaseToRoledAnimContextsPair.Value.GetDataView(), CombinationCardinality, 
+			[&AnimContexts, &CoveredRoles, Database, &RoledSkeletons, &SearchContexts](const TArrayView<const FRoledAnimContext> RoledAnimContexts, const TArrayView<int32> Combination)
 			{
 				// checking if this combination is valid for the Database:
-				AnimInstances.Reset();
+				AnimContexts.Reset();
 				CoveredRoles.Reset();
 
 				for (int32 CombinationIndex = 0; CombinationIndex < Combination.Num(); ++CombinationIndex)
 				{
-					const FRoledAnimInstance& RoledAnimInstance = RoledAnimInstances[Combination[CombinationIndex]];
+					const FRoledAnimContext& RoledAnimContext = RoledAnimContexts[Combination[CombinationIndex]];
 
 					bool bIsAlreadyInSet = false;
-					AnimInstances.Add(RoledAnimInstance.AnimInstance, &bIsAlreadyInSet);
+					AnimContexts.Add(RoledAnimContext.AnimContext, &bIsAlreadyInSet);
 
-					// is there any duplicated AnimInstance?
+					// is there any duplicated AnimContext?
 					if (bIsAlreadyInSet)
 					{
 						return false;
 					}
 
-					CoveredRoles.Add(RoledAnimInstance.Role);
+					CoveredRoles.Add(RoledAnimContext.Role);
 				}
 
 				// does it cover all the roles?
@@ -490,18 +512,18 @@ void UPoseSearchInteractionSubsystem::Tick(float DeltaSeconds)
 				// checking if all the actors are within BroadPhaseRadius to each other
 				for (int32 CombinationIndexA = 0; CombinationIndexA < Combination.Num(); ++CombinationIndexA)
 				{
-					const FRoledAnimInstance& RoledAnimInstanceA = RoledAnimInstances[Combination[CombinationIndexA]];
-					const FVector ActorLocationA = RoledAnimInstanceA.AnimInstance->GetOwningActor()->GetActorLocation();
+					const FRoledAnimContext& RoledAnimContextA = RoledAnimContexts[Combination[CombinationIndexA]];
+					const FVector ActorLocationA = GetContextLocation(RoledAnimContextA.AnimContext);
 
 					for (int32 CombinationIndexB = 0; CombinationIndexB < Combination.Num(); ++CombinationIndexB)
 					{
 						if (CombinationIndexA != CombinationIndexB)
 						{
-							const FRoledAnimInstance& RoledAnimInstanceB = RoledAnimInstances[Combination[CombinationIndexB]];
-							const FVector ActorLocationB = RoledAnimInstanceB.AnimInstance->GetOwningActor()->GetActorLocation();
+							const FRoledAnimContext& RoledAnimContextB = RoledAnimContexts[Combination[CombinationIndexB]];
+							const FVector ActorLocationB = GetContextLocation(RoledAnimContextB.AnimContext);
 
 							const FVector DeltaLocation = ActorLocationA - ActorLocationB;
-							if (DeltaLocation.Length() > RoledAnimInstanceA.BroadPhaseRadius)
+							if (DeltaLocation.Length() > RoledAnimContextA.BroadPhaseRadius)
 							{
 								return false;
 							}
@@ -511,7 +533,7 @@ void UPoseSearchInteractionSubsystem::Tick(float DeltaSeconds)
 
 				FInteractionSearchContext& SearchContext = SearchContexts.AddDefaulted_GetRef();
 				SearchContext.Database = Database;
-				SearchContext.AnimInstances.Reserve(Combination.Num());
+				SearchContext.AnimContexts.Reserve(Combination.Num());
 				SearchContext.HistoryCollectors.Reserve(Combination.Num());
 				SearchContext.Roles.Reserve(Combination.Num());
 				
@@ -521,16 +543,16 @@ void UPoseSearchInteractionSubsystem::Tick(float DeltaSeconds)
 
 				for (int32 CombinationIndex = 0; CombinationIndex < Combination.Num(); ++CombinationIndex)
 				{
-					const FRoledAnimInstance& RoledAnimInstance = RoledAnimInstances[Combination[CombinationIndex]];
-					SearchContext.AnimInstances.Add(RoledAnimInstance.AnimInstance);
-					SearchContext.HistoryCollectors.Add(RoledAnimInstance.HistoryCollector);
-					SearchContext.Roles.Add(RoledAnimInstance.Role);
+					const FRoledAnimContext& RoledAnimContext = RoledAnimContexts[Combination[CombinationIndex]];
+					SearchContext.AnimContexts.Add(RoledAnimContext.AnimContext);
+					SearchContext.HistoryCollectors.Add(RoledAnimContext.HistoryCollector);
+					SearchContext.Roles.Add(RoledAnimContext.Role);
 #if ENABLE_DRAW_DEBUG
-					SearchContext.BroadPhaseRadiuses.Add(RoledAnimInstance.BroadPhaseRadius);
+					SearchContext.BroadPhaseRadiuses.Add(RoledAnimContext.BroadPhaseRadius);
 #endif //ENABLE_DRAW_DEBUG
 
-					// search is valid only if cost result is lower than any AvailabilityRequest.MaxCost (copied in RoledAnimInstance.MaxCost)
-					SearchContext.MaxCost = FMath::Min(SearchContext.MaxCost, RoledAnimInstance.MaxCost);
+					// search is valid only if cost result is lower than any AvailabilityRequest.MaxCost (copied in RoledAnimContext.MaxCost)
+					SearchContext.MaxCost = FMath::Min(SearchContext.MaxCost, RoledAnimContext.MaxCost);
 				}
 
 				return true;
@@ -587,9 +609,9 @@ void UPoseSearchInteractionSubsystem::Tick(float DeltaSeconds)
 	{
 		UE::PoseSearch::FInteractionIsland& Island = GetAvailableIsland();
 
-		for (const UAnimInstance* AnimInstance : SearchContextGroup.AnimInstances)
+		for (const UObject* AnimContext : SearchContextGroup.AnimContexts)
 		{
-			Island.InjectToActor(AnimInstance->GetOwningActor());
+			Island.InjectToActor(AnimContext);
 		}
 
 		for (int32 SearchContextsIndex : SearchContextGroup.SearchContextsIndices)
@@ -608,7 +630,7 @@ TStatId UPoseSearchInteractionSubsystem::GetStatId() const
 	RETURN_QUICK_DECLARE_CYCLE_STAT(UPoseSearchInteractionSubsystem, STATGROUP_Tickables);
 }
 
-void UPoseSearchInteractionSubsystem::Query_AnyThread(const TArrayView<const FPoseSearchInteractionAvailability> Availabilities, UObject* AnimInstance, 
+void UPoseSearchInteractionSubsystem::Query_AnyThread(const TArrayView<const FPoseSearchInteractionAvailability> Availabilities, UObject* AnimContext, 
 	const FPoseSearchContinuingProperties& ContinuingProperties, FPoseSearchInteractionBlueprintResult& Result,
 	FName PoseHistoryName, const FAnimNode_PoseSearchHistoryCollector_Base* HistoryCollector, bool bValidateResultAgainstAvailabilities)
 {
@@ -616,13 +638,13 @@ void UPoseSearchInteractionSubsystem::Query_AnyThread(const TArrayView<const FPo
 
 	Result = FPoseSearchInteractionBlueprintResult();
 
-	if (AnimInstance)
+	if (AnimContext)
 	{
-		// if we find AnimInstance in an island, we perform ALL the Island motion matching searches.
+		// if we find AnimContext in an island, we perform ALL the Island motion matching searches.
 		// thread safety is ensured by the lock on Island->SearchResultsMutex in DoSearch_AnyThread
-		if (FInteractionIsland* Island = FindIsland(AnimInstance))
+		if (FInteractionIsland* Island = FindIsland(AnimContext))
 		{
-			Island->DoSearch_AnyThread(AnimInstance, ContinuingProperties, Result);
+			Island->DoSearch_AnyThread(AnimContext, ContinuingProperties, Result);
 
 			if (bValidateResultAgainstAvailabilities && Result.SelectedAnimation)
 			{
@@ -650,7 +672,7 @@ void UPoseSearchInteractionSubsystem::Query_AnyThread(const TArrayView<const FPo
 		{
 			FScopeLock Lock(&AvailabilityRequestsMapMutex);
 
-			TArray<FInteractionAvailabilityEx>& AvailabilityRequests = AvailabilityRequestsMap.FindOrAdd(AnimInstance);
+			TArray<FInteractionAvailabilityEx>& AvailabilityRequests = AvailabilityRequestsMap.FindOrAdd(AnimContext);
 			for (const FPoseSearchInteractionAvailability& Availability : Availabilities)
 			{
 				if (Availability.Database)
@@ -688,9 +710,9 @@ void UPoseSearchInteractionSubsystem::DebugDraw() const
 	for (FAvailabilityRequestsMap::TConstIterator Iter = AvailabilityRequestsMap.CreateConstIterator(); Iter; ++Iter)
 	{
 		// making sure the weak pointer is still valid
-		if (UAnimInstance* AnimInstance = Cast<UAnimInstance>(Iter->Key.Get()))
+		if (UObject* AnimContext = Iter->Key.Get())
 		{
-			// looking for valid results for this AnimInstance
+			// looking for valid results for this AnimContext
 			const UPoseSearchDatabase* ResultDatabase = nullptr;
 			UE::PoseSearch::FRole ResultRole;
 			for (const FInteractionIsland* Island : Islands)
@@ -700,9 +722,9 @@ void UPoseSearchInteractionSubsystem::DebugDraw() const
 					for (const FInteractionSearchResult& SearchResult : Island->GetSearchResults())
 					{
 						const FInteractionSearchContext& SearchContext = Island->GetSearchContexts()[SearchResult.SearchIndex];
-						for (int32 AnimInstanceIndex = 0; AnimInstanceIndex < SearchContext.AnimInstances.Num(); ++AnimInstanceIndex)
+						for (int32 AnimInstanceIndex = 0; AnimInstanceIndex < SearchContext.AnimContexts.Num(); ++AnimInstanceIndex)
 						{
-							if (SearchContext.AnimInstances[AnimInstanceIndex].Get() == AnimInstance)
+							if (SearchContext.AnimContexts[AnimInstanceIndex].Get() == AnimContext)
 							{
 								ResultDatabase = SearchResult.Database.Get();
 								check(ResultDatabase);
@@ -786,8 +808,8 @@ void UPoseSearchInteractionSubsystem::DebugDraw() const
 			
 			// @todo: waiting for visual logger fix to display strings :/
 			const FVector StringOffset(0.f, 0.f, 0.001f);
-			const FVector ActorLocation = AnimInstance->GetOwningActor()->GetActorLocation();
-			UE_VLOG_SEGMENT(AnimInstance, "PoseSearchInteraction", Display, ActorLocation, ActorLocation + StringOffset, FColor::Transparent, TEXT("%s"), StringBuilder.ToString());
+			const FVector ActorLocation = GetContextLocation(AnimContext);
+			UE_VLOG_SEGMENT(AnimContext, "PoseSearchInteraction", Display, ActorLocation, ActorLocation + StringOffset, FColor::Transparent, TEXT("%s"), StringBuilder.ToString());
 		}
 	}
 #endif // ENABLE_VISUAL_LOG
diff --git a/Engine/Plugins/Animation/PoseSearch/Source/Runtime/Public/PoseSearch/PoseSearchInteractionIsland.h b/Engine/Plugins/Animation/PoseSearch/Source/Runtime/Public/PoseSearch/PoseSearchInteractionIsland.h
--- a/Engine/Plugins/Animation/PoseSearch/Source/Runtime/Public/PoseSearch/PoseSearchInteractionIsland.h
+++ b/Engine/Plugins/Animation/PoseSearch/Source/Runtime/Public/PoseSearch/PoseSearchInteractionIsland.h
@@ -17,7 +17,7 @@ struct FInteractionSearchContext
 {
 	// @todo: since AnimContexts, HistoryCollectors, Roles, and BroadPhaseRadiuses have the same cardinality,
 	// add an array of structs instead of many arrays
-	TArray<TWeakObjectPtr<UAnimInstance>> AnimInstances;
+	TArray<TWeakObjectPtr<const UObject>> AnimContexts;
 	// @todo: perhaps use the animnode of the pose history collector instead
 	TArray<const FAnimNode_PoseSearchHistoryCollector_Base*> HistoryCollectors;
 	TArray<UE::PoseSearch::FRole> Roles;
@@ -45,7 +45,8 @@ struct FInteractionSearchResult : public FSearchResult
 };
 
 // Experimental, this feature might be removed without warning, not for production use
-// FInteractionIsland contains ticks functions injected between the interacting actors UCharacterMovementComponent and USkeletalMeshComponent
+// FInteractionIsland contains ticks functions injected between the interacting actors PreTickComponent (UCharacterMovementComponent, or UCharacterMoverComponent)
+// and PostTickComponent (USkeletalMeshComponent, or WIP UAnimNextComponent)
 // to create a execution threading fence to be able to perform motion matching searches between the involved characters in a thread safe manner.
 // Look at UPoseSearchInteractionSubsystem "Execution model and threading details" for additional information
 struct FInteractionIsland
@@ -55,24 +56,27 @@ struct FInteractionIsland
 	FInteractionIsland(ULevel* Level);
 	~FInteractionIsland();
 
-	bool DoSearch_AnyThread(UObject* AnimInstance, const FPoseSearchContinuingProperties& ContinuingProperties, FPoseSearchInteractionBlueprintResult& Result);
+	bool DoSearch_AnyThread(UObject* AnimContext, const FPoseSearchContinuingProperties& ContinuingProperties, FPoseSearchInteractionBlueprintResult& Result);
 
-	const TArray<TWeakObjectPtr<UCharacterMovementComponent>>& GetCharacterMovementComponents() const { return CharacterMovementComponents; }
-	const TArray<TWeakObjectPtr<USkeletalMeshComponent>>& GetSkeletalMeshComponents() const { return SkeletalMeshComponents; }
+	const TArray<TWeakObjectPtr<UActorComponent>>& GetPreTickComponents() const { return PreTickComponents; }
+	const TArray<TWeakObjectPtr<UActorComponent>>& GetPostTickComponents() const { return PostTickComponents; }
 	const TArray<FInteractionSearchContext>& GetSearchContexts() const { return SearchContexts; }
 	const TArray<FInteractionSearchResult>& GetSearchResults() const { return SearchResults; }
 
 	bool IsUninjected() const;
-	void InjectToActor(AActor* Actor);
+	void InjectToActor(const UObject* AnimContext);
 	void Uninject();
 
+	static UActorComponent* FindPreTickComponent(const AActor* Actor);
+	static UActorComponent* FindPostTickComponent(const AActor* Actor);
+
 	const FInteractionSearchResult* FindSearchResult(const FInteractionSearchContext& SearchContext) const;
 	void AddSearchContext(const FInteractionSearchContext& SearchContext);
 
 	void DebugDraw(const FColor& Color = FColor::Red) const;
 
 private:
-	bool GetResult_AnyThread(UObject* AnimInstance, FPoseSearchInteractionBlueprintResult& Result);
+	bool GetResult_AnyThread(UObject* AnimContext, FPoseSearchInteractionBlueprintResult& Result);
 
 	struct FPreTickFunction : public FTickFunction
 	{
@@ -89,8 +93,8 @@ private:
 	};
 	FPostTickFunction PostTickFunction;
 
-	TArray<TWeakObjectPtr<UCharacterMovementComponent>> CharacterMovementComponents;
-	TArray<TWeakObjectPtr<USkeletalMeshComponent>> SkeletalMeshComponents;
+	TArray<TWeakObjectPtr<UActorComponent>> PreTickComponents;
+	TArray<TWeakObjectPtr<UActorComponent>> PostTickComponents;
 	
 	// there's one FSearchContext for each search we need to perform (including all the possible roles permutations). Added by UPoseSearchInteractionSubsystem::Tick
 	TArray<FInteractionSearchContext> SearchContexts;
diff --git a/Engine/Plugins/Animation/PoseSearch/Source/Runtime/Public/PoseSearch/PoseSearchInteractionSubsystem.h b/Engine/Plugins/Animation/PoseSearch/Source/Runtime/Public/PoseSearch/PoseSearchInteractionSubsystem.h
--- a/Engine/Plugins/Animation/PoseSearch/Source/Runtime/Public/PoseSearch/PoseSearchInteractionSubsystem.h
+++ b/Engine/Plugins/Animation/PoseSearch/Source/Runtime/Public/PoseSearch/PoseSearchInteractionSubsystem.h
@@ -27,7 +27,7 @@ struct FInteractionAvailabilityEx : public FPoseSearchInteractionAvailability
 	}
 
 	FString GetPoseHistoryName() const;
-	const FAnimNode_PoseSearchHistoryCollector_Base* GetHistoryCollector(const UAnimInstance* AnimInstance) const;
+	const FAnimNode_PoseSearchHistoryCollector_Base* GetHistoryCollector(const UObject* AnimContext) const;
 
 private:
 	FName PoseHistoryName;
@@ -85,7 +85,7 @@ class POSESEARCH_API UPoseSearchInteractionSubsystem : public UTickableWorldSubs
 	GENERATED_BODY()
 
 public:
-	static UPoseSearchInteractionSubsystem* GetSubsystem_AnyThread(UObject* AnimInstance);
+	static UPoseSearchInteractionSubsystem* GetSubsystem_AnyThread(const UObject* AnimInstance);
 
 	virtual void Initialize(FSubsystemCollectionBase& Collection) override;
 	virtual void Deinitialize() override;
@@ -95,10 +95,10 @@ public:
 	
 	virtual TStatId GetStatId() const override;
 
-	// publishing FPoseSearchInteractionAvailability(s) for the requesting character (AnimInstance) and getting the FPoseSearchInteractionBlueprintResult from the PREVIOUS Tick update
+	// publishing FPoseSearchInteractionAvailability(s) for the requesting character (via AnimContext as UAnimInstance or UAnimNextComponent) and getting the FPoseSearchInteractionBlueprintResult from the PREVIOUS Tick update
 	// containing the animation to play at what time, and the assigned role to partecipate in the selected interaction
 	// Either a PoseHistoryName or a HistoryCollector are required to perform the associated motion matching searches
-	void Query_AnyThread(const TArrayView<const FPoseSearchInteractionAvailability> Availabilities, UObject* AnimInstance, 
+	void Query_AnyThread(const TArrayView<const FPoseSearchInteractionAvailability> Availabilities, UObject* AnimContext, 
 		const FPoseSearchContinuingProperties& ContinuingProperties, FPoseSearchInteractionBlueprintResult& Result,
 		FName PoseHistoryName, const FAnimNode_PoseSearchHistoryCollector_Base* HistoryCollector, bool bValidateResultAgainstAvailabilities);
 private:
@@ -111,7 +111,7 @@ private:
 	bool ValidateAllIslands() const;
 
 	void PopulateContinuingProperties(UE::PoseSearch::FInteractionSearchContext& SearchContext, float DeltaSeconds) const;
-	UE::PoseSearch::FInteractionIsland* FindIsland(UObject* InAnimInstance);
+	UE::PoseSearch::FInteractionIsland* FindIsland(const UObject* AnimContext);
 	
 	void DebugDraw() const;
 
