pose search - support for EPoseSearchInterruptMode in MotionMatch input struct FPoseSearchContinuingProperties

diff --git a/Engine/Plugins/Animation/PoseSearch/Source/Editor/Private/PoseSearchDebuggerDatabaseColumns.h b/Engine/Plugins/Animation/PoseSearch/Source/Editor/Private/PoseSearchDebuggerDatabaseColumns.h
--- a/Engine/Plugins/Animation/PoseSearch/Source/Editor/Private/PoseSearchDebuggerDatabaseColumns.h
+++ b/Engine/Plugins/Animation/PoseSearch/Source/Editor/Private/PoseSearchDebuggerDatabaseColumns.h
@@ -391,7 +391,7 @@ struct FCost : ITextColumn
 	{
 		if (Row->PoseCost.IsValid())
 		{
-			return FText::AsNumber(Row->PoseCost.GetTotalCost());
+			return FText::AsNumber(Row->PoseCost);
 		}
 		else
 		{
diff --git a/Engine/Plugins/Animation/PoseSearch/Source/Editor/Private/PoseSearchDebuggerDatabaseView.cpp b/Engine/Plugins/Animation/PoseSearch/Source/Editor/Private/PoseSearchDebuggerDatabaseView.cpp
--- a/Engine/Plugins/Animation/PoseSearch/Source/Editor/Private/PoseSearchDebuggerDatabaseView.cpp
+++ b/Engine/Plugins/Animation/PoseSearch/Source/Editor/Private/PoseSearchDebuggerDatabaseView.cpp
@@ -188,9 +188,9 @@ static void AddUnfilteredDatabaseRow(const UPoseSearchDatabase* Database,
 
 			// @todo: perhaps reuse this code to recalculate Row->PoseCost instead of using Cost
 			//const FPoseSearchCost RecalculatedCost = SearchIndex.ComparePoses(DbPoseIdx, 0.f, PoseValues, QueryVector);
-			//if (Cost.GetTotalCost() != RecalculatedCost.GetTotalCost())
+			//if (Cost != RecalculatedCost)
 			//{
-			//	UE_LOG(LogPoseSearchEditor, Error, TEXT("RecalculatedCost (%f) differs from Cost (%f) for Pose '%d'"), RecalculatedCost.GetTotalCost(), Cost.GetTotalCost(), DbPoseIdx);
+			//	UE_LOG(LogPoseSearchEditor, Error, TEXT("RecalculatedCost (%f) differs from Cost (%f) for Pose '%d'"), RecalculatedCost, Cost, DbPoseIdx);
 			//}
 		}
 		// is PoseValues padded at 16 bytes (and 16 bytes aligned by construction)?
@@ -333,7 +333,7 @@ void SDebuggerDatabaseView::Update(const FTraceMotionMatchingStateMessage& State
 		{
 			if (EnumHasAnyFlags(UnfilteredRow->PoseCandidateFlags, EPoseCandidateFlags::AnyValidMask))
 			{
-				const float Cost = UnfilteredRow->PoseCost.GetTotalCost();
+				const float Cost = UnfilteredRow->PoseCost;
 				MinCost = FMath::Min(MinCost, Cost);
 				MaxCost = FMath::Max(MaxCost, Cost);
 
@@ -368,7 +368,7 @@ void SDebuggerDatabaseView::Update(const FTraceMotionMatchingStateMessage& State
 		{
 			if (EnumHasAnyFlags(UnfilteredRow->PoseCandidateFlags, EPoseCandidateFlags::AnyValidMask))
 			{
-				const float CostBlend = (UnfilteredRow->PoseCost.GetTotalCost() - MinCost) / DeltaCost;
+				const float CostBlend = (UnfilteredRow->PoseCost - MinCost) / DeltaCost;
 				UnfilteredRow->CostColor = BestScoreRowColor + (WorstScoreRowColor - BestScoreRowColor) * CostBlend;
 
 				const float PCACostBlend = (UnfilteredRow->PosePCACost - MinPCACost) / DeltaPCACost;
@@ -645,7 +645,7 @@ void SDebuggerDatabaseView::PopulateViewRows()
 
 				if (RowDataA.DbAssetIdx == RowDataB.DbAssetIdx)
 				{
-					return RowDataA.PoseCost.GetTotalCost() < RowDataB.PoseCost.GetTotalCost();
+					return RowDataA.PoseCost < RowDataB.PoseCost;
 				}
 				return RowDataA.DbAssetIdx < RowDataB.DbAssetIdx;
 			});
diff --git a/Engine/Plugins/Animation/PoseSearch/Source/Runtime/Private/PoseSearchContext.cpp b/Engine/Plugins/Animation/PoseSearch/Source/Runtime/Private/PoseSearchContext.cpp
--- a/Engine/Plugins/Animation/PoseSearch/Source/Runtime/Private/PoseSearchContext.cpp
+++ b/Engine/Plugins/Animation/PoseSearch/Source/Runtime/Private/PoseSearchContext.cpp
@@ -531,7 +531,7 @@ void FSearchContext::UpdateCurrentResultPoseVector()
 		if (SearchIndex.IsValuesEmpty())
 		{
 			const int32 NumDimensions = CurrentResult.Database->Schema->SchemaCardinality;
-			CurrentResultPoseVectorData.AddUninitialized(NumDimensions);
+			CurrentResultPoseVectorData.SetNumUninitialized(NumDimensions, EAllowShrinking::No);
 			CurrentResultPoseVector = SearchIndex.GetReconstructedPoseValues(CurrentResult.PoseIdx, MakeArrayView(CurrentResultPoseVectorData).Slice(NumDimensions, NumDimensions));
 		}
 		else
@@ -539,6 +539,10 @@ void FSearchContext::UpdateCurrentResultPoseVector()
 			CurrentResultPoseVector = SearchIndex.GetPoseValues(CurrentResult.PoseIdx);
 		}
 	}
+	else
+	{
+		CurrentResultPoseVector = TStackAlignedArray<float>();
+	}
 }
 
 float FSearchContext::GetSampleCurveValue(float SampleTimeOffset, const FName& CurveName, const FRole& SampleRole)
@@ -839,9 +843,9 @@ void FSearchContext::ResetCurrentBestCost()
 
 void FSearchContext::UpdateCurrentBestCost(const FPoseSearchCost& PoseSearchCost)
 {
-	if (PoseSearchCost.GetTotalCost() < CurrentBestTotalCost)
+	if (PoseSearchCost < CurrentBestTotalCost)
 	{
-		CurrentBestTotalCost = PoseSearchCost.GetTotalCost();
+		CurrentBestTotalCost = PoseSearchCost;
 	};
 }
 
diff --git a/Engine/Plugins/Animation/PoseSearch/Source/Runtime/Private/PoseSearchInteractionIsland.cpp b/Engine/Plugins/Animation/PoseSearch/Source/Runtime/Private/PoseSearchInteractionIsland.cpp
--- a/Engine/Plugins/Animation/PoseSearch/Source/Runtime/Private/PoseSearchInteractionIsland.cpp
+++ b/Engine/Plugins/Animation/PoseSearch/Source/Runtime/Private/PoseSearchInteractionIsland.cpp
@@ -603,7 +603,7 @@ bool FInteractionIsland::DoSearch_AnyThread(UObject* AnimContext, const FPoseSea
 			const UE::PoseSearch::FSearchResult PoseSearchResult = UPoseSearchLibrary::MotionMatch(AnimContexts, SearchContext.Roles,
 				PoseHistories, AssetsToSearch, ContinuingPropertiesToUse, PoseSearchFutureProperties);
 
-			if (PoseSearchResult.PoseCost.GetTotalCost() < SearchContext.MaxCost)
+			if (PoseSearchResult.PoseCost < SearchContext.MaxCost)
 			{
 				PoseSearchResults[SearchIndex] = PoseSearchResult;
 			}
@@ -674,7 +674,7 @@ bool FInteractionIsland::GetResult_AnyThread(UObject* AnimContext, FPoseSearchIn
 				Result.bIsMirrored = SearchIndexAsset->IsMirrored();
 				Result.BlendParameters = SearchIndexAsset->GetBlendParameters();
 				Result.SelectedDatabase = Database;
-				Result.SearchCost = SearchResult.PoseCost.GetTotalCost();
+				Result.SearchCost = SearchResult.PoseCost;
 				Result.Role = Role;
 
 				if (SearchResult.FullAlignedActorRootBoneTransforms.IsValidIndex(AnimContextIndex))
diff --git a/Engine/Plugins/Animation/PoseSearch/Source/Runtime/Private/PoseSearchLibrary.cpp b/Engine/Plugins/Animation/PoseSearch/Source/Runtime/Private/PoseSearchLibrary.cpp
--- a/Engine/Plugins/Animation/PoseSearch/Source/Runtime/Private/PoseSearchLibrary.cpp
+++ b/Engine/Plugins/Animation/PoseSearch/Source/Runtime/Private/PoseSearchLibrary.cpp
@@ -40,16 +40,17 @@ namespace UE::PoseSearch
 	// referencing multiple databases will default to a slower TMemStackAllocator (that hides heap allocations)
 	enum { MAX_STACK_ALLOCATED_ANIMATIONS = 16 };
 	enum { MAX_STACK_ALLOCATED_SETS = 2 };
-	typedef	TArray<const UObject*, TInlineAllocator<MAX_STACK_ALLOCATED_ANIMATIONS, TMemStackAllocator<>>> TAssetsToSearch;
-	// an empty TAssetsToSearch associated to Database means we need to search ALL the assets
-	typedef TMap<const UPoseSearchDatabase*, TAssetsToSearch, TInlineSetAllocator<MAX_STACK_ALLOCATED_SETS, TMemStackSetAllocator<>>> TAssetsToSearchPerDatabaseMap;
-	typedef TPair<const UPoseSearchDatabase*, TAssetsToSearch> TAssetsToSearchPerDatabasePair;
+	typedef	TArray<const UObject*, TInlineAllocator<MAX_STACK_ALLOCATED_ANIMATIONS, TMemStackAllocator<>>> FAssetsToSearch;
+	// an empty FAssetsToSearch associated to Database means we need to search ALL the assets
+	typedef TMap<const UPoseSearchDatabase*, FAssetsToSearch, TInlineSetAllocator<MAX_STACK_ALLOCATED_SETS, TMemStackSetAllocator<>>> FAssetsToSearchPerDatabaseMap;
+	typedef TPair<const UPoseSearchDatabase*, FAssetsToSearch> TAssetsToSearchPerDatabasePair;
+	typedef TMap<const UPoseSearchDatabase*, FSearchResult, TInlineSetAllocator<MAX_STACK_ALLOCATED_SETS, TMemStackSetAllocator<>>> FReconstructedPreviousSearchBestResultMap;
 
 	// this function adds AssetToSearch to the search of Database
 	// returns bAsyncBuildIndexInProgress
-	static bool AddToSearchForDatabase(TAssetsToSearchPerDatabaseMap& AssetsToSearchPerDatabaseMap, const UObject* AssetToSearch, const UPoseSearchDatabase* Database, bool bContainsIsMandatory)
+	static bool AddToSearchForDatabase(FAssetsToSearchPerDatabaseMap& AssetsToSearchPerDatabaseMap, const UObject* AssetToSearch, const UPoseSearchDatabase* Database, bool bContainsIsMandatory)
 	{
-		TAssetsToSearch* AssetsToSearch = AssetsToSearchPerDatabaseMap.Find(Database);
+		FAssetsToSearch* AssetsToSearch = AssetsToSearchPerDatabaseMap.Find(Database);
 
 #if WITH_EDITOR
 		// no need to check if Database is indexing if found into AssetsToSearchPerDatabaseMap, since it already passed RequestAsyncBuildIndex successfully in a previous AddToSearchForDatabase call
@@ -74,7 +75,7 @@ namespace UE::PoseSearch
 
 		if (AssetsToSearch)
 		{
-			// an empty TAssetsToSearch associated to Database means we need to search ALL the assets, so we don't need to add this AssetToSearch
+			// an empty FAssetsToSearch associated to Database means we need to search ALL the assets, so we don't need to add this AssetToSearch
 			if (!AssetsToSearch->IsEmpty())
 			{
 				AssetsToSearch->AddUnique(AssetToSearch);
@@ -93,7 +94,7 @@ namespace UE::PoseSearch
 	// if AssetToSearch is a database search it ALL,
 	// if it's a sequence containing UAnimNotifyState_PoseSearchBranchIn, we add to the search of the dabase UAnimNotifyState_PoseSearchBranchIn::Database the asset AssetToSearch
 	// returns bAsyncBuildIndexInProgress
-	static bool AddToSearch(TAssetsToSearchPerDatabaseMap& AssetsToSearchPerDatabaseMap, const UObject* AssetToSearch)
+	static bool AddToSearch(FAssetsToSearchPerDatabaseMap& AssetsToSearchPerDatabaseMap, const UObject* AssetToSearch)
 	{
 		bool bAsyncBuildIndexInProgress = false;
 		if (const UAnimSequenceBase* SequenceBase = Cast<const UAnimSequenceBase>(AssetToSearch))
@@ -119,9 +120,9 @@ namespace UE::PoseSearch
 		else if (const UPoseSearchDatabase* Database = Cast<UPoseSearchDatabase>(AssetToSearch))
 		{
 			// we already added Database to AssetsToSearchPerDatabaseMap, so it already successfully passed RequestAsyncBuildIndex
-			if (TAssetsToSearch* AssetsToSearch = AssetsToSearchPerDatabaseMap.Find(Database))
+			if (FAssetsToSearch* AssetsToSearch = AssetsToSearchPerDatabaseMap.Find(Database))
 			{
-				// an empty TAssetsToSearch associated to Database means we need to search ALL the assets
+				// an empty FAssetsToSearch associated to Database means we need to search ALL the assets
 				AssetsToSearch->Reset();
 			}
 			else
@@ -133,7 +134,7 @@ namespace UE::PoseSearch
 			else
 #endif // WITH_EDITOR
 			{
-				// an empty TAssetsToSearch associated to Database means we need to search ALL the assets
+				// an empty FAssetsToSearch associated to Database means we need to search ALL the assets
 				AssetsToSearchPerDatabaseMap.Add(Database);
 			}
 		}
@@ -141,7 +142,50 @@ namespace UE::PoseSearch
 		return bAsyncBuildIndexInProgress;
 	}
 
-	static bool IsForceInterrupt(EPoseSearchInterruptMode InterruptMode, const UPoseSearchDatabase* CurrentResultDatabase, const TArray<TObjectPtr<const UPoseSearchDatabase>>& Databases)
+	static void PopulateContinuingPoseSearches(const UObject* PlayingAnimationAsset, const TArrayView<const UObject*> AssetsToSearch, FSearchContext& SearchContext, FAssetsToSearchPerDatabaseMap& ContinuingPoseAssetsToSearchPerDatabaseMap)
+	{
+		// checking if PlayingAnimationAsset has an associated database
+		if (AddToSearch(ContinuingPoseAssetsToSearchPerDatabaseMap, PlayingAnimationAsset))
+		{
+#if WITH_EDITOR
+			SearchContext.SetAsyncBuildIndexInProgress();
+#endif // WITH_EDITOR
+		}
+		
+		// checking if any of the AssetsToSearch (databse) contains PlayingAnimationAsset
+		for (const UObject* AssetToSearch : AssetsToSearch)
+		{
+			if (const UPoseSearchDatabase* Database = Cast<UPoseSearchDatabase>(AssetToSearch))
+			{
+				// since it cannot be a database we can directly add it to ContinuingPoseAssetsToSearchPerDatabaseMap
+				if (AddToSearchForDatabase(ContinuingPoseAssetsToSearchPerDatabaseMap, PlayingAnimationAsset, Database, false))
+				{
+#if WITH_EDITOR
+					SearchContext.SetAsyncBuildIndexInProgress();
+#endif // WITH_EDITOR
+				}
+			}
+		}
+	}
+	
+	static void PopulateSearches(const TArrayView<const UObject*> AssetsToSearch, FSearchContext& SearchContext, FAssetsToSearchPerDatabaseMap& AssetsToSearchPerDatabaseMap)
+	{
+		if (!AssetsToSearch.IsEmpty())
+		{
+			for (const UObject* AssetToSearch : AssetsToSearch)
+			{
+				if (AddToSearch(AssetsToSearchPerDatabaseMap, AssetToSearch))
+				{
+#if WITH_EDITOR
+					SearchContext.SetAsyncBuildIndexInProgress();
+#endif // WITH_EDITOR
+				}
+			}
+		}
+	}
+
+	template <typename DatabasesContainer>
+	static bool IsForceInterrupt(EPoseSearchInterruptMode InterruptMode, const UPoseSearchDatabase* CurrentResultDatabase, const DatabasesContainer& Databases)
 	{
 		switch (InterruptMode)
 		{
@@ -162,7 +206,8 @@ namespace UE::PoseSearch
 		}
 	}
 
-	static bool IsInvalidatingContinuingPose(EPoseSearchInterruptMode InterruptMode, const UPoseSearchDatabase* CurrentResultDatabase, const TArray<TObjectPtr<const UPoseSearchDatabase>>& Databases)
+	template <typename DatabasesContainer>
+	static bool IsInvalidatingContinuingPose(EPoseSearchInterruptMode InterruptMode, const UPoseSearchDatabase* CurrentResultDatabase, const DatabasesContainer& Databases)
 	{
 		switch (InterruptMode)
 		{
@@ -465,9 +510,9 @@ void UPoseSearchLibrary::TraceMotionMatching(UE::PoseSearch::FSearchContext& Sea
 	TraceState.DeltaTime = DeltaTime;
 
 	TraceState.RecordingTime = RecordingTime;
-	TraceState.SearchBestCost = CurrentResult.PoseCost.GetTotalCost();
+	TraceState.SearchBestCost = CurrentResult.PoseCost;
 #if WITH_EDITOR && ENABLE_ANIM_DEBUG
-	TraceState.SearchBruteForceCost = CurrentResult.BruteForcePoseCost.GetTotalCost();
+	TraceState.SearchBruteForceCost = CurrentResult.BruteForcePoseCost;
 	TraceState.SearchBestPosePos = CurrentResult.BestPosePos;
 #else // WITH_EDITOR && ENABLE_ANIM_DEBUG
 	TraceState.SearchBruteForceCost = 0.f;
@@ -743,7 +788,9 @@ void UPoseSearchLibrary::MotionMatch(
 	Roles.Add(UE::PoseSearch::DefaultRole);
 
 	TArray<const UObject*>& AssetsToSearchConst = reinterpret_cast<TArray<const UObject*>&>(AssetsToSearch);
+	PRAGMA_DISABLE_DEPRECATION_WARNINGS
 	MotionMatch(AnimInstances, Roles, AssetsToSearchConst, PoseHistoryName, ContinuingProperties, Future, Result);
+	PRAGMA_ENABLE_DEPRECATION_WARNINGS
 }
 
 void UPoseSearchLibrary::MotionMatch(
@@ -822,7 +869,7 @@ void UPoseSearchLibrary::MotionMatch(
 			Result.bIsMirrored = SearchIndexAsset->IsMirrored();
 			Result.BlendParameters = SearchIndexAsset->GetBlendParameters();
 			Result.SelectedDatabase = Database;
-			Result.SearchCost = SearchResult.PoseCost.GetTotalCost();
+			Result.SearchCost = SearchResult.PoseCost;
 			
 			// figuring out the WantedPlayRate
 			Result.WantedPlayRate = 1.f;
@@ -930,6 +977,7 @@ UE::PoseSearch::FSearchResult UPoseSearchLibrary::MotionMatch(
 		}		
 	}
 
+	FReconstructedPreviousSearchBestResultMap ReconstructedPreviousSearchBestResultMap;
 	FSearchResult ReconstructedPreviousSearchResult;
 	FSearchContext SearchContext(FutureIntervalTime, nullptr, ReconstructedPreviousSearchResult);
 
@@ -937,8 +985,6 @@ UE::PoseSearch::FSearchResult UPoseSearchLibrary::MotionMatch(
 	{
 		SearchContext.AddRole(Roles[RoleIndex], AnimContexts[RoleIndex], InternalPoseHistories[RoleIndex]);
 	}
-
-	TAssetsToSearchPerDatabaseMap AssetsToSearchPerDatabaseMap;
 	
 	bool bAsyncBuildIndexInProgress = false;
 
@@ -949,34 +995,25 @@ UE::PoseSearch::FSearchResult UPoseSearchLibrary::MotionMatch(
 		DeltaSeconds = AnimInstance->GetDeltaSeconds();
 	}
 
-	if (const UObject* PlayingAnimationAsset = ContinuingProperties.PlayingAsset.Get())
+	// collecting all the databases searches in AssetsToSearchPerDatabaseMap
+	// and all the continuing pose searches in ContinuingPoseAssetsToSearchPerDatabaseMap
+	FAssetsToSearchPerDatabaseMap AssetsToSearchPerDatabaseMap;
+	FAssetsToSearchPerDatabaseMap ContinuingPoseAssetsToSearchPerDatabaseMap;
+	PopulateSearches(AssetsToSearch, SearchContext, AssetsToSearchPerDatabaseMap);
+	PopulateContinuingPoseSearches(ContinuingProperties.PlayingAsset.Get(), AssetsToSearch, SearchContext, ContinuingPoseAssetsToSearchPerDatabaseMap);
+
+	for (const TAssetsToSearchPerDatabasePair& AssetsToSearchPerDatabasePair : ContinuingPoseAssetsToSearchPerDatabaseMap)
 	{
-		// checking if PlayingAnimationAsset has an associated database
-		if (AddToSearch(AssetsToSearchPerDatabaseMap, PlayingAnimationAsset))
-		{
-			bAsyncBuildIndexInProgress = true;
-		}
-		
-		// checking if any of the AssetsToSearch (databse) contains PlayingAnimationAsset
-		for (const UObject* AssetToSearch : AssetsToSearch)
-		{
-			if (const UPoseSearchDatabase* Database = Cast<UPoseSearchDatabase>(AssetToSearch))
-			{
-				// since it cannot be a database we can directly add it to AssetsToSearchPerDatabaseMap
-				if (AddToSearchForDatabase(AssetsToSearchPerDatabaseMap, PlayingAnimationAsset, Database, false))
-				{
-					bAsyncBuildIndexInProgress = true;
-				}
-			}
-		}
+		const UPoseSearchDatabase* Database = AssetsToSearchPerDatabasePair.Key;
+		check(Database);
 
-		for (const TAssetsToSearchPerDatabasePair& AssetsToSearchPerDatabasePair : AssetsToSearchPerDatabaseMap)
+		const bool bInvalidatingContinuingPose = IsInvalidatingContinuingPose(ContinuingProperties.InterruptMode, Database, AssetsToSearchPerDatabaseMap);
+		if (!bInvalidatingContinuingPose)
 		{
-			const UPoseSearchDatabase* Database = AssetsToSearchPerDatabasePair.Key;
-			check(Database);
+			const bool bForceInterrupt = IsForceInterrupt(ContinuingProperties.InterruptMode, Database, AssetsToSearchPerDatabaseMap);
 
 			const FSearchIndex& SearchIndex = Database->GetSearchIndex();
-			for (int32 AssetIndex : Database->GetAssetIndexesForSourceAsset(PlayingAnimationAsset))
+			for (int32 AssetIndex : Database->GetAssetIndexesForSourceAsset(ContinuingProperties.PlayingAsset.Get()))
 			{
 				const FSearchIndexAsset& SearchIndexAsset = SearchIndex.Assets[AssetIndex];
 
@@ -1011,13 +1048,14 @@ UE::PoseSearch::FSearchResult UPoseSearchLibrary::MotionMatch(
 					ReconstructedPreviousSearchResult.AssetTime = ContinuingProperties.PlayingAssetAccumulatedTime;
 				}
 
-				if (bCanAdvance)
+				ReconstructedPreviousSearchResult.Database = Database;
+				ReconstructedPreviousSearchResult.PoseIdx = Database->GetPoseIndexFromTime(ContinuingProperties.PlayingAssetAccumulatedTime, SearchIndexAsset);
+
+				FSearchResult NewSearchResult;
+				if (bCanAdvance && !bForceInterrupt)
 				{
-					ReconstructedPreviousSearchResult.Database = Database;
-					ReconstructedPreviousSearchResult.PoseIdx = Database->GetPoseIndexFromTime(ContinuingProperties.PlayingAssetAccumulatedTime, SearchIndexAsset);
 					SearchContext.UpdateCurrentResultPoseVector();
-
-					const FSearchResult NewSearchResult = Database->SearchContinuingPose(SearchContext);
+					NewSearchResult = Database->SearchContinuingPose(SearchContext);
 
 #if WITH_EDITOR && ENABLE_ANIM_DEBUG && UE_POSE_SEARCH_TRACE_ENABLED
 					const FPoseSearchCost BestBruteForcePoseCost = NewSearchResult.BruteForcePoseCost < SearchResult.BruteForcePoseCost ? NewSearchResult.BruteForcePoseCost : SearchResult.BruteForcePoseCost;
@@ -1028,60 +1066,89 @@ UE::PoseSearch::FSearchResult UPoseSearchLibrary::MotionMatch(
 						SearchResult = NewSearchResult;
 						SearchContext.UpdateCurrentBestCost(SearchResult.PoseCost);
 					}
-								
 #if WITH_EDITOR && ENABLE_ANIM_DEBUG && UE_POSE_SEARCH_TRACE_ENABLED
 					SearchResult.BruteForcePoseCost = BestBruteForcePoseCost;
 #endif // WITH_EDITOR && ENABLE_ANIM_DEBUG && UE_POSE_SEARCH_TRACE_ENABLED
 				}
+
+				// keeping track of the best ReconstructedPreviousSearchResult, as best continuing pose search result for the full Database search happening after
+				if (FSearchResult* FoundReconstructedPreviousSearchBestResult = ReconstructedPreviousSearchBestResultMap.Find(Database))
+				{
+					if (NewSearchResult.PoseCost == FoundReconstructedPreviousSearchBestResult->PoseCost)
+					{
+						check(NewSearchResult.PoseIdx != FoundReconstructedPreviousSearchBestResult->PoseIdx);
+						if (NewSearchResult.PoseIdx < FoundReconstructedPreviousSearchBestResult->PoseIdx)
+						{
+							// choosing to update with the lowest PoseIdx to avoid indeterminism 
+							// since ReconstructedPreviousSearchBestResultMap is an unsorted Map
+							*FoundReconstructedPreviousSearchBestResult = ReconstructedPreviousSearchResult;
+							FoundReconstructedPreviousSearchBestResult->PoseCost = NewSearchResult.PoseCost;
+						}
+					}
+					else if (NewSearchResult.PoseCost < FoundReconstructedPreviousSearchBestResult->PoseCost)
+					{
+						*FoundReconstructedPreviousSearchBestResult = ReconstructedPreviousSearchResult;
+						FoundReconstructedPreviousSearchBestResult->PoseCost = NewSearchResult.PoseCost;
+					}
+				}
+				else
+				{
+					FSearchResult& NewReconstructedPreviousSearchBestResult = ReconstructedPreviousSearchBestResultMap.Add(Database);
+					NewReconstructedPreviousSearchBestResult = ReconstructedPreviousSearchResult;
+					NewReconstructedPreviousSearchBestResult.PoseCost = NewSearchResult.PoseCost;
+				}
+
+				if (bForceInterrupt)
+				{
+					// no point to continue if it's bForceInterrupt: no continuing pose searches will be performed 
+					// so all the ReconstructedPreviousSearchResult will lead to the same pose with invalid PoseCost
+					// so any of the pair AssetTime / PoseIdx will be the same
+					break;
+				}
 			}
 		}
-
-		AssetsToSearchPerDatabaseMap.Reset();
 	}
 
-	// collecting all the other databases searches
-	if (!AssetsToSearch.IsEmpty())
+	// performing all the other databases searches
+	for (const TAssetsToSearchPerDatabasePair& AssetsToSearchPerDatabasePair : AssetsToSearchPerDatabaseMap)
 	{
-		for (const UObject* AssetToSearch : AssetsToSearch)
+		const UPoseSearchDatabase* Database = AssetsToSearchPerDatabasePair.Key;
+		check(Database);
+
+		// setting up the best continuing pose environment from the continuing pose searches we already performed
+		if (FSearchResult* ReconstructedPreviousSearchBestResult = ReconstructedPreviousSearchBestResultMap.Find(Database))
 		{
-			if (AddToSearch(AssetsToSearchPerDatabaseMap, AssetToSearch))
-			{
-				bAsyncBuildIndexInProgress = true;
-			}
+			ReconstructedPreviousSearchResult = *ReconstructedPreviousSearchBestResult;
 		}
-
-		for (const TAssetsToSearchPerDatabasePair& AssetsToSearchPerDatabasePair : AssetsToSearchPerDatabaseMap)
+		else
 		{
-			const UPoseSearchDatabase* Database = AssetsToSearchPerDatabasePair.Key;
-			check(Database);
+			ReconstructedPreviousSearchResult.Reset();
+		}
 
-			SearchContext.SetAssetsToConsider(AssetsToSearchPerDatabasePair.Value);
+		SearchContext.SetAssetsToConsider(AssetsToSearchPerDatabasePair.Value);
 
-			const FSearchResult NewSearchResult = Database->Search(SearchContext);
+		// in case we haven't searched the continuing pose for this Database, we haven't created and cached the query yet,
+		// but if we didn't invalidated the continuing pose (when IsInvalidatingContinuingPose is true), we still can reuse
+		// the updated ReconstructedPreviousSearchResult data, and by calling UpdateCurrentResultPoseVector we set the 
+		// SearchContext to be able to create a query for Database using the continuing pose data.
+		SearchContext.UpdateCurrentResultPoseVector();
+		const FSearchResult NewSearchResult = Database->Search(SearchContext);
 
 #if WITH_EDITOR && ENABLE_ANIM_DEBUG && UE_POSE_SEARCH_TRACE_ENABLED
-			const FPoseSearchCost BestBruteForcePoseCost = NewSearchResult.BruteForcePoseCost < SearchResult.BruteForcePoseCost ? NewSearchResult.BruteForcePoseCost : SearchResult.BruteForcePoseCost;
+		const FPoseSearchCost BestBruteForcePoseCost = NewSearchResult.BruteForcePoseCost < SearchResult.BruteForcePoseCost ? NewSearchResult.BruteForcePoseCost : SearchResult.BruteForcePoseCost;
 #endif // WITH_EDITOR && ENABLE_ANIM_DEBUG && UE_POSE_SEARCH_TRACE_ENABLED
 
-			if (NewSearchResult.PoseCost < SearchResult.PoseCost)
-			{
-				SearchResult = NewSearchResult;
-				SearchContext.UpdateCurrentBestCost(SearchResult.PoseCost);
-			}
+		if (NewSearchResult.PoseCost < SearchResult.PoseCost)
+		{
+			SearchResult = NewSearchResult;
+			SearchContext.UpdateCurrentBestCost(SearchResult.PoseCost);
+		}
 
 #if WITH_EDITOR && ENABLE_ANIM_DEBUG && UE_POSE_SEARCH_TRACE_ENABLED
-			SearchResult.BruteForcePoseCost = BestBruteForcePoseCost;
+		SearchResult.BruteForcePoseCost = BestBruteForcePoseCost;
 #endif // WITH_EDITOR && ENABLE_ANIM_DEBUG && UE_POSE_SEARCH_TRACE_ENABLED
-		}
 	}
 
-#if WITH_EDITOR
-	if (bAsyncBuildIndexInProgress)
-	{
-		SearchContext.SetAsyncBuildIndexInProgress();
-	}
-#endif // WITH_EDITOR
-
 #if ENABLE_DRAW_DEBUG && ENABLE_ANIM_DEBUG
 	if (SearchResult.IsValid())
 	{
@@ -1194,7 +1261,9 @@ void UPoseSearchLibrary::MotionMatch(
 	FPoseSearchBlueprintResult& Result,
 	const int32 DebugSessionUniqueIdentifier)
 {
+	PRAGMA_DISABLE_DEPRECATION_WARNINGS
 	MotionMatch(AnimInstances, Roles, AssetsToSearch, PoseHistoryName, ContinuingProperties, Future, Result);
+	PRAGMA_ENABLE_DEPRECATION_WARNINGS
 }
 
 UE::PoseSearch::FSearchResult UPoseSearchLibrary::MotionMatch(
diff --git a/Engine/Plugins/Animation/PoseSearch/Source/Runtime/Private/PoseSearchModule.cpp b/Engine/Plugins/Animation/PoseSearch/Source/Runtime/Private/PoseSearchModule.cpp
--- a/Engine/Plugins/Animation/PoseSearch/Source/Runtime/Private/PoseSearchModule.cpp
+++ b/Engine/Plugins/Animation/PoseSearch/Source/Runtime/Private/PoseSearchModule.cpp
@@ -74,7 +74,7 @@ public:
 			if (const FPoseSearchDatabaseAnimationAssetBase* DatabaseAnimationAssetBase = Database->GetDatabaseAnimationAsset<FPoseSearchDatabaseAnimationAssetBase>(*SearchIndexAsset))
 			{
 				ProviderResult.SelectedAsset = DatabaseAnimationAssetBase->GetAnimationAsset();
-				ProviderResult.Dissimilarity = SearchResult.PoseCost.GetTotalCost();
+				ProviderResult.Dissimilarity = SearchResult.PoseCost;
 				ProviderResult.TimeOffsetSeconds = SearchResult.AssetTime;
 				ProviderResult.bIsFromContinuingPlaying = SearchResult.bIsContinuingPoseSearch;
 				ProviderResult.bMirrored = SearchIndexAsset->IsMirrored();
diff --git a/Engine/Plugins/Animation/PoseSearch/Source/Runtime/Public/PoseSearch/PoseSearchCost.h b/Engine/Plugins/Animation/PoseSearch/Source/Runtime/Public/PoseSearch/PoseSearchCost.h
--- a/Engine/Plugins/Animation/PoseSearch/Source/Runtime/Public/PoseSearch/PoseSearchCost.h
+++ b/Engine/Plugins/Animation/PoseSearch/Source/Runtime/Public/PoseSearch/PoseSearchCost.h
@@ -21,9 +21,10 @@ public:
 	static bool IsCostValid(const float Cost) { return Cost != MAX_flt; }
 
 	bool IsValid() const { return IsCostValid(TotalCost); }
+	
+	UE_DEPRECATED(5.6, "Use operator float() instead.")
 	float GetTotalCost() const { return TotalCost; }
-	bool operator<(const FPoseSearchCost& Other) const { return TotalCost < Other.TotalCost; }
-	bool operator==(const FPoseSearchCost& Other) const { return TotalCost == Other.TotalCost; }
+	operator float() const { return TotalCost; }
 
 protected:
 	UPROPERTY()
diff --git a/Engine/Plugins/Animation/PoseSearch/Source/Runtime/Public/PoseSearch/PoseSearchLibrary.h b/Engine/Plugins/Animation/PoseSearch/Source/Runtime/Public/PoseSearch/PoseSearchLibrary.h
--- a/Engine/Plugins/Animation/PoseSearch/Source/Runtime/Public/PoseSearch/PoseSearchLibrary.h
+++ b/Engine/Plugins/Animation/PoseSearch/Source/Runtime/Public/PoseSearch/PoseSearchLibrary.h
@@ -124,6 +124,9 @@ public:
 	// Currently playing animation accumulated time
 	UPROPERTY(EditAnywhere, BlueprintReadWrite, Category=State)
 	float PlayingAssetAccumulatedTime = 0.f;
+
+	UPROPERTY(EditAnywhere, BlueprintReadWrite, Category=State)
+	EPoseSearchInterruptMode InterruptMode = EPoseSearchInterruptMode::DoNotInterrupt;
 };
 
 UCLASS()
@@ -228,6 +231,7 @@ public:
 		const FPoseSearchFutureProperties Future,
 		FPoseSearchBlueprintResult& Result);
 
+	UE_DEPRECATED(5.6, "Use other MotionMatch signatures instead")
 	static void MotionMatch(
 		const TArrayView<UAnimInstance*> AnimInstances,
 		const TArrayView<const UE::PoseSearch::FRole> Roles,
