pose search - deprecating FMotionMatchingState::bJumpedToPose

diff --git a/Engine/Plugins/Animation/PoseSearch/Source/Runtime/Private/AnimNode_MotionMatching.cpp b/Engine/Plugins/Animation/PoseSearch/Source/Runtime/Private/AnimNode_MotionMatching.cpp
--- a/Engine/Plugins/Animation/PoseSearch/Source/Runtime/Private/AnimNode_MotionMatching.cpp
+++ b/Engine/Plugins/Animation/PoseSearch/Source/Runtime/Private/AnimNode_MotionMatching.cpp
@@ -183,10 +183,16 @@ void FAnimNode_MotionMatching::UpdateAssetPlayer(const FAnimationUpdateContext&
 		#endif // ENABLE_ANIM_DEBUG
 	);
 
+	const bool bJumpToPose = MotionMatchingState.CurrentSearchResult.IsValid() && !MotionMatchingState.CurrentSearchResult.bIsContinuingPoseSearch;
+
+	PRAGMA_DISABLE_DEPRECATION_WARNINGS
+	MotionMatchingState.bJumpedToPose = bJumpToPose;
+	PRAGMA_ENABLE_DEPRECATION_WARNINGS
+
 	UE::Anim::FNodeFunctionCaller::CallFunction(GetOnUpdateMotionMatchingStateFunction(), Context, *this);
 
 	// If a new pose is requested, blend into the new asset via BlendStackNode
-	if (MotionMatchingState.bJumpedToPose)
+	if (bJumpToPose)
 	{
 		const FSearchIndexAsset* SearchIndexAsset = MotionMatchingState.CurrentSearchResult.GetSearchIndexAsset();
 		const UPoseSearchDatabase* CurrentResultDatabase = MotionMatchingState.CurrentSearchResult.Database.Get();
@@ -240,7 +246,7 @@ void FAnimNode_MotionMatching::UpdateAssetPlayer(const FAnimationUpdateContext&
 		}
 	}
 
-	const bool bDidBlendToRequestAnInertialBlend = MotionMatchingState.bJumpedToPose && bUseInertialBlend;
+	const bool bDidBlendToRequestAnInertialBlend = bJumpToPose && bUseInertialBlend;
 	UE::Anim::TOptionalScopedGraphMessage<UE::Anim::FAnimInertializationSyncScope> InertializationSync(bDidBlendToRequestAnInertialBlend, Context);
 	
 	FAnimNode_BlendStack_Standalone::UpdatePlayRate(MotionMatchingState.WantedPlayRate * PlayRateMultiplier);
diff --git a/Engine/Plugins/Animation/PoseSearch/Source/Runtime/Private/PoseSearchLibrary.cpp b/Engine/Plugins/Animation/PoseSearch/Source/Runtime/Private/PoseSearchLibrary.cpp
--- a/Engine/Plugins/Animation/PoseSearch/Source/Runtime/Private/PoseSearchLibrary.cpp
+++ b/Engine/Plugins/Animation/PoseSearch/Source/Runtime/Private/PoseSearchLibrary.cpp
@@ -268,9 +268,9 @@ void FMotionMatchingState::Reset(const FTransform& ComponentTransform)
 	// Set the elapsed time to INFINITY to trigger a search right away
 	ElapsedPoseSearchTime = std::numeric_limits<float>::infinity();
 	WantedPlayRate = 1.f;
-	bJumpedToPose = false;
 	
 PRAGMA_DISABLE_DEPRECATION_WARNINGS
+	bJumpedToPose = false;
 	ComponentDeltaYaw = 0.f;
 	ComponentWorldYaw = FRotator(ComponentTransform.GetRotation()).Yaw;
 	AnimationDeltaYaw = 0.f;
@@ -289,7 +289,9 @@ void FMotionMatchingState::JumpToPose(const FAnimationUpdateContext& Context, co
 	// Remember which pose and sequence we're playing from the database
 	CurrentSearchResult = Result;
 
+PRAGMA_DISABLE_DEPRECATION_WARNINGS
 	bJumpedToPose = true;
+PRAGMA_ENABLE_DEPRECATION_WARNINGS
 }
 
 FVector FMotionMatchingState::GetEstimatedFutureRootMotionVelocity() const
@@ -352,15 +354,32 @@ void FMotionMatchingState::UpdateWantedPlayRate(const UE::PoseSearch::FSearchCon
 }
 
 #if UE_POSE_SEARCH_TRACE_ENABLED
-void UPoseSearchLibrary::TraceMotionMatching(UE::PoseSearch::FSearchContext& SearchContext, FMotionMatchingState& CurrentState, float DeltaTime, bool bSearch, float RecordingTime)
+void UPoseSearchLibrary::TraceMotionMatching(
+		UE::PoseSearch::FSearchContext& SearchContext,
+		const UE::PoseSearch::FSearchResult& SearchResult,
+		float ElapsedPoseSearchTime,
+		float DeltaTime,
+		bool bSearch,
+		float WantedPlayRate)
 {
 	using namespace UE::PoseSearch;
-	
+
+	const bool bChannelEnabled = UE_TRACE_CHANNELEXPR_IS_ENABLED(PoseSearchChannel);
+	if (!bChannelEnabled)
+	{
+		return;
+	}
+
+	float RecordingTime = 0.f;
+	if (!SearchContext.GetAnimContexts().IsEmpty())
+	{
+		RecordingTime = FObjectTrace::GetWorldElapsedTime(SearchContext.GetAnimContexts()[0]->GetWorld());
+	}
+
 	uint32 SearchId = 787;
 
 	FTraceMotionMatchingStateMessage TraceState;
-	FSearchResult& CurrentResult = CurrentState.CurrentSearchResult;
-	const float ElapsedPoseSearchTime = CurrentState.ElapsedPoseSearchTime;
+	const FSearchResult& CurrentResult = SearchResult;
 
 	const int32 AnimContextsNum = SearchContext.GetAnimContexts().Num();
 	TraceState.SkeletalMeshComponentIds.SetNum(AnimContextsNum);
@@ -369,32 +388,35 @@ void UPoseSearchLibrary::TraceMotionMatching(UE::PoseSearch::FSearchContext& Sea
 	{
 		if (const UObject* AnimContext = SearchContext.GetAnimContexts()[AnimInstanceIndex])
 		{
+			const UObject* SkeletalMeshComponent = nullptr;
 			if (const UAnimInstance* AnimInstance = Cast<UAnimInstance>(AnimContext))
 			{
-				const UObject* SkeletalMeshComponent = AnimInstance->GetOuter();
-
-				TRACE_OBJECT(AnimInstance);
-
-				TraceState.SkeletalMeshComponentIds[AnimInstanceIndex] = FObjectTrace::GetObjectId(SkeletalMeshComponent);
-
-				SearchId = HashCombineFast(SearchId, GetTypeHash(FObjectTrace::GetObjectId(AnimInstance)));
+				SkeletalMeshComponent = AnimInstance->GetOuter();
 			}
 			else if (const UActorComponent* ActorComponent = Cast<UActorComponent>(AnimContext))
 			{
 				const AActor* Actor = ActorComponent->GetOwner();
 				check(Actor);
 				
-				const UObject* SkeletalMeshComponent = Actor->GetComponentByClass<USkeletalMeshComponent>();
-
-				TRACE_OBJECT(ActorComponent);
-
-				TraceState.SkeletalMeshComponentIds[AnimInstanceIndex] = FObjectTrace::GetObjectId(SkeletalMeshComponent);
+				SkeletalMeshComponent = Actor->GetComponentByClass<USkeletalMeshComponent>();
+			}
 
-				SearchId = HashCombineFast(SearchId, GetTypeHash(FObjectTrace::GetObjectId(ActorComponent)));
+			if (!SkeletalMeshComponent || CANNOT_TRACE_OBJECT(SkeletalMeshComponent))
+			{
+				return;
 			}
+			
+			TraceState.SkeletalMeshComponentIds[AnimInstanceIndex] = FObjectTrace::GetObjectId(SkeletalMeshComponent);
 		}
 	}
 
+	for (int32 AnimInstanceIndex = 0; AnimInstanceIndex < AnimContextsNum; ++AnimInstanceIndex)
+	{
+		const UObject* AnimContext = SearchContext.GetAnimContexts()[AnimInstanceIndex];
+		TRACE_OBJECT(AnimContext);
+		SearchId = HashCombineFast(SearchId, GetTypeHash(FObjectTrace::GetObjectId(AnimContext)));
+	}
+
 	TraceState.Roles.SetNum(AnimContextsNum);
 	for (const FRoleToIndexPair& RoleToIndexPair : SearchContext.GetRoleToIndex())
 	{
@@ -488,7 +510,7 @@ void UPoseSearchLibrary::TraceMotionMatching(UE::PoseSearch::FSearchContext& Sea
 				// Simulate the time step to get accurate root motion prediction for this frame.
 				FAnimationAssetSampler Sampler(AnimationAsset);
 
-				const float TimeStep = DeltaTime * CurrentState.WantedPlayRate;
+				const float TimeStep = DeltaTime * WantedPlayRate;
 				const FTransform PrevRoot = Sampler.ExtractRootTransform(CurrentResult.AssetTime);
 				const FTransform CurrRoot = Sampler.ExtractRootTransform(CurrentResult.AssetTime + TimeStep);
 				const FTransform RootMotionTransformDelta = PrevRoot.GetRelativeTransform(CurrRoot);
@@ -502,7 +524,7 @@ void UPoseSearchLibrary::TraceMotionMatching(UE::PoseSearch::FSearchContext& Sea
 				TraceState.AnimAngularVelocityNoTimescale = FMath::RadiansToDegrees(RootMotionTransformDeltaNoTimescale.GetRotation().GetAngle()) / DeltaTime;
 			}
 		}
-		TraceState.Playrate = CurrentState.WantedPlayRate;
+		TraceState.Playrate = WantedPlayRate;
 	}
 
 	TraceState.ElapsedPoseSearchTime = ElapsedPoseSearchTime;
@@ -560,11 +582,8 @@ void UPoseSearchLibrary::UpdateMotionMatchingState(
 			LOCTEXT("NoDatabases", "No database assets provided for motion matching."));
 		return;
 	}
-
+	
 	const float DeltaTime = Context.GetDeltaTime();
-
-	InOutMotionMatchingState.bJumpedToPose = false;
-
 	const IPoseHistory* PoseHistory = nullptr;
 	if (FPoseHistoryProvider* PoseHistoryProvider = Context.GetMessage<FPoseHistoryProvider>())
 	{
@@ -605,7 +624,6 @@ void UPoseSearchLibrary::UpdateMotionMatchingState(
 			SearchContext.UpdateCurrentBestCost(SearchResult.PoseCost);
 		}
 
-		bool bJumpToPose = false;
 		for (const TObjectPtr<const UPoseSearchDatabase>& Database : Databases)
 		{
 			if (ensure(Database))
@@ -618,7 +636,6 @@ void UPoseSearchLibrary::UpdateMotionMatchingState(
 
 				if (NewSearchResult.PoseCost < SearchResult.PoseCost)
 				{
-					bJumpToPose = true;
 					SearchResult = NewSearchResult;
 					SearchContext.UpdateCurrentBestCost(SearchResult.PoseCost);
 				}
@@ -663,7 +680,7 @@ void UPoseSearchLibrary::UpdateMotionMatchingState(
 					StringBuilder << ", ";
 				}
 			}
-			
+
 			StringBuilder << "] ";
 
 			FString String = StringBuilder.ToString();
@@ -679,18 +696,8 @@ void UPoseSearchLibrary::UpdateMotionMatchingState(
 		}
 #endif // !NO_LOGGING
 
-		if (bJumpToPose)
-		{
-			InOutMotionMatchingState.JumpToPose(Context, SearchResult, MaxActiveBlends, BlendTime);
-		}
-		else
-		{
-			// copying few properties of SearchResult into CurrentSearchResult to facilitate debug drawing
-#if WITH_EDITOR && ENABLE_ANIM_DEBUG && UE_POSE_SEARCH_TRACE_ENABLED
-			InOutMotionMatchingState.CurrentSearchResult.BruteForcePoseCost = SearchResult.BruteForcePoseCost;
-#endif // WITH_EDITOR && ENABLE_ANIM_DEBUG && UE_POSE_SEARCH_TRACE_ENABLED
-			InOutMotionMatchingState.CurrentSearchResult.PoseCost = SearchResult.PoseCost;
-		}
+		// Remember which pose and sequence we're playing from the database
+		InOutMotionMatchingState.CurrentSearchResult = SearchResult;
 	}
 	else
 	{
@@ -703,12 +710,7 @@ void UPoseSearchLibrary::UpdateMotionMatchingState(
 	InOutMotionMatchingState.PoseIndicesHistory.Update(InOutMotionMatchingState.CurrentSearchResult, DeltaTime, PoseReselectHistory);
 
 #if UE_POSE_SEARCH_TRACE_ENABLED
-	// Record debugger details
-	if (IsTracing(Context))
-	{
-		TraceMotionMatching(SearchContext, InOutMotionMatchingState, DeltaTime, bSearch,
-			AnimInstance ? FObjectTrace::GetWorldElapsedTime(AnimInstance->GetWorld()) : 0.f);
-	}
+	TraceMotionMatching(SearchContext, InOutMotionMatchingState.CurrentSearchResult, InOutMotionMatchingState.ElapsedPoseSearchTime, DeltaTime, bSearch, InOutMotionMatchingState.WantedPlayRate);
 #endif // UE_POSE_SEARCH_TRACE_ENABLED
 
 #if WITH_EDITORONLY_DATA && ENABLE_ANIM_DEBUG
@@ -1173,15 +1175,7 @@ UE::PoseSearch::FSearchResult UPoseSearchLibrary::MotionMatch(
 #endif // ENABLE_DRAW_DEBUG && ENABLE_ANIM_DEBUG
 
 #if UE_POSE_SEARCH_TRACE_ENABLED
-	FMotionMatchingState MotionMatchingState;
-	MotionMatchingState.CurrentSearchResult = SearchResult;
-	MotionMatchingState.ElapsedPoseSearchTime = 0.0f;
-	double WorldElapsedTime = 0;
-	if (const UAnimInstance* AnimInstance = Cast<UAnimInstance>(AnimContexts[0]))
-	{
-		WorldElapsedTime = FObjectTrace::GetWorldElapsedTime(AnimInstance->GetWorld());
-	}
-	TraceMotionMatching(SearchContext, MotionMatchingState, DeltaSeconds, true, WorldElapsedTime);
+	TraceMotionMatching(SearchContext, SearchResult, 0.f, DeltaSeconds, true, 1.f);
 #endif // UE_POSE_SEARCH_TRACE_ENABLED
 
 	return SearchResult;
diff --git a/Engine/Plugins/Animation/PoseSearch/Source/Runtime/Private/Trace/PoseSearchTraceLogger.cpp b/Engine/Plugins/Animation/PoseSearch/Source/Runtime/Private/Trace/PoseSearchTraceLogger.cpp
--- a/Engine/Plugins/Animation/PoseSearch/Source/Runtime/Private/Trace/PoseSearchTraceLogger.cpp
+++ b/Engine/Plugins/Animation/PoseSearch/Source/Runtime/Private/Trace/PoseSearchTraceLogger.cpp
@@ -21,27 +21,6 @@ namespace UE::PoseSearch
 const FName FTraceLogger::Name("PoseSearch");
 const FName FTraceMotionMatchingStateMessage::Name("MotionMatchingState");
 
-bool IsTracing(const FAnimationBaseContext& InContext)
-{
-#if UE_POSE_SEARCH_TRACE_ENABLED
-	const bool bChannelEnabled = UE_TRACE_CHANNELEXPR_IS_ENABLED(PoseSearchChannel);
-	if (!bChannelEnabled)
-	{
-		return false;
-	}
-
-	if (InContext.GetCurrentNodeId() == INDEX_NONE)
-	{
-		return false;
-	}
-
-	check(InContext.AnimInstanceProxy);
-	return !CANNOT_TRACE_OBJECT(InContext.AnimInstanceProxy->GetSkelMeshComponent());
-#else // UE_POSE_SEARCH_TRACE_ENABLED
-	return false;
-#endif // UE_POSE_SEARCH_TRACE_ENABLED
-}
-
 FArchive& operator<<(FArchive& Ar, FTraceMessage& State)
 {
 	Ar << State.Cycle;
diff --git a/Engine/Plugins/Animation/PoseSearch/Source/Runtime/Public/PoseSearch/PoseSearchLibrary.h b/Engine/Plugins/Animation/PoseSearch/Source/Runtime/Public/PoseSearch/PoseSearchLibrary.h
--- a/Engine/Plugins/Animation/PoseSearch/Source/Runtime/Public/PoseSearch/PoseSearchLibrary.h
+++ b/Engine/Plugins/Animation/PoseSearch/Source/Runtime/Public/PoseSearch/PoseSearchLibrary.h
@@ -55,6 +55,7 @@ PRAGMA_ENABLE_DEPRECATION_WARNINGS
 	void AdjustAssetTime(float AssetTime);
 
 	// Internally stores the 'jump' to a new pose/sequence index and asset time for evaluation
+	UE_DEPRECATED(5.6, "Set the CurrentSearchResult directly if needed, since bJumpedToPose has already been deprecated")
 	void JumpToPose(const FAnimationUpdateContext& Context, const UE::PoseSearch::FSearchResult& Result, int32 MaxActiveBlends, float BlendTime);
 
 	void UpdateWantedPlayRate(const UE::PoseSearch::FSearchContext& SearchContext, const FFloatInterval& PlayRate, float TrajectorySpeedMultiplier);
@@ -70,6 +71,7 @@ PRAGMA_ENABLE_DEPRECATION_WARNINGS
 	float WantedPlayRate = 1.f;
 
 	// true if a new animation has been selected
+	UE_DEPRECATED(5.6, "Derive this value via 'CurrentSearchResult.IsValid() && !CurrentSearchResult.bIsContinuingPoseSearch' if needed")
 	bool bJumpedToPose = false;
 
 	UE::PoseSearch::FPoseIndicesHistory PoseIndicesHistory;
@@ -136,12 +138,24 @@ class POSESEARCH_API UPoseSearchLibrary : public UBlueprintFunctionLibrary
 
 #if UE_POSE_SEARCH_TRACE_ENABLED
 
+	static void TraceMotionMatching(
+		UE::PoseSearch::FSearchContext& SearchContext,
+		const UE::PoseSearch::FSearchResult& SearchResult,
+		float ElapsedPoseSearchTime,
+		float DeltaTime,
+		bool bSearch,
+		float WantedPlayRate);
+
+	UE_DEPRECATED(5.6, "Use TraceMotionMatching with different signature instead.")
 	static void TraceMotionMatching(
 		UE::PoseSearch::FSearchContext& SearchContext,
 		FMotionMatchingState& CurrentState,
 		float DeltaTime,
 		bool bSearch,
-		float RecordingTime);
+		float RecordingTime)
+	{
+		TraceMotionMatching(SearchContext, CurrentState.CurrentSearchResult, CurrentState.ElapsedPoseSearchTime, DeltaTime, bSearch, CurrentState.WantedPlayRate);
+	}
 	
 	UE_DEPRECATED(5.5, "Use TraceMotionMatching with different signature instead.")
 	static void TraceMotionMatching(
@@ -153,10 +167,7 @@ class POSESEARCH_API UPoseSearchLibrary : public UBlueprintFunctionLibrary
 		bool bSearch,
 		float RecordingTime)
 	{
-		FMotionMatchingState MotionMatchingState;
-		MotionMatchingState.CurrentSearchResult = CurrentResult;
-		MotionMatchingState.ElapsedPoseSearchTime = ElapsedPoseSearchTime;
-		TraceMotionMatching(SearchContext, MotionMatchingState, DeltaTime, bSearch, RecordingTime);
+		TraceMotionMatching(SearchContext, CurrentResult, ElapsedPoseSearchTime, DeltaTime, bSearch, 1.f);
 	}
 
 	UE_DEPRECATED(5.4, "Use TraceMotionMatching instead")
@@ -170,10 +181,7 @@ class POSESEARCH_API UPoseSearchLibrary : public UBlueprintFunctionLibrary
 		bool bSearch,
 		float RecordingTime)
 	{
-		FMotionMatchingState MotionMatchingState;
-		MotionMatchingState.CurrentSearchResult = CurrentResult;
-		MotionMatchingState.ElapsedPoseSearchTime = ElapsedPoseSearchTime;
-		TraceMotionMatching(SearchContext, MotionMatchingState, DeltaTime, bSearch, RecordingTime);
+		TraceMotionMatching(SearchContext, CurrentResult, ElapsedPoseSearchTime, DeltaTime, bSearch, 1.f);
 	}
 	
 #endif // UE_POSE_SEARCH_TRACE_ENABLED
diff --git a/Engine/Plugins/Animation/PoseSearch/Source/Runtime/Public/PoseSearch/Trace/PoseSearchTraceLogger.h b/Engine/Plugins/Animation/PoseSearch/Source/Runtime/Public/PoseSearch/Trace/PoseSearchTraceLogger.h
--- a/Engine/Plugins/Animation/PoseSearch/Source/Runtime/Public/PoseSearch/Trace/PoseSearchTraceLogger.h
+++ b/Engine/Plugins/Animation/PoseSearch/Source/Runtime/Public/PoseSearch/Trace/PoseSearchTraceLogger.h
@@ -11,8 +11,6 @@ UE_TRACE_CHANNEL_EXTERN(PoseSearchChannel, POSESEARCH_API);
 namespace UE::PoseSearch
 {
 
-bool IsTracing(const FAnimationBaseContext& InContext);
-
 struct POSESEARCH_API FTraceLogger
 {
 	/** Used for reading trace data */
