PoseSearch - removing unnecessary editor only allocations,
PoseSearchAnimNext - support for array of databases in MM trait

diff --git a/Engine/Plugins/Animation/PoseSearch/Source/Editor/Private/PoseSearchDatabaseDataDetails.cpp b/Engine/Plugins/Animation/PoseSearch/Source/Editor/Private/PoseSearchDatabaseDataDetails.cpp
--- a/Engine/Plugins/Animation/PoseSearch/Source/Editor/Private/PoseSearchDatabaseDataDetails.cpp
+++ b/Engine/Plugins/Animation/PoseSearch/Source/Editor/Private/PoseSearchDatabaseDataDetails.cpp
@@ -353,6 +353,7 @@ public:
 							const UPoseSearchDatabase* PoseSearchDatabase = ViewModel->GetPoseSearchDatabase();
 							if (EAsyncBuildIndexResult::Success == FAsyncPoseSearchDatabasesManagement::RequestAsyncBuildIndex(PoseSearchDatabase, ERequestAsyncBuildFlag::ContinueRequest))
 							{
+								TArray<float> BufferUsedForReconstruction;
 								for (const TArray<FDatabasePreviewActor>& PreviewActorGroup : ViewModel->GetPreviewActors())
 								{
 									for (const FDatabasePreviewActor& PreviewActor : PreviewActorGroup)
@@ -360,7 +361,8 @@ public:
 										if (*GetNameSafe(PreviewActor.GetActor()) == ColumnName)
 										{
 											const int32 PoseIdx = PreviewActor.GetCurrentPoseIndex();
-											const TArray<float> PoseValues = PoseSearchDatabase->GetSearchIndex().GetPoseValuesSafe(PoseIdx);
+
+											const TConstArrayView<float> PoseValues = PoseSearchDatabase->GetSearchIndex().GetPoseValuesSafe(PoseIdx, BufferUsedForReconstruction);
 
 											if (Cardinality > 1)
 											{
diff --git a/Engine/Plugins/Animation/PoseSearch/Source/Editor/Private/PoseSearchDebuggerDatabaseView.cpp b/Engine/Plugins/Animation/PoseSearch/Source/Editor/Private/PoseSearchDebuggerDatabaseView.cpp
--- a/Engine/Plugins/Animation/PoseSearch/Source/Editor/Private/PoseSearchDebuggerDatabaseView.cpp
+++ b/Engine/Plugins/Animation/PoseSearch/Source/Editor/Private/PoseSearchDebuggerDatabaseView.cpp
@@ -173,12 +173,13 @@ static void AddUnfilteredDatabaseRow(const UPoseSearchDatabase* Database,
 		Row->bMirrored = SearchIndexAsset->IsMirrored();
 
 		Row->CostVector.SetNum(Database->Schema->SchemaCardinality);
-		const TArray<float> PoseValues = SearchIndex.GetPoseValuesSafe(DbPoseIdx);
+		TArray<float> BufferUsedForReconstruction;
+		const TConstArrayView<float> PoseValues = SearchIndex.GetPoseValuesSafe(DbPoseIdx, BufferUsedForReconstruction);
 
 		// in case we modify the schema while PIE is paused and displaying the Pose Search Editor, we could end up with a stale State with a SharedData->QueryVector saved with the previous schema
 		// so the cardinality of SharedData->QueryVector and PoseValues don't match. In that case we just use PoseValues as query to have all costs set to zero
 		const bool bIsQueryVectorValid = SharedData->QueryVector.Num() == PoseValues.Num();
-		TConstArrayView<float> QueryVector = bIsQueryVectorValid ? TConstArrayView<float>(SharedData->QueryVector) : TConstArrayView<float>(PoseValues);
+		const TConstArrayView<float> QueryVector = bIsQueryVectorValid ? TConstArrayView<float>(SharedData->QueryVector) : TConstArrayView<float>(PoseValues);
 
 		CompareFeatureVectors(PoseValues, QueryVector, SearchIndex.WeightsSqrt, Row->CostVector);
 
diff --git a/Engine/Plugins/Animation/PoseSearch/Source/Runtime/Private/PoseSearchContext.cpp b/Engine/Plugins/Animation/PoseSearch/Source/Runtime/Private/PoseSearchContext.cpp
--- a/Engine/Plugins/Animation/PoseSearch/Source/Runtime/Private/PoseSearchContext.cpp
+++ b/Engine/Plugins/Animation/PoseSearch/Source/Runtime/Private/PoseSearchContext.cpp
@@ -477,7 +477,8 @@ void FDebugDrawParams::DrawFeatureVector(int32 PoseIdx)
 {
 	if (CanDraw())
 	{
-		DrawFeatureVector(GetSearchIndex()->GetPoseValuesSafe(PoseIdx));
+		TArray<float> BufferUsedForReconstruction;
+		DrawFeatureVector(GetSearchIndex()->GetPoseValuesSafe(PoseIdx, BufferUsedForReconstruction));
 	}
 }
 #endif // ENABLE_DRAW_DEBUG
diff --git a/Engine/Plugins/Animation/PoseSearch/Source/Runtime/Private/PoseSearchDatabase.cpp b/Engine/Plugins/Animation/PoseSearch/Source/Runtime/Private/PoseSearchDatabase.cpp
--- a/Engine/Plugins/Animation/PoseSearch/Source/Runtime/Private/PoseSearchDatabase.cpp
+++ b/Engine/Plugins/Animation/PoseSearch/Source/Runtime/Private/PoseSearchDatabase.cpp
@@ -1110,6 +1110,10 @@ void UPoseSearchDatabase::PopulateNonSelectableIdx(FNonSelectableIdx& NonSelecta
 
 	const FSearchIndex& SearchIndex = GetSearchIndex();
 
+#if UE_POSE_SEARCH_TRACE_ENABLED
+	TArray<float> BufferUsedForReconstruction;
+#endif // UE_POSE_SEARCH_TRACE_ENABLED
+
 	NonSelectableIdx.Reset();
 	if (SearchContext.IsCurrentResultFromDatabase(this))
 	{
@@ -1132,7 +1136,7 @@ void UPoseSearchDatabase::PopulateNonSelectableIdx(FNonSelectableIdx& NonSelecta
 						NonSelectableIdx.Add(PoseIdx);
 
 #if UE_POSE_SEARCH_TRACE_ENABLED
-						const TArray<float> PoseValues = SearchIndex.GetPoseValuesSafe(PoseIdx);
+						const TConstArrayView<float> PoseValues = SearchIndex.GetPoseValuesSafe(PoseIdx, BufferUsedForReconstruction);
 						const FPoseSearchCost PoseCost = SearchIndex.ComparePoses(PoseIdx, 0.f, PoseValues, QueryValues);
 						SearchContext.Track(this, PoseIdx, EPoseCandidateFlags::DiscardedBy_AssetReselection, PoseCost);
 #endif // UE_POSE_SEARCH_TRACE_ENABLED
@@ -1159,7 +1163,7 @@ void UPoseSearchDatabase::PopulateNonSelectableIdx(FNonSelectableIdx& NonSelecta
 						NonSelectableIdx.AddUnique(PoseIdx);
 
 #if UE_POSE_SEARCH_TRACE_ENABLED
-						const TArray<float> PoseValues = SearchIndex.GetPoseValuesSafe(PoseIdx);
+						const TConstArrayView<float> PoseValues = SearchIndex.GetPoseValuesSafe(PoseIdx, BufferUsedForReconstruction);
 						const FPoseSearchCost PoseCost = SearchIndex.ComparePoses(PoseIdx, 0.f, PoseValues, QueryValues);
 						SearchContext.Track(this, PoseIdx, EPoseCandidateFlags::DiscardedBy_PoseJumpThresholdTime, PoseCost);
 #endif // UE_POSE_SEARCH_TRACE_ENABLED
@@ -1175,7 +1179,7 @@ void UPoseSearchDatabase::PopulateNonSelectableIdx(FNonSelectableIdx& NonSelecta
 						NonSelectableIdx.AddUnique(PoseIdx);
 
 #if UE_POSE_SEARCH_TRACE_ENABLED
-						const TArray<float> PoseValues = SearchIndex.GetPoseValuesSafe(PoseIdx);
+						const TConstArrayView<float> PoseValues = SearchIndex.GetPoseValuesSafe(PoseIdx, BufferUsedForReconstruction);
 						const FPoseSearchCost PoseCost = SearchIndex.ComparePoses(PoseIdx, 0.f, PoseValues, QueryValues);
 						SearchContext.Track(this, PoseIdx, EPoseCandidateFlags::DiscardedBy_PoseJumpThresholdTime, PoseCost);
 #endif // UE_POSE_SEARCH_TRACE_ENABLED
@@ -1201,7 +1205,7 @@ void UPoseSearchDatabase::PopulateNonSelectableIdx(FNonSelectableIdx& NonSelecta
 				// if we're editing the database and removing assets it's possible that the PoseIndicesHistory contains invalid pose indexes
 				if (HistoricalPoseIndex.PoseIndex < SearchIndex.GetNumPoses())
 				{
-					const FPoseSearchCost PoseCost = SearchIndex.ComparePoses(HistoricalPoseIndex.PoseIndex, 0.f, SearchIndex.GetPoseValuesSafe(HistoricalPoseIndex.PoseIndex), QueryValues);
+					const FPoseSearchCost PoseCost = SearchIndex.ComparePoses(HistoricalPoseIndex.PoseIndex, 0.f, SearchIndex.GetPoseValuesSafe(HistoricalPoseIndex.PoseIndex, BufferUsedForReconstruction), QueryValues);
 					SearchContext.Track(this, HistoricalPoseIndex.PoseIndex, EPoseCandidateFlags::DiscardedBy_PoseReselectHistory, PoseCost);
 				}
 #endif // UE_POSE_SEARCH_TRACE_ENABLED
diff --git a/Engine/Plugins/Animation/PoseSearch/Source/Runtime/Private/PoseSearchIndex.cpp b/Engine/Plugins/Animation/PoseSearch/Source/Runtime/Private/PoseSearchIndex.cpp
--- a/Engine/Plugins/Animation/PoseSearch/Source/Runtime/Private/PoseSearchIndex.cpp
+++ b/Engine/Plugins/Animation/PoseSearch/Source/Runtime/Private/PoseSearchIndex.cpp
@@ -687,6 +687,23 @@ TArray<float> FSearchIndex::GetPoseValuesSafe(int32 PoseIdx) const
 	return PoseValues;
 }
 
+TConstArrayView<float> FSearchIndex::GetPoseValuesSafe(int32 PoseIdx, TArray<float>& BufferUsedForReconstruction) const
+{
+	if (PoseIdx < 0 || PoseIdx >= GetNumPoses())
+	{
+		return TConstArrayView<float>();
+	}
+
+	if (IsValuesEmpty())
+	{
+		const int32 NumDimensions = GetNumDimensions();
+		BufferUsedForReconstruction.Reset(NumDimensions);
+		return GetReconstructedPoseValues(PoseIdx, BufferUsedForReconstruction);
+	}
+		
+	return GetPoseValues(PoseIdx);
+}
+
 TConstArrayView<float> FSearchIndex::GetPCAPoseValues(int32 PCAValuesVectorIdx) const
 {
 	if (PCAValues.IsEmpty())
diff --git a/Engine/Plugins/Animation/PoseSearch/Source/Runtime/Private/PoseSearchLibrary.cpp b/Engine/Plugins/Animation/PoseSearch/Source/Runtime/Private/PoseSearchLibrary.cpp
--- a/Engine/Plugins/Animation/PoseSearch/Source/Runtime/Private/PoseSearchLibrary.cpp
+++ b/Engine/Plugins/Animation/PoseSearch/Source/Runtime/Private/PoseSearchLibrary.cpp
@@ -228,7 +228,7 @@ namespace UE::PoseSearch
 		}
 	}
 
-	static bool ShouldUseCachedChannelData(const UPoseSearchDatabase* CurrentResultDatabase, const TArray<TObjectPtr<const UPoseSearchDatabase>>& Databases)
+	static bool ShouldUseCachedChannelData(const UPoseSearchDatabase* CurrentResultDatabase, const TConstArrayView<TObjectPtr<const UPoseSearchDatabase>> Databases)
 	{
 		const UPoseSearchSchema* OneOfTheSchemas = nullptr;
 		if (CurrentResultDatabase)
@@ -588,7 +588,7 @@ void UPoseSearchLibrary::UpdateMotionMatchingState(
 void UPoseSearchLibrary::UpdateMotionMatchingState(
 	const UObject* AnimContext,
 	const UE::PoseSearch::IPoseHistory* PoseHistory,
-	const TArray<TObjectPtr<const UPoseSearchDatabase>>& Databases,
+	const TConstArrayView<TObjectPtr<const UPoseSearchDatabase>> Databases,
 	float DeltaTime,
 	const FFloatInterval& PoseJumpThresholdTime,
 	float PoseReselectHistory,
diff --git a/Engine/Plugins/Animation/PoseSearch/Source/Runtime/Public/PoseSearch/PoseSearchIndex.h b/Engine/Plugins/Animation/PoseSearch/Source/Runtime/Public/PoseSearch/PoseSearchIndex.h
--- a/Engine/Plugins/Animation/PoseSearch/Source/Runtime/Public/PoseSearch/PoseSearchIndex.h
+++ b/Engine/Plugins/Animation/PoseSearch/Source/Runtime/Public/PoseSearch/PoseSearchIndex.h
@@ -441,7 +441,9 @@ struct FSearchIndex : public FSearchIndexBase
 	int32 GetNumberOfPrincipalComponents() const;
 	POSESEARCH_API TConstArrayView<float> PCAProject(TConstArrayView<float> PoseValues, TArrayView<float> BufferUsedForProjection) const;
 
+	UE_DEPRECATED(5.6, "Use other GetPoseValuesSafe signatures instead")
 	POSESEARCH_API TArray<float> GetPoseValuesSafe(int32 PoseIdx) const;
+	POSESEARCH_API TConstArrayView<float> GetPoseValuesSafe(int32 PoseIdx, TArray<float>& BufferUsedForReconstruction) const;
 
 	// since PCAValues (as well as Values can be pruned out from duplicate data, we lose the 1:1 mapping between PoseIdx and PCAValuesVectorIdx
 	// that in the case of GetPoseValuesSafe it's stored in PoseMetadata[PoseIdx].GetValueOffset(), but missing for the PCAValues, so this API input is NOT a PoseIdx
diff --git a/Engine/Plugins/Animation/PoseSearch/Source/Runtime/Public/PoseSearch/PoseSearchLibrary.h b/Engine/Plugins/Animation/PoseSearch/Source/Runtime/Public/PoseSearch/PoseSearchLibrary.h
--- a/Engine/Plugins/Animation/PoseSearch/Source/Runtime/Public/PoseSearch/PoseSearchLibrary.h
+++ b/Engine/Plugins/Animation/PoseSearch/Source/Runtime/Public/PoseSearch/PoseSearchLibrary.h
@@ -207,7 +207,7 @@ public:
 	static void UpdateMotionMatchingState(
 		const UObject* AnimContext,
 		const UE::PoseSearch::IPoseHistory* PoseHistory,
-		const TArray<TObjectPtr<const UPoseSearchDatabase>>& Databases,
+		const TConstArrayView<TObjectPtr<const UPoseSearchDatabase>> Databases,
 		float DeltaTime,
 		const FFloatInterval& PoseJumpThresholdTime,
 		float PoseReselectHistory,
diff --git a/Engine/Plugins/Experimental/AnimNextPoseSearch/Source/AnimNextPoseSearch/Private/MotionMatchingTrait.cpp b/Engine/Plugins/Experimental/AnimNextPoseSearch/Source/AnimNextPoseSearch/Private/MotionMatchingTrait.cpp
--- a/Engine/Plugins/Experimental/AnimNextPoseSearch/Source/AnimNextPoseSearch/Private/MotionMatchingTrait.cpp
+++ b/Engine/Plugins/Experimental/AnimNextPoseSearch/Source/AnimNextPoseSearch/Private/MotionMatchingTrait.cpp
@@ -91,12 +91,7 @@ namespace UE::AnimNext
 			MotionMatchingState.AdjustAssetTime(TimelineProgress.GetPosition());
 		}
 
-		TArray<TObjectPtr<const UPoseSearchDatabase>> DatabasesToSearch;
-		if (const UPoseSearchDatabase* DataBase = Cast<UPoseSearchDatabase>(SharedData->GetDatabase(Binding)))
-		{
-			DatabasesToSearch.Add(DataBase);
-		}
-
+		const TConstArrayView<TObjectPtr<const UPoseSearchDatabase>> DatabasesToSearch = SharedData->GetDatabases(Binding);
 		if (DatabasesToSearch.IsEmpty())
 		{
 			UE_LOG(LogPoseSearch, Error, TEXT("FMotionMatchingTrait::PreUpdate, No database assets provided for motion matching."));
diff --git a/Engine/Plugins/Experimental/AnimNextPoseSearch/Source/AnimNextPoseSearch/Private/MotionMatchingTrait.h b/Engine/Plugins/Experimental/AnimNextPoseSearch/Source/AnimNextPoseSearch/Private/MotionMatchingTrait.h
--- a/Engine/Plugins/Experimental/AnimNextPoseSearch/Source/AnimNextPoseSearch/Private/MotionMatchingTrait.h
+++ b/Engine/Plugins/Experimental/AnimNextPoseSearch/Source/AnimNextPoseSearch/Private/MotionMatchingTrait.h
@@ -22,11 +22,8 @@ struct FMotionMatchingTraitSharedData : public FAnimNextTraitSharedData
 	UPROPERTY(EditAnywhere, Category = "Default", meta=(CustomWidget = "ParamName", AllowedParamType = "FAnimNextEntryPoint"))
 	FName EntryPoint;
 
-	// @todo: support array of databases
-	// @todo: perhaps support databases AND choosers to be able to access chooser metadata =);>
-	// database to search from
 	UPROPERTY(EditAnywhere, Category = "Default")
-	TObjectPtr<const UObject> Database;
+	TArray<TObjectPtr<const UPoseSearchDatabase>> Databases;
 
 	UPROPERTY(EditAnywhere, Category = "Default")
 	FAlphaBlendArgs BlendArguments;
@@ -66,7 +63,7 @@ struct FMotionMatchingTraitSharedData : public FAnimNextTraitSharedData
 	#define TRAIT_LATENT_PROPERTIES_ENUMERATOR(GeneratorMacro) \
 		GeneratorMacro(AnimationGraph) \
 		GeneratorMacro(EntryPoint) \
-		GeneratorMacro(Database) \
+		GeneratorMacro(Databases) \
 		GeneratorMacro(BlendArguments) \
 		GeneratorMacro(PoseJumpThresholdTime) \
 		GeneratorMacro(PoseReselectHistory) \
