Unshelved from pending changelist '36821343':

Add a vlog version of the pose search trajectory debug to have rewind debugger support

diff --git a/Engine/Plugins/Animation/PoseSearch/Source/Runtime/Private/PoseSearchTrajectoryTypes.cpp b/Engine/Plugins/Animation/PoseSearch/Source/Runtime/Private/PoseSearchTrajectoryTypes.cpp
--- a/Engine/Plugins/Animation/PoseSearch/Source/Runtime/Private/PoseSearchTrajectoryTypes.cpp
+++ b/Engine/Plugins/Animation/PoseSearch/Source/Runtime/Private/PoseSearchTrajectoryTypes.cpp
@@ -7,6 +7,7 @@
 #include "DrawDebugHelpers.h"
 #include "PoseSearch/PoseSearchDefines.h"
 #include "Math/Axis.h"
+#include "VisualLogger/VisualLogger.h"
 
 FPoseSearchQueryTrajectorySample FPoseSearchQueryTrajectorySample::Lerp(const FPoseSearchQueryTrajectorySample& Other, float Alpha) const
 {
@@ -152,6 +153,51 @@ void FPoseSearchQueryTrajectory::DebugDrawTrajectory(FAnimInstanceProxy& AnimIns
 	}
 }
 
+void FPoseSearchQueryTrajectory::DebugDrawTrajectory(const UObject* Owner, const FLogCategoryBase& Category, ELogVerbosity::Type Verbosity, const float DebugThickness, float HeightOffset) const
+{
+#if ENABLE_VISUAL_LOG
+
+	UWorld* World = nullptr;
+	FVisualLogEntry* CurrentEntry = nullptr;
+	if (FVisualLogger::CheckVisualLogInputInternal(Owner, Category.GetCategoryName(), Verbosity, &World, &CurrentEntry) == false)
+	{
+		return;
+	}
+
+	const FVector OffsetVector = FVector::UpVector * HeightOffset;
+
+	const int32 LastIndex = Samples.Num() - 1;
+	if (LastIndex >= 0)
+	{
+		for (int32 Index = 0; ; ++Index)
+		{
+			const FVector Pos = Samples[Index].Position + OffsetVector;
+
+			CurrentEntry->AddSphere(Pos, 1, Category.GetCategoryName(), Verbosity, FColor::Black, TEXT(""), true);
+
+			const FRotationMatrix R(FRotator(Samples[Index].Facing));
+			const FVector X = R.GetScaledAxis(EAxis::X);
+			const FVector Y = R.GetScaledAxis(EAxis::Y);
+
+			const float Scale = 12.f;
+
+			const bool IsPast = Samples[Index].AccumulatedSeconds <= 0.f;
+
+			CurrentEntry->AddSegment(Pos, Pos + X * Scale, Category.GetCategoryName(), Verbosity, IsPast ? FColor::Red : FColor::Blue, TEXT(""), DebugThickness);
+			CurrentEntry->AddSegment(Pos, Pos + Y * Scale, Category.GetCategoryName(), Verbosity, IsPast ? FColor::Orange : FColor::Turquoise, TEXT(""), DebugThickness);
+
+			if (Index == LastIndex)
+			{
+				break;
+			}
+
+			const FVector NextPos = Samples[Index + 1].Position + OffsetVector;
+			CurrentEntry->AddSegment(Pos, NextPos, Category.GetCategoryName(), Verbosity, FColor::Black, TEXT(""), DebugThickness);
+		}
+	}
+#endif // ENABLE_VISUAL_LOG
+}
+
 #endif // ENABLE_ANIM_DEBUG
 
 FArchive& operator<<(FArchive& Ar, FPoseSearchQueryTrajectorySample& TrajectorySample)
diff --git a/Engine/Plugins/Animation/PoseSearch/Source/Runtime/Public/PoseSearch/PoseSearchTrajectoryTypes.h b/Engine/Plugins/Animation/PoseSearch/Source/Runtime/Public/PoseSearch/PoseSearchTrajectoryTypes.h
--- a/Engine/Plugins/Animation/PoseSearch/Source/Runtime/Public/PoseSearch/PoseSearchTrajectoryTypes.h
+++ b/Engine/Plugins/Animation/PoseSearch/Source/Runtime/Public/PoseSearch/PoseSearchTrajectoryTypes.h
@@ -43,6 +43,7 @@ struct POSESEARCH_API FPoseSearchQueryTrajectory
 #if ENABLE_ANIM_DEBUG
 	void DebugDrawTrajectory(const UWorld* World, const float DebugThickness = 0.f, float HeightOffset = 0.f) const;
 	void DebugDrawTrajectory(FAnimInstanceProxy& AnimInstanceProxy, const float DebugThickness = 0.f, float HeightOffset = 0.f, int MaxHistorySamples = -1, int MaxPredictionSamples = -1) const;
+	void DebugDrawTrajectory(const UObject* Owner, const FLogCategoryBase& Category, ELogVerbosity::Type Verbosity, const float DebugThickness = 0.f, float HeightOffset = 0.f) const;
 #endif // ENABLE_ANIM_DEBUG
 };
 POSESEARCH_API FArchive& operator<<(FArchive& Ar, FPoseSearchQueryTrajectory& Trajectory);
