PoseSearch - adding UPoseSearchFeatureChannel_Position::IsFilterValid depending on squared distance between query and poses

diff --git a/Engine/Plugins/Animation/PoseSearch/Source/Runtime/Private/PoseSearchFeatureChannel_Position.cpp b/Engine/Plugins/Animation/PoseSearch/Source/Runtime/Private/PoseSearchFeatureChannel_Position.cpp
--- a/Engine/Plugins/Animation/PoseSearch/Source/Runtime/Private/PoseSearchFeatureChannel_Position.cpp
+++ b/Engine/Plugins/Animation/PoseSearch/Source/Runtime/Private/PoseSearchFeatureChannel_Position.cpp
@@ -77,6 +77,20 @@ void UPoseSearchFeatureChannel_Position::AddDependentChannels(UPoseSearchSchema*
 	}
 }
 
+bool UPoseSearchFeatureChannel_Position::IsFilterValid(TConstArrayView<float> PoseValues, TConstArrayView<float> QueryValues, int32 PoseIdx, const UE::PoseSearch::FPoseMetadata& Metadata) const
+{
+	using namespace UE::PoseSearch;
+	const FVector Pose = FFeatureVectorHelper::DecodeVector(PoseValues, ChannelDataOffset, ComponentStripping);
+	const FVector Query = FFeatureVectorHelper::DecodeVector(QueryValues, ChannelDataOffset, ComponentStripping);
+
+	const float SquaredLength = (Pose - Query).SquaredLength();
+
+	check(MaxPositionDistanceSquared > 0.f);
+
+	return SquaredLength <= MaxPositionDistanceSquared;
+}
+
+
 void UPoseSearchFeatureChannel_Position::BuildQuery(UE::PoseSearch::FSearchContext& SearchContext) const
 {
 	using namespace UE::PoseSearch;
diff --git a/Engine/Plugins/Animation/PoseSearch/Source/Runtime/Private/PoseSearchFeatureChannel_Position.h b/Engine/Plugins/Animation/PoseSearch/Source/Runtime/Private/PoseSearchFeatureChannel_Position.h
--- a/Engine/Plugins/Animation/PoseSearch/Source/Runtime/Private/PoseSearchFeatureChannel_Position.h
+++ b/Engine/Plugins/Animation/PoseSearch/Source/Runtime/Private/PoseSearchFeatureChannel_Position.h
@@ -67,6 +67,13 @@ public:
 	UPROPERTY(EditAnywhere, Category = "Settings")
 	EPermutationTimeType PermutationTimeType = EPermutationTimeType::UseSampleTime;
 
+	// Experimental, this feature might be removed without warning, not for production use
+	// During pose selection if the squared distance between query versus candidate poses for this 
+	// position channel is greater than MaxPositionDistanceSquared the candidate will be discarded.
+	// the filtering will be enabled only for MaxPositionDistanceSquared > 0
+	UPROPERTY(EditAnywhere, Category = Experimental)
+	float MaxPositionDistanceSquared = 0.f;
+
 #if WITH_EDITORONLY_DATA
 	// if set, all the channels of the same class with the same cardinality, and the same NormalizationGroup, will be normalized together.
 	// for example in a locomotion database of a character holding a weapon, containing non mirrorable animations, you'd still want to normalize together 
@@ -89,6 +96,12 @@ public:
 	virtual EPermutationTimeType GetPermutationTimeType() const override { return PermutationTimeType; }
 	virtual void AddDependentChannels(UPoseSearchSchema* Schema) const override;
 
+	// IPoseSearchFilter interface
+	// Experimental, this feature might be removed without warning, not for production use
+	virtual bool IsFilterActive() const override { return MaxPositionDistanceSquared > 0.f; }
+	// Experimental, this feature might be removed without warning, not for production use
+	virtual bool IsFilterValid(TConstArrayView<float> PoseValues, TConstArrayView<float> QueryValues, int32 PoseIdx, const UE::PoseSearch::FPoseMetadata& Metadata) const override;
+
 #if ENABLE_DRAW_DEBUG
 	virtual void DebugDraw(const UE::PoseSearch::FDebugDrawParams& DrawParams, TConstArrayView<float> PoseVector) const override;
 #endif // ENABLE_DRAW_DEBUG
