PoseSearch - fix for SDebuggerDatabaseView::PopulateViewRows bShowOnlyBestAssetPose

diff --git a/Engine/Plugins/Animation/PoseSearch/Source/Editor/Private/PoseSearchDebuggerDatabaseView.cpp b/Engine/Plugins/Animation/PoseSearch/Source/Editor/Private/PoseSearchDebuggerDatabaseView.cpp
--- a/Engine/Plugins/Animation/PoseSearch/Source/Editor/Private/PoseSearchDebuggerDatabaseView.cpp
+++ b/Engine/Plugins/Animation/PoseSearch/Source/Editor/Private/PoseSearchDebuggerDatabaseView.cpp
@@ -644,6 +644,17 @@ void SDebuggerDatabaseView::PopulateViewRows()
 				const FDebuggerDatabaseRowData& RowDataA = UnfilteredDatabaseRows[IndexA].Get();
 				const FDebuggerDatabaseRowData& RowDataB = UnfilteredDatabaseRows[IndexB].Get();
 
+				if (RowDataA.SharedData->SourceDatabase != RowDataB.SharedData->SourceDatabase)
+				{
+					const int32 DatabasePathCompare = RowDataA.SharedData->DatabasePath.Compare(RowDataB.SharedData->DatabasePath);
+					if (DatabasePathCompare < 0)
+					{
+						return true;
+					}
+					const int32 DatabaseNameCompare = RowDataA.SharedData->DatabaseName.Compare(RowDataB.SharedData->DatabaseName);
+					return DatabaseNameCompare < 0;
+				}
+
 				if (RowDataA.DbAssetIdx == RowDataB.DbAssetIdx)
 				{
 					return RowDataA.PoseCost < RowDataB.PoseCost;
@@ -653,15 +664,17 @@ void SDebuggerDatabaseView::PopulateViewRows()
 
 			// populating BestAssetPoseIndex only with the best pose index (lowest total cost) for each asset
 			int32 PreviousDbAssetIdx = INDEX_NONE;
+			TWeakObjectPtr<const UPoseSearchDatabase> PreviousSourceDatabase;
 			for (int32 i = 0; i < UnfilteredDatabaseRowsNum; ++i)
 			{
 				const FDebuggerDatabaseRowData& RowData = UnfilteredDatabaseRows[SortIndex[i]].Get();
-				if (PreviousDbAssetIdx != RowData.DbAssetIdx)
+				if (PreviousDbAssetIdx != RowData.DbAssetIdx || PreviousSourceDatabase != RowData.SharedData->SourceDatabase)
 				{
 					// adding only the best pose index for the RowData.DbAssetIdx asset to BestAssetPoseIndex 
 					// (the first we find by iterating over UnfilteredDatabaseRows[SortIndex[i]])
 					BestAssetPoseIndex.Add(RowData.PoseIdx);
 					PreviousDbAssetIdx = RowData.DbAssetIdx;
+					PreviousSourceDatabase = RowData.SharedData->SourceDatabase;
 				}
 			}
 		}
