PoseSearch - pose history fixes
AnimNextPoseSearch - FAnimNextHistoryCollectorTask native support for FLODPoseStack

diff --git a/Engine/Plugins/Animation/PoseSearch/Source/Runtime/Private/AnimNode_PoseSearchHistoryCollector.cpp b/Engine/Plugins/Animation/PoseSearch/Source/Runtime/Private/AnimNode_PoseSearchHistoryCollector.cpp
--- a/Engine/Plugins/Animation/PoseSearch/Source/Runtime/Private/AnimNode_PoseSearchHistoryCollector.cpp
+++ b/Engine/Plugins/Animation/PoseSearch/Source/Runtime/Private/AnimNode_PoseSearchHistoryCollector.cpp
@@ -29,7 +29,8 @@ void FAnimNode_PoseSearchHistoryCollector_Base::Initialize_AnyThread(const FAnim
 			FCSPose<FCompactPose> ComponentSpacePose;
 			FBlendedCurve EmptyCurves;
 			ComponentSpacePose.InitPose(&BoneContainer);
-			PoseHistory.EvaluateComponentSpace_AnyThread(0.f, ComponentSpacePose, bStoreScales,
+			UE::PoseSearch::FComponentSpacePoseProvider ComponentSpacePoseProvider(ComponentSpacePose);
+			PoseHistory.EvaluateComponentSpace_AnyThread(0.f, ComponentSpacePoseProvider, bStoreScales,
 				RootBoneRecoveryTime, RootBoneTranslationRecoveryRatio, RootBoneRotationRecoveryRatio, true, true, 
 				GetRequiredBones(Context.AnimInstanceProxy), EmptyCurves, MakeConstArrayView(CollectedCurves));
 		}
@@ -149,10 +150,12 @@ void FAnimNode_PoseSearchHistoryCollector::Evaluate_AnyThread(FPoseContext& Outp
 	TArray<FBoneIndexType> RequiredBones;
 	if (bCacheBones)
 	{
+		// array of skeleton bone indexes (encoded as FBoneIndexType)
 		RequiredBones = GetRequiredBones(Output.AnimInstanceProxy);
 	}
 
-	PoseHistory.EvaluateComponentSpace_AnyThread(Output.AnimInstanceProxy->GetDeltaSeconds(), ComponentSpacePose, bStoreScales,
+	UE::PoseSearch::FComponentSpacePoseProvider ComponentSpacePoseProvider(ComponentSpacePose);
+	PoseHistory.EvaluateComponentSpace_AnyThread(Output.AnimInstanceProxy->GetDeltaSeconds(), ComponentSpacePoseProvider, bStoreScales,
 		RootBoneRecoveryTime, RootBoneTranslationRecoveryRatio, RootBoneRotationRecoveryRatio, bNeedsReset, bCacheBones, 
 		RequiredBones, Output.Curve, MakeConstArrayView(CollectedCurves));
 
@@ -215,11 +218,14 @@ void FAnimNode_PoseSearchComponentSpaceHistoryCollector::EvaluateComponentSpace_
 	TArray<FBoneIndexType> RequiredBones;
 	if (bCacheBones)
 	{
+		// array of skeleton bone indexes (encoded as FBoneIndexType)
 		RequiredBones = GetRequiredBones(Output.AnimInstanceProxy);
 	}
 
-	PoseHistory.EvaluateComponentSpace_AnyThread(Output.AnimInstanceProxy->GetDeltaSeconds(), Output.Pose, bStoreScales, 
-		RootBoneRecoveryTime, RootBoneTranslationRecoveryRatio, RootBoneRotationRecoveryRatio, bNeedsReset, bCacheBones, RequiredBones);
+	UE::PoseSearch::FComponentSpacePoseProvider ComponentSpacePoseProvider(Output.Pose);
+	PoseHistory.EvaluateComponentSpace_AnyThread(Output.AnimInstanceProxy->GetDeltaSeconds(), ComponentSpacePoseProvider, bStoreScales, 
+		RootBoneRecoveryTime, RootBoneTranslationRecoveryRatio, RootBoneRotationRecoveryRatio, bNeedsReset, bCacheBones,
+		RequiredBones, Output.Curve, MakeConstArrayView(CollectedCurves));
 	
 	bCacheBones = false;
 
diff --git a/Engine/Plugins/Animation/PoseSearch/Source/Runtime/Private/PoseSearchHistory.cpp b/Engine/Plugins/Animation/PoseSearch/Source/Runtime/Private/PoseSearchHistory.cpp
--- a/Engine/Plugins/Animation/PoseSearch/Source/Runtime/Private/PoseSearchHistory.cpp
+++ b/Engine/Plugins/Animation/PoseSearch/Source/Runtime/Private/PoseSearchHistory.cpp
@@ -74,7 +74,7 @@ FORCEINLINE auto LowerBound(IteratorType First, IteratorType Last, const ValueTy
 
 static FBoneIndexType GetRemappedBoneIndexType(FBoneIndexType BoneIndexType, const USkeleton* BoneIndexSkeleton, const USkeleton* LastUpdateSkeleton)
 {
-	// remapping BoneIndexType in case the skeleton used to store history (LastUpdateSkeleton) is different from BoneIndexSkeleton
+	// remapping the skeleton bone index (encoded as BoneIndexType) in case the skeleton we used to store the history (LastUpdateSkeleton) is different from BoneIndexSkeleton
 	if (LastUpdateSkeleton != nullptr && LastUpdateSkeleton != BoneIndexSkeleton)
 	{
 		const FSkeletonRemapping& SkeletonRemapping = UE::Anim::FSkeletonRemappingRegistry::Get().GetRemapping(BoneIndexSkeleton, LastUpdateSkeleton);
@@ -276,26 +276,54 @@ static void CopyEntries(const TRingBuffer<FPoseHistoryEntry>& From, TRingBuffer<
 }
 
 //////////////////////////////////////////////////////////////////////////
-// FPoseHistoryEntry
-void FPoseHistoryEntry::Update(float Time, FCSPose<FCompactPose>& ComponentSpacePose, const FBoneToTransformMap& BoneToTransformMap, bool bStoreScales, const FBlendedCurve& Curves, const TConstArrayView<FName>& CollectedCurves)
+// FComponentSpacePoseProvider
+FComponentSpacePoseProvider::FComponentSpacePoseProvider(FCSPose<FCompactPose>& InComponentSpacePose)
+	: ComponentSpacePose(InComponentSpacePose)
 {
-	AccumulatedSeconds = Time;
+	check(GetSkeletonAsset());
+}
 
+FTransform FComponentSpacePoseProvider::CalculateComponentSpaceTransform(const FSkeletonPoseBoneIndex SkeletonBoneIdx)
+{
 	const FBoneContainer& BoneContainer = ComponentSpacePose.GetPose().GetBoneContainer();
-	const USkeleton* SkeletonAsset = BoneContainer.GetSkeletonAsset();
-	check(SkeletonAsset);
-	const FReferenceSkeleton& RefSkeleton = SkeletonAsset->GetReferenceSkeleton();
-	const TArray<FTransform>& RefBonePose = RefSkeleton.GetRefBonePose();
-	const int32 NumSkeletonBones = RefSkeleton.GetNum();
+	const FCompactPoseBoneIndex CompactBoneIdx = BoneContainer.GetCompactPoseIndexFromSkeletonPoseIndex(SkeletonBoneIdx);
+	if (CompactBoneIdx.IsValid())
+	{
+		return ComponentSpacePose.GetComponentSpaceTransform(CompactBoneIdx);
+	}
+	
+	// NoTe: this chunk of code is very unlikely to be called, but in case:
+	// @todo: cache any transform outside the domain of ComponentSpacePose if needed
+	// @todo: use the skeletal mesh reference pose instead of the one from the skeleton if needed
+	const USkeleton* Skeleton = BoneContainer.GetSkeletonAsset();
+	const FReferenceSkeleton& ReferenceSkeleton = Skeleton->GetReferenceSkeleton();
+	const int32 ParentIndex = ReferenceSkeleton.GetParentIndex(SkeletonBoneIdx.GetInt());
+	check(ParentIndex >= 0);
+
+	const TArray<FTransform>& RefBonePose = Skeleton->GetReferenceSkeleton().GetRefBonePose();
+	return RefBonePose[SkeletonBoneIdx.GetInt()] * CalculateComponentSpaceTransform(FSkeletonPoseBoneIndex(ParentIndex));
+}
 
+const USkeleton* FComponentSpacePoseProvider::GetSkeletonAsset() const
+{
+	return ComponentSpacePose.GetPose().GetBoneContainer().GetSkeletonAsset();
+}
+
+//////////////////////////////////////////////////////////////////////////
+// FPoseHistoryEntry
+void FPoseHistoryEntry::Update(float Time, IComponentSpacePoseProvider& ComponentSpacePoseProvider, const FBoneToTransformMap& BoneToTransformMap, bool bStoreScales, const FBlendedCurve& Curves, const TConstArrayView<FName>& CollectedCurves)
+{
+	AccumulatedSeconds = Time;
+
+	const USkeleton* Skeleton = ComponentSpacePoseProvider.GetSkeletonAsset();
+	const int32 NumSkeletonBones = Skeleton->GetReferenceSkeleton().GetNum();
 	if (BoneToTransformMap.IsEmpty())
 	{
 		// no mapping: we add all the transforms
 		SetNum(NumSkeletonBones, bStoreScales);
 		for (FSkeletonPoseBoneIndex SkeletonBoneIdx(0); SkeletonBoneIdx != NumSkeletonBones; ++SkeletonBoneIdx)
 		{
-			const FCompactPoseBoneIndex CompactBoneIdx = BoneContainer.GetCompactPoseIndexFromSkeletonPoseIndex(SkeletonBoneIdx);
-			SetComponentSpaceTransform(SkeletonBoneIdx.GetInt(), (CompactBoneIdx.IsValid() ? ComponentSpacePose.GetComponentSpaceTransform(CompactBoneIdx) : RefBonePose[SkeletonBoneIdx.GetInt()]));
+			SetComponentSpaceTransform(SkeletonBoneIdx.GetInt(), ComponentSpacePoseProvider.CalculateComponentSpaceTransform(SkeletonBoneIdx));
 		}
 	}
 	else
@@ -304,8 +332,7 @@ void FPoseHistoryEntry::Update(float Time, FCSPose<FCompactPose>& ComponentSpace
 		for (const FBoneToTransformPair& BoneToTransformPair : BoneToTransformMap)
 		{
 			const FSkeletonPoseBoneIndex SkeletonBoneIdx(BoneToTransformPair.Key);
-			const FCompactPoseBoneIndex CompactBoneIdx = BoneContainer.GetCompactPoseIndexFromSkeletonPoseIndex(SkeletonBoneIdx);
-			SetComponentSpaceTransform(BoneToTransformPair.Value, (CompactBoneIdx.IsValid() ? ComponentSpacePose.GetComponentSpaceTransform(CompactBoneIdx) : RefBonePose[SkeletonBoneIdx.GetInt()]));
+			SetComponentSpaceTransform(BoneToTransformPair.Value, ComponentSpacePoseProvider.CalculateComponentSpaceTransform(SkeletonBoneIdx));
 		}
 	}
 
@@ -391,15 +418,18 @@ void IPoseHistory::DebugDraw(FAnimInstanceProxy& AnimInstanceProxy, FColor Color
 	if (Color.A > 0 && BoneContainer.IsValid())
 	{
 		const USkeleton* Skeleton = BoneContainer.GetSkeletonAsset();
+		check(Skeleton);
+
 		FTransform OutBoneTransform;
 
 		const FBoneToTransformMap& BoneToTransformMap = GetBoneToTransformMap();
 		if (BoneToTransformMap.IsEmpty())
 		{
-			for (FSkeletonPoseBoneIndex SkeletonBoneIdx(0); SkeletonBoneIdx != BoneContainer.GetNumBones(); ++SkeletonBoneIdx)
+			const FReferenceSkeleton& RefSkeleton = Skeleton->GetReferenceSkeleton();
+			const int32 NumSkeletonBones = RefSkeleton.GetNum();
+			for (FSkeletonPoseBoneIndex SkeletonBoneIdx(0); SkeletonBoneIdx != NumSkeletonBones; ++SkeletonBoneIdx)
 			{
-				const FCompactPoseBoneIndex CompactBoneIdx = BoneContainer.GetCompactPoseIndexFromSkeletonPoseIndex(SkeletonBoneIdx);
-				if (GetTransformAtTime(Time, OutBoneTransform, Skeleton, CompactBoneIdx.GetInt(), WorldSpaceIndexType, bExtrapolate))
+				if (GetTransformAtTime(Time, OutBoneTransform, Skeleton, SkeletonBoneIdx.GetInt(), WorldSpaceIndexType, bExtrapolate))
 				{
 					AnimInstanceProxy.AnimDrawDebugPoint(OutBoneTransform.GetTranslation(), PointSize, Color, false, 0.f, ESceneDepthPriorityGroup::SDPG_Foreground);
 				}
@@ -410,8 +440,7 @@ void IPoseHistory::DebugDraw(FAnimInstanceProxy& AnimInstanceProxy, FColor Color
 			for (const FBoneToTransformPair& BoneToTransformPair : BoneToTransformMap)
 			{
 				const FSkeletonPoseBoneIndex SkeletonBoneIdx(BoneToTransformPair.Key);
-				const FCompactPoseBoneIndex CompactBoneIdx = BoneContainer.GetCompactPoseIndexFromSkeletonPoseIndex(SkeletonBoneIdx);
-				if (GetTransformAtTime(Time, OutBoneTransform, Skeleton, CompactBoneIdx.GetInt(), WorldSpaceIndexType, bExtrapolate))
+				if (GetTransformAtTime(Time, OutBoneTransform, Skeleton, SkeletonBoneIdx.GetInt(), WorldSpaceIndexType, bExtrapolate))
 				{
 					AnimInstanceProxy.AnimDrawDebugPoint(OutBoneTransform.GetTranslation(), PointSize, Color, false, 0.f, ESceneDepthPriorityGroup::SDPG_Foreground);
 				}
@@ -446,6 +475,7 @@ void FArchivedPoseHistory::InitFrom(const IPoseHistory* PoseHistory)
 	}
 }
 
+// in here FBoneIndexType BoneIndexType is a skeleton bone index, and it's used to dereference a BoneToTransformMap (skeleton bone index -> pose history bone index)
 bool FArchivedPoseHistory::GetTransformAtTime(float Time, FTransform& OutBoneTransform, const USkeleton* BoneIndexSkeleton, FBoneIndexType BoneIndexType, FBoneIndexType ReferenceBoneIndexType, bool bExtrapolate) const
 {
 	static_assert(RootBoneIndexType == 0 && ComponentSpaceIndexType == FBoneIndexType(-1) && WorldSpaceIndexType == FBoneIndexType(-2)); // some assumptions
@@ -825,23 +855,22 @@ void FPoseHistory::SetTrajectory(const FPoseSearchQueryTrajectory& InTrajectory,
 	}
 }
 
-void FPoseHistory::EvaluateComponentSpace_AnyThread(float DeltaTime, FCSPose<FCompactPose>& ComponentSpacePose, bool bStoreScales, float RootBoneRecoveryTime,
-	float RootBoneTranslationRecoveryRatio, float RootBoneRotationRecoveryRatio, bool bNeedsReset, bool bCacheBones,
-	const TArray<FBoneIndexType>& RequiredBones, const FBlendedCurve& Curves, const TConstArrayView<FName>& CollectedCurves)
+void FPoseHistory::EvaluateComponentSpace_AnyThread(float DeltaTime, IComponentSpacePoseProvider& ComponentSpacePoseProvider, bool bStoreScales,
+	float RootBoneRecoveryTime, float RootBoneTranslationRecoveryRatio, float RootBoneRotationRecoveryRatio,
+	bool bNeedsReset, bool bCacheBones, const TArray<FBoneIndexType>& RequiredBones,
+	const FBlendedCurve& Curves, const TConstArrayView<FName>& CollectedCurves)
 {
 	CheckThreadSafetyWrite(WritePoseDataThreadSafeCounter);
 	CheckThreadSafetyRead(ReadPoseDataThreadSafeCounter);
 
 	check(MaxNumPoses >= 2);
 
-	const USkeleton* Skeleton = ComponentSpacePose.GetPose().GetBoneContainer().GetSkeletonAsset();
-	check(Skeleton);
-
+	const USkeleton* Skeleton = ComponentSpacePoseProvider.GetSkeletonAsset();
 	const FPoseData& ReadPoseData = GetReadPoseData();
 	FPoseData& WritePoseData = EditWritePoseData();
 
 	WritePoseData.LastUpdateSkeleton = ReadPoseData.LastUpdateSkeleton;
-	
+
 	if (bCacheBones)
 	{
 		WritePoseData.BoneToTransformMap.Reset();
@@ -917,7 +946,7 @@ void FPoseHistory::EvaluateComponentSpace_AnyThread(float DeltaTime, FCSPose<FCo
 
 	// Regardless of the retention policy, we always update the most recent Entry
 	FPoseHistoryEntry& MostRecentEntry = WritePoseData.Entries.Last();
-	MostRecentEntry.Update(0.f, ComponentSpacePose, WritePoseData.BoneToTransformMap, bStoreScales, Curves, WritePoseData.CollectedCurves);
+	MostRecentEntry.Update(0.f, ComponentSpacePoseProvider, WritePoseData.BoneToTransformMap, bStoreScales, Curves, WritePoseData.CollectedCurves);
 
 	if (RootBoneRecoveryTime > 0.f && !Trajectory.Samples.IsEmpty())
 	{
@@ -945,12 +974,12 @@ void FPoseHistory::EvaluateComponentSpace_AnyThread(float DeltaTime, FCSPose<FCo
 
 void FPoseHistory::EvaluateComponentSpace_AnyThread(float DeltaTime, FCSPose<FCompactPose>& ComponentSpacePose, bool bStoreScales,
 	float RootBoneRecoveryTime, float RootBoneTranslationRecoveryRatio, float RootBoneRotationRecoveryRatio,
-	bool bNeedsReset, bool bCacheBones, const TArray<FBoneIndexType>& RequiredBones)
+	bool bNeedsReset, bool bCacheBones, const TArray<FBoneIndexType>& RequiredBones,
+	const FBlendedCurve& Curves, const TConstArrayView<FName>& CollectedCurves)
 {
-	FBlendedCurve Curves;
-	EvaluateComponentSpace_AnyThread(DeltaTime, ComponentSpacePose, bStoreScales,
-		RootBoneRecoveryTime, RootBoneTranslationRecoveryRatio, RootBoneRotationRecoveryRatio, bNeedsReset, bCacheBones, 
-		RequiredBones, Curves, TConstArrayView<FName>());
+	FComponentSpacePoseProvider ComponentSpacePoseProvider(ComponentSpacePose);
+	EvaluateComponentSpace_AnyThread(DeltaTime, ComponentSpacePoseProvider, bStoreScales, RootBoneRecoveryTime, RootBoneTranslationRecoveryRatio,
+		RootBoneRotationRecoveryRatio, bNeedsReset, bCacheBones, RequiredBones, Curves, CollectedCurves);
 }
 
 #if ENABLE_DRAW_DEBUG && ENABLE_ANIM_DEBUG
@@ -1072,19 +1101,19 @@ void FMemStackPoseHistory::AddFutureRootBone(float Time, const FTransform& Futur
 	FutureEntry.AccumulatedSeconds = Time;
 }
 
-void FMemStackPoseHistory::AddFuturePose(float Time, FCSPose<FCompactPose>& ComponentSpacePose)
-{
-	FBlendedCurve Curves;
-	AddFuturePose(Time, ComponentSpacePose, Curves);
-}
-
-void FMemStackPoseHistory::AddFuturePose(float Time, FCSPose<FCompactPose>& ComponentSpacePose, const FBlendedCurve& Curves)
+void FMemStackPoseHistory::AddFuturePose(float Time, IComponentSpacePoseProvider& ComponentSpacePoseProvider, const FBlendedCurve& Curves)
 {
 	// we don't allow to add "past" or "present" poses to FutureEntries
 	check(Time > 0.f);
 	check(PoseHistory);	
 	const int32 LowerBoundIdx = Algo::LowerBound(FutureEntries, Time, [](const FPoseHistoryEntry& Entry, float Value) { return Value > Entry.AccumulatedSeconds; });
-	FutureEntries.InsertDefaulted_GetRef(LowerBoundIdx).Update(Time, ComponentSpacePose, GetBoneToTransformMap(), true, Curves, MakeConstArrayView(GetCollectedCurves()));
+	FutureEntries.InsertDefaulted_GetRef(LowerBoundIdx).Update(Time, ComponentSpacePoseProvider, GetBoneToTransformMap(), true, Curves, MakeConstArrayView(GetCollectedCurves()));
+}
+
+void FMemStackPoseHistory::AddFuturePose(float Time, FCSPose<FCompactPose>& ComponentSpacePose, const FBlendedCurve& Curves)
+{
+	FComponentSpacePoseProvider ComponentSpacePoseProvider(ComponentSpacePose);
+	AddFuturePose(Time, ComponentSpacePoseProvider, Curves);
 }
 
 int32 FMemStackPoseHistory::GetNumEntries() const
diff --git a/Engine/Plugins/Animation/PoseSearch/Source/Runtime/Private/PoseSearchLibrary.cpp b/Engine/Plugins/Animation/PoseSearch/Source/Runtime/Private/PoseSearchLibrary.cpp
--- a/Engine/Plugins/Animation/PoseSearch/Source/Runtime/Private/PoseSearchLibrary.cpp
+++ b/Engine/Plugins/Animation/PoseSearch/Source/Runtime/Private/PoseSearchLibrary.cpp
@@ -984,7 +984,8 @@ UE::PoseSearch::FSearchResult UPoseSearchLibrary::MotionMatch(
 
 					Sampler.ExtractPose(FuturePoseExtractionTime, Pose, Curves);
 					ComponentSpacePose.InitPose(Pose);
-					MemStackPoseHistories[RoleIndex].AddFuturePose(FuturePoseAnimationTime, ComponentSpacePose);
+					FComponentSpacePoseProvider ComponentSpacePoseProvider(ComponentSpacePose);
+					MemStackPoseHistories[RoleIndex].AddFuturePose(FuturePoseAnimationTime, ComponentSpacePoseProvider);
 				}
 
 				InternalPoseHistories[RoleIndex] = MemStackPoseHistories[RoleIndex].GetThisOrPoseHistory();
diff --git a/Engine/Plugins/Animation/PoseSearch/Source/Runtime/Public/PoseSearch/AnimNode_PoseSearchHistoryCollector.h b/Engine/Plugins/Animation/PoseSearch/Source/Runtime/Public/PoseSearch/AnimNode_PoseSearchHistoryCollector.h
--- a/Engine/Plugins/Animation/PoseSearch/Source/Runtime/Public/PoseSearch/AnimNode_PoseSearchHistoryCollector.h
+++ b/Engine/Plugins/Animation/PoseSearch/Source/Runtime/Public/PoseSearch/AnimNode_PoseSearchHistoryCollector.h
@@ -107,6 +107,7 @@ public:
 	void GenerateTrajectory(const UAnimInstance* InAnimInstance);
 
 protected:
+	// returns an array of skeleton bone indexes (encoded as FBoneIndexType)
 	TArray<FBoneIndexType> GetRequiredBones(const FAnimInstanceProxy* AnimInstanceProxy) const;
 
 	UE::PoseSearch::FPoseHistory PoseHistory;
diff --git a/Engine/Plugins/Animation/PoseSearch/Source/Runtime/Public/PoseSearch/PoseSearchHistory.h b/Engine/Plugins/Animation/PoseSearch/Source/Runtime/Public/PoseSearch/PoseSearchHistory.h
--- a/Engine/Plugins/Animation/PoseSearch/Source/Runtime/Public/PoseSearch/PoseSearchHistory.h
+++ b/Engine/Plugins/Animation/PoseSearch/Source/Runtime/Public/PoseSearch/PoseSearchHistory.h
@@ -27,6 +27,24 @@ typedef uint16 FComponentSpaceTransformIndex;
 typedef TPair<FBoneIndexType, FComponentSpaceTransformIndex> FBoneToTransformPair;
 typedef TMap<FBoneIndexType, FComponentSpaceTransformIndex> FBoneToTransformMap;
 
+struct POSESEARCH_API IComponentSpacePoseProvider
+{
+	virtual ~IComponentSpacePoseProvider() {}
+	virtual FTransform CalculateComponentSpaceTransform(const FSkeletonPoseBoneIndex SkeletonBoneIdx) = 0;
+	virtual const USkeleton* GetSkeletonAsset() const = 0;
+};
+
+struct FComponentSpacePoseProvider : public IComponentSpacePoseProvider
+{
+	POSESEARCH_API FComponentSpacePoseProvider(FCSPose<FCompactPose>& InComponentSpacePose);
+
+	virtual FTransform CalculateComponentSpaceTransform(const FSkeletonPoseBoneIndex SkeletonBoneIdx) override;
+	virtual const USkeleton* GetSkeletonAsset() const override;
+
+private:
+	FCSPose<FCompactPose>& ComponentSpacePose;
+};
+
 struct FPoseHistoryEntry
 {
 	// collected bones transforms in component space
@@ -36,7 +54,7 @@ struct FPoseHistoryEntry
 	TArray<float> CurveValues;
 	float AccumulatedSeconds = 0.f;
 
-	void Update(float Time, FCSPose<FCompactPose>& ComponentSpacePose, const FBoneToTransformMap& BoneToTransformMap, bool bStoreScales, const FBlendedCurve& Curves, const TConstArrayView<FName>& CollectedCurves);
+	void Update(float Time, IComponentSpacePoseProvider& ComponentSpacePoseProvider, const FBoneToTransformMap& BoneToTransformMap, bool bStoreScales, const FBlendedCurve& Curves, const TConstArrayView<FName>& CollectedCurves);
 
 	POSESEARCH_API void SetNum(int32 Num, bool bStoreScales);
 	int32 Num() const;
@@ -49,10 +67,9 @@ FArchive& operator<<(FArchive& Ar, FPoseHistoryEntry& Entry);
 
 struct POSESEARCH_API IPoseHistory
 {
-public:
 	virtual ~IPoseHistory() {}
 	
-	// returns the BoneIndexType transform relative to ReferenceBoneIndexType: 
+	// returns transform of the the skeleton bone index BoneIndexType (encoded as FBoneIndexType) relative to the transform of ReferenceBoneIndexType:
 	// if ReferenceBoneIndexType is 0 (RootBoneIndexType), OutBoneTransform is in root bone space
 	// if ReferenceBoneIndexType is FBoneIndexType(-1) (ComponentSpaceIndexType), OutBoneTransform is in component space
 	// if ReferenceBoneIndexType is FBoneIndexType(-2) (WorldSpaceIndexType), OutBoneTransform is in world space
@@ -118,13 +135,17 @@ struct POSESEARCH_API FPoseHistory : public IPoseHistory
 	void PreUpdate();
 
 	void Initialize_AnyThread(int32 InNumPoses, float InSamplingInterval);
-	void EvaluateComponentSpace_AnyThread(float DeltaTime, FCSPose<FCompactPose>& ComponentSpacePose, bool bStoreScales,
+	void EvaluateComponentSpace_AnyThread(float DeltaTime, IComponentSpacePoseProvider& ComponentSpacePoseProvider, bool bStoreScales,
 		float RootBoneRecoveryTime, float RootBoneTranslationRecoveryRatio, float RootBoneRotationRecoveryRatio,
-		bool bNeedsReset, bool bCacheBones, const TArray<FBoneIndexType>& RequiredBones);
+		bool bNeedsReset, bool bCacheBones, const TArray<FBoneIndexType>& RequiredBones, 
+		const FBlendedCurve& Curves = FBlendedCurve(), const TConstArrayView<FName>& CollectedCurves = TConstArrayView<FName>());
+
+	UE_DEPRECATED(5.6, "Use other EvaluateComponentSpace_AnyThread signatures instead")
 	void EvaluateComponentSpace_AnyThread(float DeltaTime, FCSPose<FCompactPose>& ComponentSpacePose, bool bStoreScales,
 		float RootBoneRecoveryTime, float RootBoneTranslationRecoveryRatio, float RootBoneRotationRecoveryRatio,
-		bool bNeedsReset, bool bCacheBones, const TArray<FBoneIndexType>& RequiredBones, const FBlendedCurve& Curves, const TConstArrayView<FName>& CollectedCurves);
-
+		bool bNeedsReset, bool bCacheBones, const TArray<FBoneIndexType>& RequiredBones,
+		const FBlendedCurve& Curves = FBlendedCurve(), const TConstArrayView<FName>& CollectedCurves = TConstArrayView<FName>());
+		
 	// IPoseHistory interface
 	virtual bool GetTransformAtTime(float Time, FTransform& OutBoneTransform, const USkeleton* BoneIndexSkeleton = nullptr, FBoneIndexType BoneIndexType = RootBoneIndexType, FBoneIndexType ReferenceBoneIndexType = ComponentSpaceIndexType, bool bExtrapolate = false) const override;
 	virtual bool GetCurveValueAtTime(float Time, const FName& CurveName, float& outcurvevalue, bool bExtrapolate = false) const override;
@@ -164,7 +185,8 @@ private:
 		// skeleton from the last Update, to keep tracking skeleton changes, and support compatible skeletons
 		TWeakObjectPtr<const USkeleton> LastUpdateSkeleton;
 
-		// map of FBoneIndexType(s) to collect. If Empty all the bones get collected
+		// Map from skeleton bone indexes (encoded as FBoneIndexType) to internal pose history transform index (FComponentSpaceTransformIndex).
+		// If Empty all the bones get collected
 		FBoneToTransformMap BoneToTransformMap;
 		
 		// list of curves that we want to collect into our history.
@@ -218,8 +240,9 @@ struct POSESEARCH_API FMemStackPoseHistory : public IPoseHistory
 	// End of IPoseHistory interface
 
 	void AddFutureRootBone(float Time, const FTransform& FutureRootBoneTransform, bool bStoreScales);
-	void AddFuturePose(float Time, FCSPose<FCompactPose>& ComponentSpacePose);
-	void AddFuturePose(float Time, FCSPose<FCompactPose>& ComponentSpacePose, const FBlendedCurve& Curves);
+	void AddFuturePose(float Time, IComponentSpacePoseProvider& ComponentSpacePoseProvider, const FBlendedCurve& Curves = FBlendedCurve());
+	UE_DEPRECATED(5.6, "Use other AddFuturePose signatures instead")
+	void AddFuturePose(float Time, FCSPose<FCompactPose>& ComponentSpacePose, const FBlendedCurve& Curves = FBlendedCurve());
 
 	const IPoseHistory* GetThisOrPoseHistory() const { return FutureEntries.IsEmpty() ? PoseHistory : this; }
 
diff --git a/Engine/Plugins/Experimental/AnimNextPoseSearch/Source/AnimNextPoseSearch/Private/HistoryCollectorTrait.cpp b/Engine/Plugins/Experimental/AnimNextPoseSearch/Source/AnimNextPoseSearch/Private/HistoryCollectorTrait.cpp
--- a/Engine/Plugins/Experimental/AnimNextPoseSearch/Source/AnimNextPoseSearch/Private/HistoryCollectorTrait.cpp
+++ b/Engine/Plugins/Experimental/AnimNextPoseSearch/Source/AnimNextPoseSearch/Private/HistoryCollectorTrait.cpp
@@ -2,6 +2,7 @@
 
 #include "HistoryCollectorTrait.h"
 #include "Animation/AnimNodeBase.h"
+#include "Animation/Skeleton.h"
 #include "AnimationRuntime.h"
 #include "EvaluationVM/EvaluationVM.h"
 #include "GenerationTools.h"
@@ -57,6 +58,73 @@ namespace UE::AnimNext
 			}
 		}
 	}
+
+	// @todo: revisit this code once a FCSPose<FLODPoseStack> is implemented!
+	struct FComponentSpacePoseProvider : public UE::PoseSearch::IComponentSpacePoseProvider
+	{
+		FComponentSpacePoseProvider(const FLODPoseStack& InPose)
+			: Pose(InPose)
+		{
+			const USkeleton* Skeleton = GetSkeletonInternal();
+			const FReferenceSkeleton& ReferenceSkeleton = Skeleton->GetReferenceSkeleton();
+			const int32 NumBones = ReferenceSkeleton.GetNum();
+			ComponentSpaceFlags.SetNumZeroed(NumBones);
+			ComponentSpaceTransforms.SetNum(NumBones);
+		}
+
+		virtual FTransform CalculateComponentSpaceTransform(const FSkeletonPoseBoneIndex SkeletonBoneIdx) override
+		{
+			FTransform& ComponentSpaceTransform = ComponentSpaceTransforms[SkeletonBoneIdx.GetInt()];
+			uint8& ComponentSpaceFlag = ComponentSpaceFlags[SkeletonBoneIdx.GetInt()];
+
+			if (ComponentSpaceFlag)
+			{
+				return ComponentSpaceTransform;
+			}
+
+			const USkeleton* Skeleton = GetSkeletonInternal();
+			const FReferenceSkeleton& ReferenceSkeleton = Skeleton->GetReferenceSkeleton();
+
+			const TArrayView<const FBoneIndexType> SkeletonBoneIndexToLODBoneIndexMap = Pose.GetSkeletonBoneIndexToLODBoneIndexMap();
+			const FBoneIndexType LODBoneIndex = SkeletonBoneIndexToLODBoneIndexMap[SkeletonBoneIdx.GetInt()];
+			
+			if (LODBoneIndex != FBoneIndexType(INDEX_NONE))
+			{
+				ComponentSpaceTransform = Pose.LocalTransformsView[LODBoneIndex];
+			}
+			else
+			{
+				// @todo: use the skeletal mesh reference pose instead of the skeleton one to account for retargeting
+				ComponentSpaceTransform = ReferenceSkeleton.GetRefBonePose()[SkeletonBoneIdx.GetInt()];
+			}
+			
+			const int32 ParentIndex = ReferenceSkeleton.GetParentIndex(SkeletonBoneIdx.GetInt());
+			if (ParentIndex >= 0)
+			{
+				ComponentSpaceTransform *= CalculateComponentSpaceTransform(FSkeletonPoseBoneIndex(ParentIndex));
+			}
+
+			ComponentSpaceFlag = 1;
+			return ComponentSpaceTransform;
+		}
+
+		virtual const USkeleton* GetSkeletonAsset() const override
+		{
+			return GetSkeletonInternal();
+		}
+
+	private:
+		const USkeleton* GetSkeletonInternal() const
+		{
+			const USkeleton* Skeleton = Pose.RefPose->Skeleton.Get();
+			check(Skeleton);
+			return Skeleton;
+		}
+
+		const FLODPoseStack& Pose;
+		TArray<uint8, TInlineAllocator<256 , TMemStackAllocator<>>> ComponentSpaceFlags;
+		TArray<FTransform, TInlineAllocator<256 , TMemStackAllocator<>>> ComponentSpaceTransforms;
+	};
 }
 
 void FAnimNextHistoryCollectorTask::Execute(UE::AnimNext::FEvaluationVM& VM) const
@@ -67,51 +135,47 @@ void FAnimNextHistoryCollectorTask::Execute(UE::AnimNext::FEvaluationVM& VM) con
 	{
 		if (const TUniquePtr<FKeyframeState>* Keyframe = VM.PeekValue<TUniquePtr<FKeyframeState>>(KEYFRAME_STACK_NAME, 0))
 		{
-			FLODPoseStack& Pose = (*Keyframe)->Pose;
-			const USkeleton* Skeleton = Pose.RefPose->Skeleton.Get();
-			check(Skeleton);
-
 			bool bNeedsReset = false;
-			bool bCacheBones = true;
+			bool bCacheBones = false;
+
+			FMemMark Mark(FMemStack::Get());
+
+			const FLODPoseStack& Pose = (*Keyframe)->Pose;
+			TArray<FBoneIndexType> BoneIndicesWithParents;
 
 			if (SharedData->PoseCount != PoseHistory->GetMaxNumPoses() || SharedData->SamplingInterval != PoseHistory->GetSamplingInterval())
 			{
+				bCacheBones = true;
 				PoseHistory->Initialize_AnyThread(SharedData->PoseCount, SharedData->SamplingInterval);
 			}
 
-			FMemMark Mark(FMemStack::Get());
-			TArray<uint16> BoneIndicesWithParents;
-			BoneIndicesWithParents.Add(UE::PoseSearch::RootBoneIndexType);
-
-			for (const FBoneReference& CollectedBone : SharedData->CollectedBones)
+			if (bCacheBones)
 			{
-				if (CollectedBone.BoneName != NAME_None)
+				const USkeleton* Skeleton = Pose.RefPose->Skeleton.Get();
+				check(Skeleton);
+
+				BoneIndicesWithParents.Add(UE::PoseSearch::RootBoneIndexType);
+
+				for (const FBoneReference& CollectedBone : SharedData->CollectedBones)
 				{
-					FBoneReference TempCollectedBone = CollectedBone;
-					TempCollectedBone.Initialize(Skeleton);
-					if (TempCollectedBone.HasValidSetup())
+					if (CollectedBone.BoneName != NAME_None)
 					{
-						BoneIndicesWithParents.AddUnique(TempCollectedBone.BoneIndex);
+						FBoneReference TempCollectedBone = CollectedBone;
+						TempCollectedBone.Initialize(Skeleton);
+						if (TempCollectedBone.HasValidSetup())
+						{
+							BoneIndicesWithParents.AddUnique(TempCollectedBone.BoneIndex);
+						}
 					}
 				}
-			}
 
-			// Build separate index array with parent indices guaranteed to be present. Sort for EnsureParentsPresent.
-			BoneIndicesWithParents.Sort();
-			FAnimationRuntime::EnsureParentsPresent(BoneIndicesWithParents, Skeleton->GetReferenceSkeleton());
-			
-			FBoneContainer BoneContainer;
-			BoneContainer.InitializeTo(BoneIndicesWithParents, UE::Anim::FCurveFilterSettings(UE::Anim::ECurveFilterMode::DisallowAll), *const_cast<USkeleton*>(Skeleton));
-
-			FCompactPose LocalSpacePose;
-			LocalSpacePose.SetBoneContainer(&BoneContainer);
-
-			RemapPoseInternal(Pose, LocalSpacePose);
-
-			FCSPose<FCompactPose> ComponentSpacePose;
-			ComponentSpacePose.InitPose(LocalSpacePose);
+				// Build separate index array with parent indices guaranteed to be present. Sort for EnsureParentsPresent.
+				BoneIndicesWithParents.Sort();
+				FAnimationRuntime::EnsureParentsPresent(BoneIndicesWithParents, Skeleton->GetReferenceSkeleton());
+			}
 
-			PoseHistory->EvaluateComponentSpace_AnyThread(DeltaTime, ComponentSpacePose, bStoreScales,
+			FComponentSpacePoseProvider ComponentSpacePoseProvider(Pose);
+			PoseHistory->EvaluateComponentSpace_AnyThread(DeltaTime, ComponentSpacePoseProvider, bStoreScales,
 				SharedData->RootBoneRecoveryTime, SharedData->RootBoneTranslationRecoveryRatio, SharedData->RootBoneRotationRecoveryRatio,
 				bNeedsReset, bCacheBones, BoneIndicesWithParents);
 		}
