PoseSearch - MaxHeap KNN result

diff --git a/Engine/Plugins/Animation/PoseSearch/Source/Runtime/Private/KDTree.cpp b/Engine/Plugins/Animation/PoseSearch/Source/Runtime/Private/KDTree.cpp
--- a/Engine/Plugins/Animation/PoseSearch/Source/Runtime/Private/KDTree.cpp
+++ b/Engine/Plugins/Animation/PoseSearch/Source/Runtime/Private/KDTree.cpp
@@ -337,6 +337,21 @@ int32 FKDTree::FindNeighbors(FRadiusResultSet& Result, TConstArrayView<float> Qu
 	return FindNeighborsInternal(KDTreeImplementation, Result, Query);
 }
 
+int32 FKDTree::FindNeighbors(FKNNMaxHeapResultSet& Result, TConstArrayView<float> Query) const
+{
+	return FindNeighborsInternal(KDTreeImplementation, Result, Query);
+}
+
+int32 FKDTree::FindNeighbors(FFilteredKNNMaxHeapResultSet& Result, TConstArrayView<float> Query) const
+{
+	return FindNeighborsInternal(KDTreeImplementation, Result, Query);
+}
+
+int32 FKDTree::FindNeighbors(FRadiusMaxHeapResultSet& Result, TConstArrayView<float> Query) const
+{
+	return FindNeighborsInternal(KDTreeImplementation, Result, Query);
+}
+
 SIZE_T FKDTree::GetAllocatedSize() const
 {
 	SIZE_T AllocatedSize = sizeof(FKDTree);
diff --git a/Engine/Plugins/Animation/PoseSearch/Source/Runtime/Private/PoseSearchDatabase.cpp b/Engine/Plugins/Animation/PoseSearch/Source/Runtime/Private/PoseSearchDatabase.cpp
--- a/Engine/Plugins/Animation/PoseSearch/Source/Runtime/Private/PoseSearchDatabase.cpp
+++ b/Engine/Plugins/Animation/PoseSearch/Source/Runtime/Private/PoseSearchDatabase.cpp
@@ -38,6 +38,8 @@ namespace UE::PoseSearch
 #if WITH_EDITOR && ENABLE_ANIM_DEBUG
 static TAutoConsoleVariable<bool> CVarMotionMatchCompareAgainstBruteForce(TEXT("a.MotionMatch.CompareAgainstBruteForce"), false, TEXT("Compare optimized search against brute force search"));
 static TAutoConsoleVariable<bool> CVarMotionMatchValidateKNNSearch(TEXT("a.MotionMatch.ValidateKNNSearch"), false, TEXT("Validate KNN search"));
+// Experimental, this feature might be removed without warning, not for production use
+static TAutoConsoleVariable<bool> CVarMotionMatchProfileMaxHeapKNNSearch(TEXT("a.MotionMatch.ProfileMaxHeapKNNSearch"), false, TEXT("Profile MaxHeap KNN Search"));
 #endif
 
 struct FSearchFilters
@@ -1357,8 +1359,6 @@ UE::PoseSearch::FSearchResult UPoseSearchDatabase::SearchPCAKDTree(UE::PoseSearc
 		const bool bArePCAValuesPruned = SearchIndex.PCAValuesVectorToPoseIndexes.Num() > 0;
 
 		//stack allocated temporaries
-		TArrayView<int32> ResultIndexes((int32*)FMemory_Alloca((ClampedKDTreeQueryNumNeighbors + 1) * sizeof(int32)), ClampedKDTreeQueryNumNeighbors + 1);
-		TArrayView<float> ResultDistanceSqr((float*)FMemory_Alloca((ClampedKDTreeQueryNumNeighbors + 1) * sizeof(float)), ClampedKDTreeQueryNumNeighbors + 1);
 		TArrayView<float> ProjectedQueryValues((float*)FMemory_Alloca(ClampedNumberOfPrincipalComponents * sizeof(float)), ClampedNumberOfPrincipalComponents);
 	
 		TConstArrayView<float> QueryValues = SearchContext.GetOrBuildQuery(Schema);
@@ -1375,6 +1375,7 @@ UE::PoseSearch::FSearchResult UPoseSearchDatabase::SearchPCAKDTree(UE::PoseSearc
 
 		bool bRunNonSelectableIdxPostKDTree = bArePCAValuesPruned;
 #if WITH_EDITOR && ENABLE_ANIM_DEBUG
+		const bool bProfileMaxHeapKNNSearch = CVarMotionMatchProfileMaxHeapKNNSearch.GetValueOnAnyThread();
 		const bool bValidateKNNSearch = CVarMotionMatchValidateKNNSearch.GetValueOnAnyThread();
 		bRunNonSelectableIdxPostKDTree |= bValidateKNNSearch;
 #endif // WITH_EDITOR && ENABLE_ANIM_DEBUG
@@ -1385,18 +1386,80 @@ UE::PoseSearch::FSearchResult UPoseSearchDatabase::SearchPCAKDTree(UE::PoseSearc
 		check(PCAQueryValues.Num() == ClampedNumberOfPrincipalComponents);
 
 		int32 NumResults = 0;
+
+		TArrayView<FKDTree::FKNNMaxHeapResultSet::FResult> Results((FKDTree::FKNNMaxHeapResultSet::FResult*)FMemory_Alloca(ClampedKDTreeQueryNumNeighbors * sizeof(FKDTree::FKNNMaxHeapResultSet::FResult)), ClampedKDTreeQueryNumNeighbors);
+
+#if WITH_EDITOR && ENABLE_ANIM_DEBUG
+		const double StartTimeMaxHeap = bProfileMaxHeapKNNSearch ? FPlatformTime::Seconds() : 0;
+#endif // WITH_EDITOR && ENABLE_ANIM_DEBUG
+
 		if (bRunNonSelectableIdxPostKDTree || NonSelectableIdx.IsEmpty())
 		{
-			FKDTree::FKNNResultSet ResultSet(ClampedKDTreeQueryNumNeighbors, ResultIndexes, ResultDistanceSqr);
+			FKDTree::FKNNMaxHeapResultSet ResultSet(Results);
 			NumResults = SearchIndex.KDTree.FindNeighbors(ResultSet, PCAQueryValues);
 		}
 		else
 		{
-			FKDTree::FFilteredKNNResultSet ResultSet(ClampedKDTreeQueryNumNeighbors, ResultIndexes, ResultDistanceSqr, NonSelectableIdx);
+			FKDTree::FFilteredKNNMaxHeapResultSet ResultSet(Results, NonSelectableIdx);
 			NumResults = SearchIndex.KDTree.FindNeighbors(ResultSet, PCAQueryValues);
 		}
 
 #if WITH_EDITOR && ENABLE_ANIM_DEBUG
+
+		TArrayView<int32> ResultIndexes((int32*)FMemory_Alloca((ClampedKDTreeQueryNumNeighbors + 1) * sizeof(int32)), ClampedKDTreeQueryNumNeighbors + 1);
+		TArrayView<float> ResultDistanceSqr((float*)FMemory_Alloca((ClampedKDTreeQueryNumNeighbors + 1) * sizeof(float)), ClampedKDTreeQueryNumNeighbors + 1);
+
+		if (bProfileMaxHeapKNNSearch)
+		{
+			// debug code to log profiling comparison between FKNNResultSet/FFilteredKNNResultSet vs FKNNMaxHeapResultSet/FFilteredKNNMaxHeapResultSet
+			const double EndTimeMaxHeap = FPlatformTime::Seconds();
+			const double StartTimeRegular = EndTimeMaxHeap;
+			
+			if (bRunNonSelectableIdxPostKDTree || NonSelectableIdx.IsEmpty())
+			{
+				FKDTree::FKNNResultSet ResultSet(ClampedKDTreeQueryNumNeighbors, ResultIndexes, ResultDistanceSqr);
+				SearchIndex.KDTree.FindNeighbors(ResultSet, PCAQueryValues);
+			}
+			else
+			{
+				FKDTree::FFilteredKNNResultSet ResultSet(ClampedKDTreeQueryNumNeighbors, ResultIndexes, ResultDistanceSqr, NonSelectableIdx);
+				SearchIndex.KDTree.FindNeighbors(ResultSet, PCAQueryValues);
+			}
+			const double EndTimeRegular(FPlatformTime::Seconds());
+
+			const double DeltaTimeRegular = EndTimeRegular - StartTimeRegular;
+			const double DeltaTimeMaxHeap = EndTimeMaxHeap - StartTimeMaxHeap;
+			const int32 Winner = DeltaTimeRegular == DeltaTimeMaxHeap ? 0 : DeltaTimeRegular < DeltaTimeMaxHeap ? -1 : 1;
+			UE_LOG(LogPoseSearch, Log, TEXT("Profiling: Regular(%f), MaxHeap(%f), Winner(%d)"), DeltaTimeRegular, DeltaTimeMaxHeap, Winner);
+		}
+
+		// SortedResultsIndexes contains the sorted indexes by Distance of the Results. This is because FKNNMaxHeapResultSet doesn't return a sorted array like FKNNResultSet does
+		// and we need to feed EvaluatePoseKernel with the position in the search (from the KDTree in PCA space) of the best result position, 
+		// to be able to plot a graph SCostTimelineView::BestPosePosView->CurveData->Points, useful to understand how to tune the KDTreeQueryNumNeighbors property:
+		// if your graph NEVER shows a value higher than 50, than it's safe to set KDTreeQueryNumNeighbors to 50! Since lowering KDTreeQueryNumNeighbors will improve performances!
+		TArrayView<int32> SortedResultsIndexes((int32*)FMemory_Alloca(NumResults * sizeof(int32)), NumResults);
+		for (int32 Index = 0; Index < NumResults; ++Index)
+		{
+			SortedResultsIndexes[Index] = Index;
+		}
+		SortedResultsIndexes.Sort([&Results](const int32 IndexA, const int32 IndexB) { return Results[IndexA].Distance < Results[IndexB].Distance; });
+
+		if (bProfileMaxHeapKNNSearch)
+		{
+			for (int32 ResultIndex = 0; ResultIndex < NumResults; ++ResultIndex)
+			{
+				const int32 ExpectdIndex = ResultIndexes[ResultIndex];
+				const int32 ActualIndex = Results[SortedResultsIndexes[ResultIndex]].Index;
+
+				if (ExpectdIndex != ActualIndex)
+				{
+					const float ExpectdDistanceSqr = ResultDistanceSqr[ResultIndex];
+					const float ActualDistanceSqr = Results[SortedResultsIndexes[ResultIndex]].Distance;
+					UE_LOG(LogPoseSearch, Warning, TEXT("Inconsistent Result at index %d (%d-%d) (%f-%f)"), ResultIndex, ExpectdIndex, ActualIndex, ExpectdDistanceSqr, ActualDistanceSqr);
+				}
+			}
+		}
+
 		if (bValidateKNNSearch)
 		{
 			const int32 NumPCAValuesVectors = SearchIndex.GetNumPCAValuesVectors(ClampedNumberOfPrincipalComponents);
@@ -1418,9 +1481,9 @@ UE::PoseSearch::FSearchResult UPoseSearchDatabase::SearchPCAKDTree(UE::PoseSearc
 
 			for (int32 ResultIndex = 0; ResultIndex < NumResults; ++ResultIndex)
 			{
-				if (PCAValueIndexCost[ResultIndex].Key != ResultIndexes[ResultIndex])
+				if (PCAValueIndexCost[ResultIndex].Key != Results[ResultIndex].Index)
 				{
-					if (!FMath::IsNearlyEqual(PCAValueIndexCost[ResultIndex].Value, ResultDistanceSqr[ResultIndex], UE_KINDA_SMALL_NUMBER))
+					if (!FMath::IsNearlyEqual(PCAValueIndexCost[ResultIndex].Value, Results[ResultIndex].Distance, UE_KINDA_SMALL_NUMBER))
 					{
 						UE_LOG(LogPoseSearch, Error, TEXT("UPoseSearchDatabase::SearchPCAKDTree - KDTree search order is inconsistent with exaustive search in PCA space"));
 					}
@@ -1429,7 +1492,7 @@ UE::PoseSearch::FSearchResult UPoseSearchDatabase::SearchPCAKDTree(UE::PoseSearc
 						UE_LOG(LogPoseSearch, Log, TEXT("UPoseSearchDatabase::SearchPCAKDTree - found two points at the same distance from the query in different order between KDTree and exaustive search"));
 					}
 				}
-				else if (!FMath::IsNearlyEqual(PCAValueIndexCost[ResultIndex].Value, ResultDistanceSqr[ResultIndex], UE_KINDA_SMALL_NUMBER))
+				else if (!FMath::IsNearlyEqual(PCAValueIndexCost[ResultIndex].Value, Results[ResultIndex].Distance, UE_KINDA_SMALL_NUMBER))
 				{
 					UE_LOG(LogPoseSearch, Error, TEXT("UPoseSearchDatabase::SearchPCAKDTree - KDTree search cost is inconsistent with exaustive search in PCA space"));
 				}
@@ -1454,10 +1517,14 @@ UE::PoseSearch::FSearchResult UPoseSearchDatabase::SearchPCAKDTree(UE::PoseSearc
 				int32 NumEvaluatePoseKernelCalls = 0;
 				for (int32 ResultIndex = 0; ResultIndex < NumResults; ++ResultIndex)
 				{
-					const TConstArrayView<int32> PoseIndexes = SearchIndex.PCAValuesVectorToPoseIndexes[ResultIndexes[ResultIndex]];
+					const TConstArrayView<int32> PoseIndexes = SearchIndex.PCAValuesVectorToPoseIndexes[Results[ResultIndex].Index];
 					for (int32 Index = 0; Index < PoseIndexes.Num() && NumEvaluatePoseKernelCalls < MaxNumEvaluatePoseKernelCalls; ++Index, ++NumEvaluatePoseKernelCalls)
 					{
-						EvaluatePoseKernel<false, true>(Result, SearchIndex, QueryValues, TArrayView<float>(), PoseIndexes[Index], SearchFilters, SearchContext, this, true, ResultIndex);
+						EvaluatePoseKernel<false, true>(Result, SearchIndex, QueryValues, TArrayView<float>(), PoseIndexes[Index], SearchFilters, SearchContext, this, true
+#if WITH_EDITOR && ENABLE_ANIM_DEBUG
+							, SortedResultsIndexes[ResultIndex]
+#endif // WITH_EDITOR && ENABLE_ANIM_DEBUG
+						);
 					}
 				}
 			}
@@ -1466,10 +1533,14 @@ UE::PoseSearch::FSearchResult UPoseSearchDatabase::SearchPCAKDTree(UE::PoseSearc
 				int32 NumEvaluatePoseKernelCalls = 0;
 				for (int32 ResultIndex = 0; ResultIndex < NumResults; ++ResultIndex)
 				{
-					const TConstArrayView<int32> PoseIndexes = SearchIndex.PCAValuesVectorToPoseIndexes[ResultIndexes[ResultIndex]];
+					const TConstArrayView<int32> PoseIndexes = SearchIndex.PCAValuesVectorToPoseIndexes[Results[ResultIndex].Index];
 					for (int32 Index = 0; Index < PoseIndexes.Num() && NumEvaluatePoseKernelCalls < MaxNumEvaluatePoseKernelCalls; ++Index, ++NumEvaluatePoseKernelCalls)
 					{
-						EvaluatePoseKernel<false, false>(Result, SearchIndex, QueryValues, TArrayView<float>(), PoseIndexes[Index], SearchFilters, SearchContext, this, true, ResultIndex);
+						EvaluatePoseKernel<false, false>(Result, SearchIndex, QueryValues, TArrayView<float>(), PoseIndexes[Index], SearchFilters, SearchContext, this, true
+#if WITH_EDITOR && ENABLE_ANIM_DEBUG
+							, SortedResultsIndexes[ResultIndex]
+#endif // WITH_EDITOR && ENABLE_ANIM_DEBUG
+						);
 					}
 				}
 			}
@@ -1482,7 +1553,11 @@ UE::PoseSearch::FSearchResult UPoseSearchDatabase::SearchPCAKDTree(UE::PoseSearc
 			check(IsAligned(ReconstructedPoseValuesBuffer.GetData(), alignof(VectorRegister4Float)));
 			for (int32 ResultIndex = 0; ResultIndex < NumResults; ++ResultIndex)
 			{
-				EvaluatePoseKernel<true, false>(Result, SearchIndex, QueryValues, ReconstructedPoseValuesBuffer, ResultIndexes[ResultIndex], SearchFilters, SearchContext, this, true, ResultIndex);
+				EvaluatePoseKernel<true, false>(Result, SearchIndex, QueryValues, ReconstructedPoseValuesBuffer, Results[ResultIndex].Index, SearchFilters, SearchContext, this, true
+#if WITH_EDITOR && ENABLE_ANIM_DEBUG
+					, SortedResultsIndexes[ResultIndex]
+#endif // WITH_EDITOR && ENABLE_ANIM_DEBUG
+				);
 			}
 		}
 		// is the data padded at 16 bytes (and 16 bytes aligned by construction)?
@@ -1490,7 +1565,11 @@ UE::PoseSearch::FSearchResult UPoseSearchDatabase::SearchPCAKDTree(UE::PoseSearc
 		{
 			for (int32 ResultIndex = 0; ResultIndex < NumResults; ++ResultIndex)
 			{
-				EvaluatePoseKernel<false, true>(Result, SearchIndex, QueryValues, TArrayView<float>(), ResultIndexes[ResultIndex], SearchFilters, SearchContext, this, true, ResultIndex);
+				EvaluatePoseKernel<false, true>(Result, SearchIndex, QueryValues, TArrayView<float>(), Results[ResultIndex].Index, SearchFilters, SearchContext, this, true
+#if WITH_EDITOR && ENABLE_ANIM_DEBUG
+					, SortedResultsIndexes[ResultIndex]
+#endif // WITH_EDITOR && ENABLE_ANIM_DEBUG
+				);
 			}
 		}
 		// no reconstruction, but data is not 16 bytes padded
@@ -1498,7 +1577,11 @@ UE::PoseSearch::FSearchResult UPoseSearchDatabase::SearchPCAKDTree(UE::PoseSearc
 		{
 			for (int32 ResultIndex = 0; ResultIndex < NumResults; ++ResultIndex)
 			{
-				EvaluatePoseKernel<false, false>(Result, SearchIndex, QueryValues, TArrayView<float>(), ResultIndexes[ResultIndex], SearchFilters, SearchContext, this, true, ResultIndex);
+				EvaluatePoseKernel<false, false>(Result, SearchIndex, QueryValues, TArrayView<float>(), Results[ResultIndex].Index, SearchFilters, SearchContext, this, true
+#if WITH_EDITOR && ENABLE_ANIM_DEBUG
+					, SortedResultsIndexes[ResultIndex]
+#endif // WITH_EDITOR && ENABLE_ANIM_DEBUG
+						);
 			}
 		}
 	}
@@ -1574,7 +1657,11 @@ UE::PoseSearch::FSearchResult UPoseSearchDatabase::SearchVPTree(UE::PoseSearch::
 				const TConstArrayView<int32> PoseIndexes = SearchIndex.ValuesVectorToPoseIndexes[IndexDistance.Index];
 				for (int32 Index = 0; Index < PoseIndexes.Num() && NumEvaluatePoseKernelCalls < MaxNumEvaluatePoseKernelCalls; ++Index, ++NumEvaluatePoseKernelCalls)
 				{
-					EvaluatePoseKernel<false, false>(Result, SearchIndex, QueryValues, TArrayView<float>(), PoseIndexes[Index], SearchFilters, SearchContext, this, true, ResultIndex);
+					EvaluatePoseKernel<false, false>(Result, SearchIndex, QueryValues, TArrayView<float>(), PoseIndexes[Index], SearchFilters, SearchContext, this, true
+#if WITH_EDITOR && ENABLE_ANIM_DEBUG
+						, ResultIndex
+#endif // WITH_EDITOR && ENABLE_ANIM_DEBUG
+						);
 				}
 			}
 		}
@@ -1585,7 +1672,11 @@ UE::PoseSearch::FSearchResult UPoseSearchDatabase::SearchVPTree(UE::PoseSearch::
 				const FIndexDistance& IndexDistance = UnsortedResults[ResultIndex];
 
 				// @todo: IndexDistance.Distance is the Sqrt(DissimilarityCost), so there's no need to calculate it again in SearchIndex.ComparePoses
-				EvaluatePoseKernel<false, false>(Result, SearchIndex, QueryValues, TArrayView<float>(), IndexDistance.Index, SearchFilters, SearchContext, this, true, ResultIndex);
+				EvaluatePoseKernel<false, false>(Result, SearchIndex, QueryValues, TArrayView<float>(), IndexDistance.Index, SearchFilters, SearchContext, this, true
+#if WITH_EDITOR && ENABLE_ANIM_DEBUG
+					, ResultIndex
+#endif // WITH_EDITOR && ENABLE_ANIM_DEBUG
+				);
 			}
 		}
 	}
@@ -1650,7 +1741,11 @@ UE::PoseSearch::FSearchResult UPoseSearchDatabase::SearchBruteForce(UE::PoseSear
 				check(IsAligned(ReconstructedPoseValuesBuffer.GetData(), alignof(VectorRegister4Float)));
 				for (int32 PoseIdx = 0; PoseIdx < SearchIndex.GetNumPoses(); ++PoseIdx)
 				{
-					EvaluatePoseKernel<true, false>(Result, SearchIndex, QueryValues, ReconstructedPoseValuesBuffer, PoseIdx, SearchFilters, SearchContext, this, bUpdateBestCandidates, PoseIdx);
+					EvaluatePoseKernel<true, false>(Result, SearchIndex, QueryValues, ReconstructedPoseValuesBuffer, PoseIdx, SearchFilters, SearchContext, this, bUpdateBestCandidates
+#if WITH_EDITOR && ENABLE_ANIM_DEBUG
+						, PoseIdx
+#endif // WITH_EDITOR && ENABLE_ANIM_DEBUG
+					);
 				}
 			}
 			// is the data padded at 16 bytes (and 16 bytes aligned by construction)?
@@ -1658,7 +1753,11 @@ UE::PoseSearch::FSearchResult UPoseSearchDatabase::SearchBruteForce(UE::PoseSear
 			{
 				for (int32 PoseIdx = 0; PoseIdx < SearchIndex.GetNumPoses(); ++PoseIdx)
 				{
-					EvaluatePoseKernel<false, true>(Result, SearchIndex, QueryValues, TArrayView<float>(), PoseIdx, SearchFilters, SearchContext, this, bUpdateBestCandidates, PoseIdx);
+					EvaluatePoseKernel<false, true>(Result, SearchIndex, QueryValues, TArrayView<float>(), PoseIdx, SearchFilters, SearchContext, this, bUpdateBestCandidates
+#if WITH_EDITOR && ENABLE_ANIM_DEBUG
+						, PoseIdx
+#endif // WITH_EDITOR && ENABLE_ANIM_DEBUG
+						);
 				}
 			}
 			// no reconstruction, but data is not 16 bytes padded
@@ -1666,7 +1765,11 @@ UE::PoseSearch::FSearchResult UPoseSearchDatabase::SearchBruteForce(UE::PoseSear
 			{
 				for (int32 PoseIdx = 0; PoseIdx < SearchIndex.GetNumPoses(); ++PoseIdx)
 				{
-					EvaluatePoseKernel<false, false>(Result, SearchIndex, QueryValues, TArrayView<float>(), PoseIdx, SearchFilters, SearchContext, this, bUpdateBestCandidates, PoseIdx);
+					EvaluatePoseKernel<false, false>(Result, SearchIndex, QueryValues, TArrayView<float>(), PoseIdx, SearchFilters, SearchContext, this, bUpdateBestCandidates
+#if WITH_EDITOR && ENABLE_ANIM_DEBUG
+						, PoseIdx
+#endif // WITH_EDITOR && ENABLE_ANIM_DEBUG
+						);
 				}
 			}
 		}
@@ -1688,7 +1791,11 @@ UE::PoseSearch::FSearchResult UPoseSearchDatabase::SearchBruteForce(UE::PoseSear
 					const int32 LastPoseIdx = FirstPoseIdx + SearchIndexAsset.GetNumPoses();
 					for (int32 PoseIdx = FirstPoseIdx; PoseIdx < LastPoseIdx; ++PoseIdx)
 					{
-						EvaluatePoseKernel<true, false>(Result, SearchIndex, QueryValues, ReconstructedPoseValuesBuffer, PoseIdx, SearchFilters, SearchContext, this, bUpdateBestCandidates, ++ResultIndex);
+						EvaluatePoseKernel<true, false>(Result, SearchIndex, QueryValues, ReconstructedPoseValuesBuffer, PoseIdx, SearchFilters, SearchContext, this, bUpdateBestCandidates
+#if WITH_EDITOR && ENABLE_ANIM_DEBUG
+							, ++ResultIndex
+#endif // WITH_EDITOR && ENABLE_ANIM_DEBUG
+						);
 					}
 				}
 			}
@@ -1702,7 +1809,11 @@ UE::PoseSearch::FSearchResult UPoseSearchDatabase::SearchBruteForce(UE::PoseSear
 					const int32 LastPoseIdx = FirstPoseIdx + SearchIndexAsset.GetNumPoses();
 					for (int32 PoseIdx = FirstPoseIdx; PoseIdx < LastPoseIdx; ++PoseIdx)
 					{
-						EvaluatePoseKernel<false, true>(Result, SearchIndex, QueryValues, TArrayView<float>(), PoseIdx, SearchFilters, SearchContext, this, bUpdateBestCandidates, ++ResultIndex);
+						EvaluatePoseKernel<false, true>(Result, SearchIndex, QueryValues, TArrayView<float>(), PoseIdx, SearchFilters, SearchContext, this, bUpdateBestCandidates
+#if WITH_EDITOR && ENABLE_ANIM_DEBUG
+							, ++ResultIndex
+#endif // WITH_EDITOR && ENABLE_ANIM_DEBUG
+						);
 					}
 				}
 			}
@@ -1716,7 +1827,11 @@ UE::PoseSearch::FSearchResult UPoseSearchDatabase::SearchBruteForce(UE::PoseSear
 					const int32 LastPoseIdx = FirstPoseIdx + SearchIndexAsset.GetNumPoses();
 					for (int32 PoseIdx = FirstPoseIdx; PoseIdx < LastPoseIdx; ++PoseIdx)
 					{
-						EvaluatePoseKernel<false, false>(Result, SearchIndex, QueryValues, TArrayView<float>(), PoseIdx, SearchFilters, SearchContext, this, bUpdateBestCandidates, ++ResultIndex);
+						EvaluatePoseKernel<false, false>(Result, SearchIndex, QueryValues, TArrayView<float>(), PoseIdx, SearchFilters, SearchContext, this, bUpdateBestCandidates
+#if WITH_EDITOR && ENABLE_ANIM_DEBUG
+							, ++ResultIndex
+#endif // WITH_EDITOR && ENABLE_ANIM_DEBUG
+						);
 					}
 				}
 			}
diff --git a/Engine/Plugins/Animation/PoseSearch/Source/Runtime/Private/PoseSearchDerivedData.cpp b/Engine/Plugins/Animation/PoseSearch/Source/Runtime/Private/PoseSearchDerivedData.cpp
--- a/Engine/Plugins/Animation/PoseSearch/Source/Runtime/Private/PoseSearchDerivedData.cpp
+++ b/Engine/Plugins/Animation/PoseSearch/Source/Runtime/Private/PoseSearchDerivedData.cpp
@@ -758,26 +758,23 @@ static void PreprocessSearchIndexKDTree(FSearchIndex& SearchIndex, const UPoseSe
 			// testing the KDTree is returning the proper searches for all the points in pca space
 			const int32 KDTreeQueryNumNeighbors = Database->KDTreeQueryNumNeighbors;
 
-			TArray<int32> ResultIndexes;
-			TArray<float> ResultDistanceSqr;
-			ResultIndexes.SetNum(KDTreeQueryNumNeighbors + 1);
-			ResultDistanceSqr.SetNum(KDTreeQueryNumNeighbors + 1);
-
+			TArray<FKDTree::FKNNMaxHeapResultSet::FResult, TInlineAllocator<256>> Results;
+			Results.SetNumUninitialized(KDTreeQueryNumNeighbors);
 			int32 MaxNumNeighborToFindAPoint = 0;
 			for (int32 PointIndex = 0; PointIndex < NumPCAValuesVectors; ++PointIndex)
 			{
 				// searching the kdtree for PointIndex
-				FKDTree::FRadiusResultSet ResultSet(UE_SMALL_NUMBER, KDTreeQueryNumNeighbors, ResultIndexes, ResultDistanceSqr);
-				SearchIndex.KDTree.FindNeighbors(ResultSet, MakeArrayView(&SearchIndex.PCAValues[PointIndex * NumberOfPrincipalComponents], NumberOfPrincipalComponents));
+				FKDTree::FRadiusMaxHeapResultSet ResultSet(Results, UE_SMALL_NUMBER);
+				const int32 NumResults = SearchIndex.KDTree.FindNeighbors(ResultSet, MakeArrayView(&SearchIndex.PCAValues[PointIndex * NumberOfPrincipalComponents], NumberOfPrincipalComponents));
 
 				bool bFound = false;
-				for (int32 ResultIndex = 0; ResultIndex < ResultSet.Num(); ++ResultIndex)
+				for (int32 ResultIndex = 0; ResultIndex < NumResults; ++ResultIndex)
 				{
-					if (PointIndex == ResultIndexes[ResultIndex])
+					if (PointIndex == Results[ResultIndex].Index)
 					{
 						// PointIndex is the ResultIndex-th candidates out of the kdtree. if ResultIndex-th is greater than KDTreeQueryNumNeighbors,
 						// we wouldn't have found it in a runtime search, so we log the errot (later on only once, with the worst case scenario)
-						check(ResultDistanceSqr[ResultIndex] < UE_KINDA_SMALL_NUMBER);
+						check(Results[ResultIndex].Distance < UE_KINDA_SMALL_NUMBER);
 						MaxNumNeighborToFindAPoint = FMath::Max(MaxNumNeighborToFindAPoint, ResultIndex);
 						bFound = true;
 						break;
@@ -803,22 +800,22 @@ static void PreprocessSearchIndexKDTree(FSearchIndex& SearchIndex, const UPoseSe
 				TArrayView<float> ProjectedValues((float*)FMemory_Alloca(NumberOfPrincipalComponents * sizeof(float)), NumberOfPrincipalComponents);
 				for (int32 PointIndex = 0; PointIndex < NumPCAValuesVectors; ++PointIndex)
 				{
-					FKDTree::FKNNResultSet ResultSet(KDTreeQueryNumNeighbors, ResultIndexes, ResultDistanceSqr);
-					SearchIndex.KDTree.FindNeighbors(ResultSet, SearchIndex.PCAProject(SearchIndex.GetPoseValuesBase(PointIndex, NumDimensions), ProjectedValues));
+					FKDTree::FKNNMaxHeapResultSet ResultSet(Results);
+					const int32 NumResults = SearchIndex.KDTree.FindNeighbors(ResultSet, SearchIndex.PCAProject(SearchIndex.GetPoseValuesBase(PointIndex, NumDimensions), ProjectedValues));
 
 					int32 ResultIndex = 0;
-					for (; ResultIndex < ResultSet.Num(); ++ResultIndex)
+					for (; ResultIndex < NumResults; ++ResultIndex)
 					{
-						if (PointIndex == ResultIndexes[ResultIndex])
+						if (PointIndex == Results[ResultIndex].Index)
 						{
-							if (ResultDistanceSqr[ResultIndex] > UE_KINDA_SMALL_NUMBER)
+							if (Results[ResultIndex].Distance > UE_KINDA_SMALL_NUMBER)
 							{
 								UE_LOG(LogPoseSearch, Error, TEXT("PreprocessSearchIndexKDTree - kdtree for %s is not properly constructed! Couldn't find the Point %d in it within UE_KINDA_SMALL_NUMBER tolerance, after PCA projection"), *Database->GetName(), PointIndex);
 							}
 							break;
 						}
 					}
-					if (ResultIndex == ResultSet.Num())
+					if (ResultIndex == NumResults)
 					{
 						UE_LOG(LogPoseSearch, Error, TEXT("PreprocessSearchIndexKDTree - kdtree for %s is not properly constructed! Couldn't find the Point %d in it, after PCA projection"), *Database->GetName(), PointIndex);
 					}
diff --git a/Engine/Plugins/Animation/PoseSearch/Source/Runtime/Private/PoseSearchDerivedDataKey.cpp b/Engine/Plugins/Animation/PoseSearch/Source/Runtime/Private/PoseSearchDerivedDataKey.cpp
--- a/Engine/Plugins/Animation/PoseSearch/Source/Runtime/Private/PoseSearchDerivedDataKey.cpp
+++ b/Engine/Plugins/Animation/PoseSearch/Source/Runtime/Private/PoseSearchDerivedDataKey.cpp
@@ -57,7 +57,7 @@ FKeyBuilder::FKeyBuilder(const UObject* Object, bool bUseDataVer, bool bPerformC
 	if (bUseDataVer)
 	{
 		// used to invalidate the key without having to change POSESEARCHDB_DERIVEDDATA_VER all the times
-		int32 POSESEARCHDB_DERIVEDDATA_VER_SMALL = 2656;
+		int32 POSESEARCHDB_DERIVEDDATA_VER_SMALL = 2657;
 		FGuid VersionGuid = FDevSystemGuids::GetSystemGuid(FDevSystemGuids::Get().POSESEARCHDB_DERIVEDDATA_VER);
 		FString AnimationCompressionVersionString = UE::Anim::Compression::AnimationCompressionVersionString;
 
diff --git a/Engine/Plugins/Animation/PoseSearch/Source/Runtime/Private/PoseSearchIndex.cpp b/Engine/Plugins/Animation/PoseSearch/Source/Runtime/Private/PoseSearchIndex.cpp
--- a/Engine/Plugins/Animation/PoseSearch/Source/Runtime/Private/PoseSearchIndex.cpp
+++ b/Engine/Plugins/Animation/PoseSearch/Source/Runtime/Private/PoseSearchIndex.cpp
@@ -105,28 +105,25 @@ static bool CalculateSimilarities(TArray<FPosePairSimilarity>& PosePairSimilarit
 	check(Values.Num() == NumPoses * DataCardinality);
 	FKDTree KDTree(NumPoses, DataCardinality, Values.GetData());
 
-	TArray<int32> ResultIndexes;
-	TArray<float> ResultDistanceSqr;
-	ResultIndexes.SetNum(NumPoses + 1);
-	ResultDistanceSqr.SetNum(NumPoses + 1);
+	TArray<FKDTree::FKNNMaxHeapResultSet::FResult, TInlineAllocator<256>> Results;
+	Results.SetNumUninitialized(NumPoses);
 
 	for (int32 PoseIdx = 0; PoseIdx < NumPoses; ++PoseIdx)
 	{
 		TConstArrayView<float> ValuesA = GetValuesVector(PoseIdx, DataCardinality);
 
 		// searching for duplicates within a radius of SimilarityThreshold
-		FKDTree::FRadiusResultSet ResultSet(SimilarityThreshold, NumPoses, ResultIndexes, ResultDistanceSqr);
-		KDTree.FindNeighbors(ResultSet, ValuesA);
+		FKDTree::FRadiusMaxHeapResultSet ResultSet(Results, SimilarityThreshold);
+		const int32 NumResults = KDTree.FindNeighbors(ResultSet, ValuesA);
 
-		for (int32 ResultIndex = 0; ResultIndex < ResultSet.Num(); ++ResultIndex)
+		for (int32 ResultIndex = 0; ResultIndex < NumResults; ++ResultIndex)
 		{
-			const int32 ResultPoseIdx = ResultIndexes[ResultIndex];
-			if (PoseIdx != ResultPoseIdx)
+			if (PoseIdx != Results[ResultIndex].Index)
 			{
 				FPosePairSimilarity PosePair;
 				PosePair.PoseIdxA = PoseIdx;
-				PosePair.PoseIdxB = ResultPoseIdx;
-				PosePair.Similarity = ResultDistanceSqr[ResultIndex];
+				PosePair.PoseIdxB = Results[ResultIndex].Index;
+				PosePair.Similarity = Results[ResultIndex].Distance;
 				PosePairSimilarities.Emplace(PosePair);
 			}
 		}
diff --git a/Engine/Plugins/Animation/PoseSearch/Source/Runtime/Public/PoseSearch/KDTree.h b/Engine/Plugins/Animation/PoseSearch/Source/Runtime/Public/PoseSearch/KDTree.h
--- a/Engine/Plugins/Animation/PoseSearch/Source/Runtime/Public/PoseSearch/KDTree.h
+++ b/Engine/Plugins/Animation/PoseSearch/Source/Runtime/Public/PoseSearch/KDTree.h
@@ -178,6 +178,122 @@ struct FKDTree
 			return true;
 		}
 	};
+
+	// Experimental, this feature might be removed without warning, not for production use
+	struct FKNNMaxHeapResultSet
+	{
+		struct FResult
+		{
+			float Distance;
+			AccessorType Index;
+
+			FORCEINLINE bool operator()(const FResult& A, const FResult& B) const
+			{
+				// making it a max heap
+				return A.Distance > B.Distance;
+			}
+		};
+
+		inline FKNNMaxHeapResultSet(TArrayView<FResult> Results, float InitWorstDist = UE_BIG_NUMBER)
+		: MaxHeap(Results)
+		, MaxHeapNum(0)
+		, WorstDist(InitWorstDist)
+		{
+		}
+
+		inline int32 Num() const
+		{
+			return MaxHeapNum;
+		}
+	
+		inline bool full() const
+		{
+			return MaxHeap.Num() == MaxHeapNum;
+		}
+
+		inline bool addPoint(float dist, AccessorType index)
+		{
+			if (full())
+			{
+				if (dist < MaxHeap[0].Distance)
+				{
+					MaxHeap[0].Distance = dist;
+					MaxHeap[0].Index = index;
+
+					AlgoImpl::HeapSiftDown(MaxHeap.GetData(), 0, MaxHeapNum, FIdentityFunctor(), FResult());
+					WorstDist = MaxHeap[0].Distance;
+				}
+			}
+			else
+			{
+				MaxHeap[MaxHeapNum] = { dist, index };
+				// calling HeapSiftUp every addPoint seems to be faster than calling HeapifyInternal the first time MaxHeap is full
+				// @todo: do more profiling! (comment out AlgoImpl::HeapSiftUp and uncomment below AlgoImpl::HeapifyInternal)
+				AlgoImpl::HeapSiftUp(MaxHeap.GetData(), 0, MaxHeapNum, FIdentityFunctor(), FResult());
+				++MaxHeapNum;
+
+				if (full())
+				{
+					//AlgoImpl::HeapifyInternal(MaxHeap.GetData(), MaxHeapNum, FIdentityFunctor(), FResult());
+					WorstDist = MaxHeap[0].Distance;
+				}
+			}
+
+			// tell caller that the search shall continue
+			return true;
+		}
+
+		inline float worstDist() const
+		{
+			return WorstDist;
+		}
+
+	protected:
+		TArrayView<FResult> MaxHeap;
+		int32 MaxHeapNum;
+		float WorstDist;
+	};
+
+	// Experimental, this feature might be removed without warning, not for production use
+	struct FFilteredKNNMaxHeapResultSet : public FKNNMaxHeapResultSet
+	{
+		inline FFilteredKNNMaxHeapResultSet(TArrayView<FResult> Results, TConstArrayView<AccessorType> InExcludeFromSearchIndexes)
+		: FKNNMaxHeapResultSet(Results)
+		, ExcludeFromSearchIndexes(InExcludeFromSearchIndexes)
+		{
+			check(Algo::IsSorted(InExcludeFromSearchIndexes));
+		}
+
+		inline bool addPoint(float dist, AccessorType index)
+		{
+			if (Algo::BinarySearch(ExcludeFromSearchIndexes, index) == INDEX_NONE)
+			{
+				FKNNMaxHeapResultSet::addPoint(dist, index);
+			}
+			return true;
+		}
+
+	private:
+		TConstArrayView<AccessorType> ExcludeFromSearchIndexes; // sorted array view
+	};
+
+	// Experimental, this feature might be removed without warning, not for production use
+	struct FRadiusMaxHeapResultSet : public FKNNMaxHeapResultSet
+	{
+		inline FRadiusMaxHeapResultSet(TArrayView<FResult> Results, float Radius)
+		: FKNNMaxHeapResultSet(Results, Radius)
+		{
+		}
+
+		inline bool addPoint(float dist, AccessorType index)
+		{
+			if (dist < worstDist())
+			{
+				FKNNMaxHeapResultSet::addPoint(dist, index);
+			}
+			return true;
+		}
+	};
 	
 
 	FKDTree(AccessorType Count, AccessorType Dim, const float* Data, AccessorType MaxLeafSize = 16);
@@ -197,6 +313,13 @@ struct FKDTree
 	int32 FindNeighbors(FKNNResultSet& Result, TConstArrayView<float> Query) const;
 	int32 FindNeighbors(FFilteredKNNResultSet& Result, TConstArrayView<float> Query) const;
 	int32 FindNeighbors(FRadiusResultSet& Result, TConstArrayView<float> Query) const;
+	
+	// Experimental, this feature might be removed without warning, not for production use
+	int32 FindNeighbors(FKNNMaxHeapResultSet& Result, TConstArrayView<float> Query) const;
+	// Experimental, this feature might be removed without warning, not for production use
+	int32 FindNeighbors(FFilteredKNNMaxHeapResultSet& Result, TConstArrayView<float> Query) const;
+	// Experimental, this feature might be removed without warning, not for production use
+	int32 FindNeighbors(FRadiusMaxHeapResultSet& Result, TConstArrayView<float> Query) const;
 
 	POSESEARCH_API SIZE_T GetAllocatedSize() const;
 
