Marking PoseSearchTrajectoryPredictor as experimental.

#rnx

diff --git a/Engine/Plugins/Animation/PoseSearch/Source/Runtime/Public/PoseSearch/PoseSearchTrajectoryPredictor.h b/Engine/Plugins/Animation/PoseSearch/Source/Runtime/Public/PoseSearch/PoseSearchTrajectoryPredictor.h
--- a/Engine/Plugins/Animation/PoseSearch/Source/Runtime/Public/PoseSearch/PoseSearchTrajectoryPredictor.h
+++ b/Engine/Plugins/Animation/PoseSearch/Source/Runtime/Public/PoseSearch/PoseSearchTrajectoryPredictor.h
@@ -8,7 +8,7 @@
 struct FPoseSearchQueryTrajectory;
 
 
-UINTERFACE(BlueprintType)
+UINTERFACE(BlueprintType, Experimental)
 class POSESEARCH_API UPoseSearchTrajectoryPredictorInterface : public UInterface
 {
 	GENERATED_BODY()
