PoseSearch - CalculateDynamicWeightsSqrt useful for schema A/B testing

diff --git a/Engine/Plugins/Animation/PoseSearch/Source/Editor/Private/PoseSearchDebuggerDatabaseView.cpp b/Engine/Plugins/Animation/PoseSearch/Source/Editor/Private/PoseSearchDebuggerDatabaseView.cpp
--- a/Engine/Plugins/Animation/PoseSearch/Source/Editor/Private/PoseSearchDebuggerDatabaseView.cpp
+++ b/Engine/Plugins/Animation/PoseSearch/Source/Editor/Private/PoseSearchDebuggerDatabaseView.cpp
@@ -155,7 +155,7 @@ private:
 	TArray<FLabelToChannels> LabelToChannels;
 };
 
-static void AddUnfilteredDatabaseRow(const UPoseSearchDatabase* Database, 
+static void AddUnfilteredDatabaseRow(const UPoseSearchDatabase* Database, TConstArrayView<float> DynamicWeightsSqrt,
 	TArray<TSharedRef<FDebuggerDatabaseRowData>>& UnfilteredDatabaseRows, TSharedRef<FDebuggerDatabaseSharedData> SharedData,
 	int32 DbPoseIdx, EPoseCandidateFlags PoseCandidateFlags, TConstArrayView<uint32> PoseToPCAValuesVectorIndexes, const FPoseSearchCost& Cost = FPoseSearchCost())
 {
@@ -181,7 +181,7 @@ static void AddUnfilteredDatabaseRow(const UPoseSearchDatabase* Database,
 		const bool bIsQueryVectorValid = SharedData->QueryVector.Num() == PoseValues.Num();
 		const TConstArrayView<float> QueryVector = bIsQueryVectorValid ? TConstArrayView<float>(SharedData->QueryVector) : TConstArrayView<float>(PoseValues);
 
-		CompareFeatureVectors(PoseValues, QueryVector, SearchIndex.WeightsSqrt, Row->CostVector);
+		CompareFeatureVectors(PoseValues, QueryVector, DynamicWeightsSqrt, Row->CostVector);
 
 		if (Cost.IsValid())
 		{
@@ -197,11 +197,11 @@ static void AddUnfilteredDatabaseRow(const UPoseSearchDatabase* Database,
 		// is PoseValues padded at 16 bytes (and 16 bytes aligned by construction)?
 		else if (PoseValues.Num() % 4 == 0)
 		{
-			Row->PoseCost = SearchIndex.CompareAlignedPoses(DbPoseIdx, 0.f, PoseValues, QueryVector);
+			Row->PoseCost = SearchIndex.CompareAlignedPoses(DbPoseIdx, 0.f, PoseValues, QueryVector, DynamicWeightsSqrt);
 		}
 		else
 		{
-			Row->PoseCost = SearchIndex.ComparePoses(DbPoseIdx, 0.f, PoseValues, QueryVector);
+			Row->PoseCost = SearchIndex.ComparePoses(DbPoseIdx, 0.f, PoseValues, QueryVector, DynamicWeightsSqrt);
 		}
 
 		if (!SharedData->PCAQueryVector.IsEmpty())
@@ -254,6 +254,7 @@ void SDebuggerDatabaseView::Update(const FTraceMotionMatchingStateMessage& State
 
 	bool bAddPCACost = false;
 	TArray<uint32> PoseToPCAValuesVectorIndexes;
+	TAlignedArray<float> DynamicWeightsSqrtBuffer;
 	for (const FTraceMotionMatchingStateDatabaseEntry& DbEntry : State.DatabaseEntries)
 	{
 		const UPoseSearchDatabase* Database = FTraceMotionMatchingStateMessage::GetObjectFromId<UPoseSearchDatabase>(DbEntry.DatabaseId);
@@ -273,13 +274,16 @@ void SDebuggerDatabaseView::Update(const FTraceMotionMatchingStateMessage& State
 			if (Database->PoseSearchMode == EPoseSearchMode::PCAKDTree && DbEntry.QueryVector.Num() == Database->Schema->SchemaCardinality)
 			{
 				SharedData->PCAQueryVector.SetNumZeroed(Database->GetNumberOfPrincipalComponents());
-				Database->GetSearchIndex().PCAProject(DbEntry.QueryVector, SharedData->PCAQueryVector);
+				SearchIndex.PCAProject(DbEntry.QueryVector, SharedData->PCAQueryVector);
 				bAddPCACost = true;
 			}
 
+			DynamicWeightsSqrtBuffer.SetNumUninitialized(SearchIndex.WeightsSqrt.Num());
+			TConstArrayView<float> DynamicWeightsSqrt = Database->CalculateDynamicWeightsSqrt(DynamicWeightsSqrtBuffer);
+
 			for (const FTraceMotionMatchingStatePoseEntry& PoseEntry : DbEntry.PoseEntries)
 			{
-				AddUnfilteredDatabaseRow(Database, UnfilteredDatabaseRows, SharedData, PoseEntry.DbPoseIdx, PoseEntry.PoseCandidateFlags, PoseToPCAValuesVectorIndexes, PoseEntry.Cost);
+				AddUnfilteredDatabaseRow(Database, DynamicWeightsSqrt, UnfilteredDatabaseRows, SharedData, PoseEntry.DbPoseIdx, PoseEntry.PoseCandidateFlags, PoseToPCAValuesVectorIndexes, PoseEntry.Cost);
 			}
 
 			if (bShowAllPoses)
@@ -303,7 +307,7 @@ void SDebuggerDatabaseView::Update(const FTraceMotionMatchingStateMessage& State
 					{
 						if (!PoseEntriesIdx.Find(DbPoseIdx))
 						{
-							AddUnfilteredDatabaseRow(Database, UnfilteredDatabaseRows, SharedData, DbPoseIdx, EPoseCandidateFlags::DiscardedBy_Search, PoseToPCAValuesVectorIndexes);
+							AddUnfilteredDatabaseRow(Database, DynamicWeightsSqrt, UnfilteredDatabaseRows, SharedData, DbPoseIdx, EPoseCandidateFlags::DiscardedBy_Search, PoseToPCAValuesVectorIndexes);
 						}
 					}
 				}
diff --git a/Engine/Plugins/Animation/PoseSearch/Source/Editor/Private/PoseSearchDebuggerView.cpp b/Engine/Plugins/Animation/PoseSearch/Source/Editor/Private/PoseSearchDebuggerView.cpp
--- a/Engine/Plugins/Animation/PoseSearch/Source/Editor/Private/PoseSearchDebuggerView.cpp
+++ b/Engine/Plugins/Animation/PoseSearch/Source/Editor/Private/PoseSearchDebuggerView.cpp
@@ -301,7 +301,7 @@ void SDebuggerView::Tick(const FGeometry& AllottedGeometry, const double InCurre
 						EAsyncBuildIndexResult::Success == FAsyncPoseSearchDatabasesManagement::RequestAsyncBuildIndex(CurrentDatabase, ERequestAsyncBuildFlag::ContinueRequest) &&
 						DbEntry.QueryVector.Num() == Database->Schema->SchemaCardinality)
 					{
-						FDebugDrawParams DrawParams(AnimContexts, PoseHistories, RoleToIndex, CurrentDatabase, EDebugDrawFlags::DrawQuery);
+						FDebugDrawParams DrawParams(AnimContexts, PoseHistories, RoleToIndex, CurrentDatabase);
 						DrawParams.DrawFeatureVector(DbEntry.QueryVector);
 						break;
 					}
diff --git a/Engine/Plugins/Animation/PoseSearch/Source/Runtime/Private/PoseSearchContext.cpp b/Engine/Plugins/Animation/PoseSearch/Source/Runtime/Private/PoseSearchContext.cpp
--- a/Engine/Plugins/Animation/PoseSearch/Source/Runtime/Private/PoseSearchContext.cpp
+++ b/Engine/Plugins/Animation/PoseSearch/Source/Runtime/Private/PoseSearchContext.cpp
@@ -60,12 +60,12 @@ static const UWorld* GetWorld(TConstArrayView<const UObject*> AnimContexts)
 
 //////////////////////////////////////////////////////////////////////////
 // FDebugDrawParams
+PRAGMA_DISABLE_DEPRECATION_WARNINGS
 FDebugDrawParams::FDebugDrawParams(TArrayView<FAnimInstanceProxy*> InAnimInstanceProxies, TConstArrayView<const IPoseHistory*> InPoseHistories, const FRoleToIndex& InRoleToIndex, const UPoseSearchDatabase* InDatabase, EDebugDrawFlags InFlags)
 : AnimContexts()
 , PoseHistories(InPoseHistories)
 , RoleToIndex(InRoleToIndex)
 , Database(InDatabase)
-, Flags(InFlags)
 {
 }
 
@@ -74,19 +74,24 @@ FDebugDrawParams::FDebugDrawParams(TArrayView<const USkinnedMeshComponent*> InMe
 , PoseHistories(InPoseHistories)
 , RoleToIndex(InRoleToIndex)
 , Database(InDatabase)
-, Flags(InFlags)
 {
 }
+PRAGMA_ENABLE_DEPRECATION_WARNINGS
 
-FDebugDrawParams::FDebugDrawParams(TConstArrayView<const UObject*> InAnimContexts, TConstArrayView<const IPoseHistory*> InPoseHistories, const FRoleToIndex& InRoleToIndex, const UPoseSearchDatabase* InDatabase, EDebugDrawFlags InFlags)
+FDebugDrawParams::FDebugDrawParams(TConstArrayView<const UObject*> InAnimContexts, TConstArrayView<const IPoseHistory*> InPoseHistories, const FRoleToIndex& InRoleToIndex, const UPoseSearchDatabase* InDatabase)
 : AnimContexts(InAnimContexts)
 , PoseHistories(InPoseHistories)
 , RoleToIndex(InRoleToIndex)
 , Database(InDatabase)
-, Flags(InFlags)
 {
 	check(RoleToIndex.Num() == PoseHistories.Num());
 	check(IsValid(RoleToIndex));
+
+	if (const UPoseSearchSchema* Schema = GetSchema())
+	{
+		DynamicWeightsSqrtBuffer.SetNum(Schema->SchemaCardinality);
+		DynamicWeightsSqrt = Database->CalculateDynamicWeightsSqrt(DynamicWeightsSqrtBuffer);
+	}
 }
 
 bool FDebugDrawParams::CanDraw() const
@@ -481,6 +486,21 @@ void FDebugDrawParams::DrawFeatureVector(int32 PoseIdx)
 		DrawFeatureVector(GetSearchIndex()->GetPoseValuesSafe(PoseIdx, BufferUsedForReconstruction));
 	}
 }
+
+bool FDebugDrawParams::IsAnyWeightRelevant(const UPoseSearchFeatureChannel* Channel) const
+{
+	const int32 StartIndex = Channel->GetChannelDataOffset();
+	const int32 EndIndex = StartIndex + Channel->GetChannelCardinality();
+	for (int32 Index = StartIndex; Index < EndIndex; ++Index)
+	{
+		if (DynamicWeightsSqrt[Index] > UE_SMALL_NUMBER)
+		{
+			return true;
+		}
+	}
+	return false;
+}
+
 #endif // ENABLE_DRAW_DEBUG
 
 //////////////////////////////////////////////////////////////////////////
diff --git a/Engine/Plugins/Animation/PoseSearch/Source/Runtime/Private/PoseSearchDatabase.cpp b/Engine/Plugins/Animation/PoseSearch/Source/Runtime/Private/PoseSearchDatabase.cpp
--- a/Engine/Plugins/Animation/PoseSearch/Source/Runtime/Private/PoseSearchDatabase.cpp
+++ b/Engine/Plugins/Animation/PoseSearch/Source/Runtime/Private/PoseSearchDatabase.cpp
@@ -11,7 +11,7 @@
 #include "PoseSearch/PoseSearchContext.h"
 #include "PoseSearch/PoseSearchDefines.h"
 #include "PoseSearch/PoseSearchDerivedData.h"
-#include "PoseSearch/PoseSearchFeatureChannel.h"
+#include "PoseSearchFeatureChannel_Group.h"
 #include "PoseSearch/PoseSearchHistory.h"
 #include "PoseSearch/PoseSearchSchema.h"
 #include "Serialization/ArchiveCountMem.h"
@@ -36,12 +36,23 @@ DEFINE_STAT(STAT_PoseSearch_VPTree);
 namespace UE::PoseSearch
 {
 #if WITH_EDITOR && ENABLE_ANIM_DEBUG
-static TAutoConsoleVariable<bool> CVarMotionMatchCompareAgainstBruteForce(TEXT("a.MotionMatch.CompareAgainstBruteForce"), false, TEXT("Compare optimized search against brute force search"));
-static TAutoConsoleVariable<bool> CVarMotionMatchValidateKNNSearch(TEXT("a.MotionMatch.ValidateKNNSearch"), false, TEXT("Validate KNN search"));
+
+static bool GVarMotionMatchCompareAgainstBruteForce = false;
+static FAutoConsoleVariableRef CVarMotionMatchCompareAgainstBruteForce(TEXT("a.MotionMatch.CompareAgainstBruteForce"), GVarMotionMatchCompareAgainstBruteForce, TEXT("Compare optimized search against brute force search"));
+
+static bool GVarMotionMatchValidateKNNSearch = false;
+static FAutoConsoleVariableRef CVarMotionMatchValidateKNNSearch(TEXT("a.MotionMatch.ValidateKNNSearch"), GVarMotionMatchValidateKNNSearch, TEXT("Validate KNN search"));
+
 // Experimental, this feature might be removed without warning, not for production use
-static TAutoConsoleVariable<bool> CVarMotionMatchProfileMaxHeapKNNSearch(TEXT("a.MotionMatch.ProfileMaxHeapKNNSearch"), false, TEXT("Profile MaxHeap KNN Search"));
+static bool GVarMotionMatchProfileMaxHeapKNNSearch = false;
+static FAutoConsoleVariableRef CVarMotionMatchProfileMaxHeapKNNSearch(TEXT("a.MotionMatch.ProfileMaxHeapKNNSearch"), GVarMotionMatchProfileMaxHeapKNNSearch, TEXT("Profile MaxHeap KNN Search"));
+
 #endif
 
+// Experimental, this feature might be removed without warning, not for production use
+static int32 GVarMotionMatchDebugWeightGroupID = 0;
+static FAutoConsoleVariableRef CVarMotionMatchDebugWeightGroupID(TEXT("a.MotionMatch.DebugWeightGroupID"), GVarMotionMatchDebugWeightGroupID, TEXT("Only the channels with this or negative DebugWeightGroupID will have non zero weights"));
+
 struct FSearchFilters
 {
 	FSearchFilters(const UPoseSearchSchema* Schema, TConstArrayView<int32> NonSelectableIdx, TConstArrayView<int32> SelectableAssetIdx, bool bAddBlockTransitionFilter)
@@ -70,7 +81,9 @@ struct FSearchFilters
 		}
 	}
 
-	bool AreFiltersValid(const FSearchIndex& SearchIndex, TConstArrayView<float> PoseValues, TConstArrayView<float> QueryValues, int32 PoseIdx
+	// @todo: template this with bAlignedAndPadded to be able to use faster ComparePoses
+	bool AreFiltersValid(const FSearchIndex& SearchIndex, TConstArrayView<float> PoseValues, TConstArrayView<float> QueryValues, 
+		TConstArrayView<float> DynamicWeightsSqrt, int32 PoseIdx
 #if UE_POSE_SEARCH_TRACE_ENABLED
 		, UE::PoseSearch::FSearchContext& SearchContext, const UPoseSearchDatabase* Database
 #endif // UE_POSE_SEARCH_TRACE_ENABLED
@@ -87,17 +100,17 @@ struct FSearchFilters
 				}
 				else if (Filter == &SelectableAssetIdxFilter)
 				{
-					const FPoseSearchCost PoseCost = SearchIndex.ComparePoses(PoseIdx, 0.f, PoseValues, QueryValues);
+					const FPoseSearchCost PoseCost = SearchIndex.ComparePoses(PoseIdx, 0.f, PoseValues, QueryValues, DynamicWeightsSqrt);
 					SearchContext.Track(Database, PoseIdx, EPoseCandidateFlags::DiscardedBy_AssetIdxFilter, PoseCost);
 				}
 				else if (Filter == &BlockTransitionFilter)
 				{
-					const FPoseSearchCost PoseCost = SearchIndex.ComparePoses(PoseIdx, 0.f, PoseValues, QueryValues);
+					const FPoseSearchCost PoseCost = SearchIndex.ComparePoses(PoseIdx, 0.f, PoseValues, QueryValues, DynamicWeightsSqrt);
 					SearchContext.Track(Database, PoseIdx, EPoseCandidateFlags::DiscardedBy_BlockTransition, PoseCost);
 				}
 				else
 				{
-					const FPoseSearchCost PoseCost = SearchIndex.ComparePoses(PoseIdx, 0.f, PoseValues, QueryValues);
+					const FPoseSearchCost PoseCost = SearchIndex.ComparePoses(PoseIdx, 0.f, PoseValues, QueryValues, DynamicWeightsSqrt);
 					SearchContext.Track(Database, PoseIdx, EPoseCandidateFlags::DiscardedBy_PoseFilter, PoseCost);
 				}
 #endif // UE_POSE_SEARCH_TRACE_ENABLED
@@ -174,19 +187,19 @@ private:
 
 template<bool bReconstructPoseValues, bool bAlignedAndPadded>
 static inline void EvaluatePoseKernel(UE::PoseSearch::FSearchResult& Result, const UE::PoseSearch::FSearchIndex& SearchIndex, TConstArrayView<float> QueryValues, TArrayView<float> ReconstructedPoseValuesBuffer,
-	int32 PoseIdx, const UE::PoseSearch::FSearchFilters& SearchFilters, UE::PoseSearch::FSearchContext& SearchContext, const UPoseSearchDatabase* Database, bool bUpdateBestCandidates, int32 ResultIndex = -1)
+	int32 PoseIdx, const UE::PoseSearch::FSearchFilters& SearchFilters, UE::PoseSearch::FSearchContext& SearchContext, const UPoseSearchDatabase* Database, TConstArrayView<float> DynamicWeightsSqrt, bool bUpdateBestCandidates, int32 ResultIndex = -1)
 {
 	using namespace UE::PoseSearch;
 
 	const TConstArrayView<float> PoseValues = bReconstructPoseValues ? SearchIndex.GetReconstructedPoseValues(PoseIdx, ReconstructedPoseValuesBuffer) : SearchIndex.GetPoseValues(PoseIdx);
 
-	if (SearchFilters.AreFiltersValid(SearchIndex, PoseValues, QueryValues, PoseIdx
+	if (SearchFilters.AreFiltersValid(SearchIndex, PoseValues, QueryValues, DynamicWeightsSqrt, PoseIdx
 #if UE_POSE_SEARCH_TRACE_ENABLED
 		, SearchContext, Database
 #endif // UE_POSE_SEARCH_TRACE_ENABLED
 	))
 	{
-		const FPoseSearchCost PoseCost = bAlignedAndPadded ? SearchIndex.CompareAlignedPoses(PoseIdx, 0.f, PoseValues, QueryValues) : SearchIndex.ComparePoses(PoseIdx, 0.f, PoseValues, QueryValues);
+		const FPoseSearchCost PoseCost = bAlignedAndPadded ? SearchIndex.CompareAlignedPoses(PoseIdx, 0.f, PoseValues, QueryValues, DynamicWeightsSqrt) : SearchIndex.ComparePoses(PoseIdx, 0.f, PoseValues, QueryValues, DynamicWeightsSqrt);
 		if (PoseCost < Result.PoseCost)
 		{
 			Result.PoseCost = PoseCost;
@@ -633,6 +646,42 @@ TConstArrayView<int32> UPoseSearchDatabase::GetAssetIndexesForSourceAsset(const
 	return TConstArrayView<int32>();
 }
 
+TConstArrayView<float> UPoseSearchDatabase::CalculateDynamicWeightsSqrt(TArrayView<float> DynamicWeightsSqrtBuffer) const
+{
+	using namespace UE::PoseSearch;
+	check(IsAligned(DynamicWeightsSqrtBuffer.GetData(), alignof(VectorRegister4Float)));
+
+	const UE::PoseSearch::FSearchIndex& SearchIndex = GetSearchIndex();
+	check(DynamicWeightsSqrtBuffer.Num() == SearchIndex.WeightsSqrt.Num());
+
+	bool bInitialized = false;
+	Schema->FindChannel([&bInitialized, &SearchIndex, &DynamicWeightsSqrtBuffer](const UPoseSearchFeatureChannel* Channel)
+		{
+			if (const UPoseSearchFeatureChannel_GroupBase* ChannelGroupBase = Cast<UPoseSearchFeatureChannel_GroupBase>(Channel))
+			{
+				if (ChannelGroupBase->DebugWeightGroupID != INDEX_NONE && ChannelGroupBase->DebugWeightGroupID != GVarMotionMatchDebugWeightGroupID)
+				{
+					if (!bInitialized)
+					{
+						// initializing DynamicWeightsSqrtBuffer with the SearchIndex.WeightsSqrt
+						FMemory::Memcpy(DynamicWeightsSqrtBuffer.GetData(), SearchIndex.WeightsSqrt.GetData(), SearchIndex.WeightsSqrt.Num() * sizeof(float));
+						bInitialized = true;
+					}
+
+					// zeroing out interval of weights in DynamicWeightsSqrtBuffer associated with ChannelGroupBase since its GroupID mismatch ValidWeightChannelGroup
+					FMemory::Memzero(DynamicWeightsSqrtBuffer.GetData() + ChannelGroupBase->GetChannelDataOffset(), ChannelGroupBase->GetChannelCardinality() * sizeof(float));
+				}
+			}
+			return nullptr;
+		});
+
+	if (bInitialized)
+	{
+		return DynamicWeightsSqrtBuffer;
+	}
+
+	return SearchIndex.WeightsSqrt;
+}
 
 const UE::PoseSearch::FSearchIndex& UPoseSearchDatabase::GetSearchIndex() const
 {
@@ -709,7 +758,7 @@ int32 UPoseSearchDatabase::GetNumberOfPrincipalComponents() const
 bool UPoseSearchDatabase::GetSkipSearchIfPossible() const
 {
 #if WITH_EDITOR && ENABLE_ANIM_DEBUG && UE_POSE_SEARCH_TRACE_ENABLED
-	if (UE::PoseSearch::CVarMotionMatchCompareAgainstBruteForce.GetValueOnAnyThread())
+	if (UE::PoseSearch::GVarMotionMatchCompareAgainstBruteForce)
 	{
 		return false;
 	}
@@ -1061,7 +1110,7 @@ UE::PoseSearch::FSearchResult UPoseSearchDatabase::Search(UE::PoseSearch::FSearc
 			Result = SearchBruteForce(SearchContext);
 
 #if WITH_EDITOR && ENABLE_ANIM_DEBUG && UE_POSE_SEARCH_TRACE_ENABLED
-			if (CVarMotionMatchCompareAgainstBruteForce.GetValueOnAnyThread())
+			if (GVarMotionMatchCompareAgainstBruteForce)
 			{
 				Result.BruteForcePoseCost = Result.PoseCost;
 			}
@@ -1073,7 +1122,7 @@ UE::PoseSearch::FSearchResult UPoseSearchDatabase::Search(UE::PoseSearch::FSearc
 			Result = SearchVPTree(SearchContext);
 
 #if WITH_EDITOR && ENABLE_ANIM_DEBUG && UE_POSE_SEARCH_TRACE_ENABLED
-			if (CVarMotionMatchCompareAgainstBruteForce.GetValueOnAnyThread())
+			if (GVarMotionMatchCompareAgainstBruteForce)
 			{
 				Result.BruteForcePoseCost = SearchBruteForce(SearchContext).PoseCost;
 			}
@@ -1085,7 +1134,7 @@ UE::PoseSearch::FSearchResult UPoseSearchDatabase::Search(UE::PoseSearch::FSearc
 			Result = SearchPCAKDTree(SearchContext);
 
 #if WITH_EDITOR && ENABLE_ANIM_DEBUG && UE_POSE_SEARCH_TRACE_ENABLED
-			if (CVarMotionMatchCompareAgainstBruteForce.GetValueOnAnyThread())
+			if (GVarMotionMatchCompareAgainstBruteForce)
 			{
 				Result.BruteForcePoseCost = SearchBruteForce(SearchContext).PoseCost;
 			}
@@ -1104,7 +1153,7 @@ UE::PoseSearch::FSearchResult UPoseSearchDatabase::Search(UE::PoseSearch::FSearc
 
 void UPoseSearchDatabase::PopulateNonSelectableIdx(FNonSelectableIdx& NonSelectableIdx, UE::PoseSearch::FSearchContext& SearchContext
 #if UE_POSE_SEARCH_TRACE_ENABLED
-	, TConstArrayView<float> QueryValues
+	, TConstArrayView<float> QueryValues, TConstArrayView<float> DynamicWeightsSqrt
 #endif //UE_POSE_SEARCH_TRACE_ENABLED
 ) const
 {
@@ -1139,7 +1188,7 @@ void UPoseSearchDatabase::PopulateNonSelectableIdx(FNonSelectableIdx& NonSelecta
 
 #if UE_POSE_SEARCH_TRACE_ENABLED
 						const TConstArrayView<float> PoseValues = SearchIndex.GetPoseValuesSafe(PoseIdx, BufferUsedForReconstruction);
-						const FPoseSearchCost PoseCost = SearchIndex.ComparePoses(PoseIdx, 0.f, PoseValues, QueryValues);
+						const FPoseSearchCost PoseCost = SearchIndex.ComparePoses(PoseIdx, 0.f, PoseValues, QueryValues, DynamicWeightsSqrt);
 						SearchContext.Track(this, PoseIdx, EPoseCandidateFlags::DiscardedBy_AssetReselection, PoseCost);
 #endif // UE_POSE_SEARCH_TRACE_ENABLED
 					}
@@ -1166,7 +1215,7 @@ void UPoseSearchDatabase::PopulateNonSelectableIdx(FNonSelectableIdx& NonSelecta
 
 #if UE_POSE_SEARCH_TRACE_ENABLED
 						const TConstArrayView<float> PoseValues = SearchIndex.GetPoseValuesSafe(PoseIdx, BufferUsedForReconstruction);
-						const FPoseSearchCost PoseCost = SearchIndex.ComparePoses(PoseIdx, 0.f, PoseValues, QueryValues);
+						const FPoseSearchCost PoseCost = SearchIndex.ComparePoses(PoseIdx, 0.f, PoseValues, QueryValues, DynamicWeightsSqrt);
 						SearchContext.Track(this, PoseIdx, EPoseCandidateFlags::DiscardedBy_PoseJumpThresholdTime, PoseCost);
 #endif // UE_POSE_SEARCH_TRACE_ENABLED
 					}
@@ -1182,7 +1231,7 @@ void UPoseSearchDatabase::PopulateNonSelectableIdx(FNonSelectableIdx& NonSelecta
 
 #if UE_POSE_SEARCH_TRACE_ENABLED
 						const TConstArrayView<float> PoseValues = SearchIndex.GetPoseValuesSafe(PoseIdx, BufferUsedForReconstruction);
-						const FPoseSearchCost PoseCost = SearchIndex.ComparePoses(PoseIdx, 0.f, PoseValues, QueryValues);
+						const FPoseSearchCost PoseCost = SearchIndex.ComparePoses(PoseIdx, 0.f, PoseValues, QueryValues, DynamicWeightsSqrt);
 						SearchContext.Track(this, PoseIdx, EPoseCandidateFlags::DiscardedBy_PoseJumpThresholdTime, PoseCost);
 #endif // UE_POSE_SEARCH_TRACE_ENABLED
 					}
@@ -1207,7 +1256,7 @@ void UPoseSearchDatabase::PopulateNonSelectableIdx(FNonSelectableIdx& NonSelecta
 				// if we're editing the database and removing assets it's possible that the PoseIndicesHistory contains invalid pose indexes
 				if (HistoricalPoseIndex.PoseIndex < SearchIndex.GetNumPoses())
 				{
-					const FPoseSearchCost PoseCost = SearchIndex.ComparePoses(HistoricalPoseIndex.PoseIndex, 0.f, SearchIndex.GetPoseValuesSafe(HistoricalPoseIndex.PoseIndex, BufferUsedForReconstruction), QueryValues);
+					const FPoseSearchCost PoseCost = SearchIndex.ComparePoses(HistoricalPoseIndex.PoseIndex, 0.f, SearchIndex.GetPoseValuesSafe(HistoricalPoseIndex.PoseIndex, BufferUsedForReconstruction), QueryValues, DynamicWeightsSqrt);
 					SearchContext.Track(this, HistoricalPoseIndex.PoseIndex, EPoseCandidateFlags::DiscardedBy_PoseReselectHistory, PoseCost);
 				}
 #endif // UE_POSE_SEARCH_TRACE_ENABLED
@@ -1308,16 +1357,20 @@ UE::PoseSearch::FSearchResult UPoseSearchDatabase::SearchContinuingPose(UE::Pose
 		check(IsAligned(ReconstructedPoseValuesBuffer.GetData(), alignof(VectorRegister4Float)));
 		const TConstArrayView<float> PoseValues = SearchIndex.IsValuesEmpty() ? SearchIndex.GetReconstructedPoseValues(PoseIdx, ReconstructedPoseValuesBuffer) : SearchIndex.GetPoseValues(PoseIdx);
 
+		// @todo: perhaps store the DynamicWeightsSqrt into the GetOrBuildQuery to share between the continuing pose and the full search
+		TArrayView<float> DynamicWeightsSqrtBuffer((float*)FMemory_Alloca(NumDimensions * sizeof(float)), NumDimensions);
+		const TConstArrayView<float> DynamicWeightsSqrt = CalculateDynamicWeightsSqrt(DynamicWeightsSqrtBuffer);
+
 		const int32 ContinuingPoseIdx = SearchContext.GetCurrentResult().PoseIdx;
 		// is the data padded at 16 bytes (and 16 bytes aligned by construction)?
 		if (NumDimensions % 4 == 0)
 		{
-			Result.PoseCost = SearchIndex.CompareAlignedPoses(ContinuingPoseIdx, UpdatedContinuingPoseCostBias, PoseValues, SearchContext.GetOrBuildQuery(Schema));
+			Result.PoseCost = SearchIndex.CompareAlignedPoses(ContinuingPoseIdx, UpdatedContinuingPoseCostBias, PoseValues, SearchContext.GetOrBuildQuery(Schema), DynamicWeightsSqrt);
 		}
 		// data is not 16 bytes padded
 		else
 		{
-			Result.PoseCost = SearchIndex.ComparePoses(ContinuingPoseIdx, UpdatedContinuingPoseCostBias, PoseValues, SearchContext.GetOrBuildQuery(Schema));
+			Result.PoseCost = SearchIndex.ComparePoses(ContinuingPoseIdx, UpdatedContinuingPoseCostBias, PoseValues, SearchContext.GetOrBuildQuery(Schema), DynamicWeightsSqrt);
 		}
 
 		Result.AssetTime = SearchContext.GetCurrentResult().AssetTime;
@@ -1330,7 +1383,7 @@ UE::PoseSearch::FSearchResult UPoseSearchDatabase::SearchContinuingPose(UE::Pose
 	}
 
 #if WITH_EDITOR && ENABLE_ANIM_DEBUG && UE_POSE_SEARCH_TRACE_ENABLED
-	if (CVarMotionMatchCompareAgainstBruteForce.GetValueOnAnyThread())
+	if (GVarMotionMatchCompareAgainstBruteForce)
 	{
 		Result.BruteForcePoseCost = Result.PoseCost;
 	}
@@ -1361,7 +1414,12 @@ UE::PoseSearch::FSearchResult UPoseSearchDatabase::SearchPCAKDTree(UE::PoseSearc
 		//stack allocated temporaries
 		TArrayView<float> ProjectedQueryValues((float*)FMemory_Alloca(ClampedNumberOfPrincipalComponents * sizeof(float)), ClampedNumberOfPrincipalComponents);
 	
+		// @todo: perhaps store the DynamicWeightsSqrt into the GetOrBuildQuery to share between the continuing pose and the full search
+		TArrayView<float> DynamicWeightsSqrtBuffer((float*)FMemory_Alloca(NumDimensions * sizeof(float)), NumDimensions);
+		const TConstArrayView<float> DynamicWeightsSqrt = CalculateDynamicWeightsSqrt(DynamicWeightsSqrtBuffer);
+
 		TConstArrayView<float> QueryValues = SearchContext.GetOrBuildQuery(Schema);
+		check(QueryValues.Num() == NumDimensions);
 
 		FSelectableAssetIdx SelectableAssetIdx;
 		PopulateSelectableAssetIdx(SelectableAssetIdx, SearchContext.GetAssetsToConsider());
@@ -1369,18 +1427,15 @@ UE::PoseSearch::FSearchResult UPoseSearchDatabase::SearchPCAKDTree(UE::PoseSearc
 		FNonSelectableIdx NonSelectableIdx;
 		PopulateNonSelectableIdx(NonSelectableIdx, SearchContext
 #if UE_POSE_SEARCH_TRACE_ENABLED
-			, QueryValues
+			, QueryValues, DynamicWeightsSqrt
 #endif // UE_POSE_SEARCH_TRACE_ENABLED
 		);
 
 		bool bRunNonSelectableIdxPostKDTree = bArePCAValuesPruned;
 #if WITH_EDITOR && ENABLE_ANIM_DEBUG
-		const bool bProfileMaxHeapKNNSearch = CVarMotionMatchProfileMaxHeapKNNSearch.GetValueOnAnyThread();
-		const bool bValidateKNNSearch = CVarMotionMatchValidateKNNSearch.GetValueOnAnyThread();
-		bRunNonSelectableIdxPostKDTree |= bValidateKNNSearch;
+		bRunNonSelectableIdxPostKDTree |= GVarMotionMatchValidateKNNSearch;
 #endif // WITH_EDITOR && ENABLE_ANIM_DEBUG
 
-		check(QueryValues.Num() == NumDimensions);
 		// projecting QueryValues into the PCA space 
 		TConstArrayView<float> PCAQueryValues = SearchIndex.PCAProject(QueryValues, ProjectedQueryValues);
 		check(PCAQueryValues.Num() == ClampedNumberOfPrincipalComponents);
@@ -1390,7 +1445,7 @@ UE::PoseSearch::FSearchResult UPoseSearchDatabase::SearchPCAKDTree(UE::PoseSearc
 		TArrayView<FKDTree::FKNNMaxHeapResultSet::FResult> Results((FKDTree::FKNNMaxHeapResultSet::FResult*)FMemory_Alloca(ClampedKDTreeQueryNumNeighbors * sizeof(FKDTree::FKNNMaxHeapResultSet::FResult)), ClampedKDTreeQueryNumNeighbors);
 
 #if WITH_EDITOR && ENABLE_ANIM_DEBUG
-		const double StartTimeMaxHeap = bProfileMaxHeapKNNSearch ? FPlatformTime::Seconds() : 0;
+		const double StartTimeMaxHeap = GVarMotionMatchProfileMaxHeapKNNSearch ? FPlatformTime::Seconds() : 0;
 #endif // WITH_EDITOR && ENABLE_ANIM_DEBUG
 
 		if (bRunNonSelectableIdxPostKDTree || NonSelectableIdx.IsEmpty())
@@ -1409,7 +1464,7 @@ UE::PoseSearch::FSearchResult UPoseSearchDatabase::SearchPCAKDTree(UE::PoseSearc
 		TArrayView<int32> ResultIndexes((int32*)FMemory_Alloca((ClampedKDTreeQueryNumNeighbors + 1) * sizeof(int32)), ClampedKDTreeQueryNumNeighbors + 1);
 		TArrayView<float> ResultDistanceSqr((float*)FMemory_Alloca((ClampedKDTreeQueryNumNeighbors + 1) * sizeof(float)), ClampedKDTreeQueryNumNeighbors + 1);
 
-		if (bProfileMaxHeapKNNSearch)
+		if (GVarMotionMatchProfileMaxHeapKNNSearch)
 		{
 			// debug code to log profiling comparison between FKNNResultSet/FFilteredKNNResultSet vs FKNNMaxHeapResultSet/FFilteredKNNMaxHeapResultSet
 			const double EndTimeMaxHeap = FPlatformTime::Seconds();
@@ -1444,7 +1499,7 @@ UE::PoseSearch::FSearchResult UPoseSearchDatabase::SearchPCAKDTree(UE::PoseSearc
 		}
 		SortedResultsIndexes.Sort([&Results](const int32 IndexA, const int32 IndexB) { return Results[IndexA].Distance < Results[IndexB].Distance; });
 
-		if (bProfileMaxHeapKNNSearch)
+		if (GVarMotionMatchProfileMaxHeapKNNSearch)
 		{
 			for (int32 ResultIndex = 0; ResultIndex < NumResults; ++ResultIndex)
 			{
@@ -1460,7 +1515,7 @@ UE::PoseSearch::FSearchResult UPoseSearchDatabase::SearchPCAKDTree(UE::PoseSearc
 			}
 		}
 
-		if (bValidateKNNSearch)
+		if (GVarMotionMatchValidateKNNSearch)
 		{
 			const int32 NumPCAValuesVectors = SearchIndex.GetNumPCAValuesVectors(ClampedNumberOfPrincipalComponents);
 
@@ -1520,7 +1575,7 @@ UE::PoseSearch::FSearchResult UPoseSearchDatabase::SearchPCAKDTree(UE::PoseSearc
 					const TConstArrayView<int32> PoseIndexes = SearchIndex.PCAValuesVectorToPoseIndexes[Results[ResultIndex].Index];
 					for (int32 Index = 0; Index < PoseIndexes.Num() && NumEvaluatePoseKernelCalls < MaxNumEvaluatePoseKernelCalls; ++Index, ++NumEvaluatePoseKernelCalls)
 					{
-						EvaluatePoseKernel<false, true>(Result, SearchIndex, QueryValues, TArrayView<float>(), PoseIndexes[Index], SearchFilters, SearchContext, this, true
+						EvaluatePoseKernel<false, true>(Result, SearchIndex, QueryValues, TArrayView<float>(), PoseIndexes[Index], SearchFilters, SearchContext, this, DynamicWeightsSqrt, true
 #if WITH_EDITOR && ENABLE_ANIM_DEBUG
 							, SortedResultsIndexes[ResultIndex]
 #endif // WITH_EDITOR && ENABLE_ANIM_DEBUG
@@ -1536,7 +1591,7 @@ UE::PoseSearch::FSearchResult UPoseSearchDatabase::SearchPCAKDTree(UE::PoseSearc
 					const TConstArrayView<int32> PoseIndexes = SearchIndex.PCAValuesVectorToPoseIndexes[Results[ResultIndex].Index];
 					for (int32 Index = 0; Index < PoseIndexes.Num() && NumEvaluatePoseKernelCalls < MaxNumEvaluatePoseKernelCalls; ++Index, ++NumEvaluatePoseKernelCalls)
 					{
-						EvaluatePoseKernel<false, false>(Result, SearchIndex, QueryValues, TArrayView<float>(), PoseIndexes[Index], SearchFilters, SearchContext, this, true
+						EvaluatePoseKernel<false, false>(Result, SearchIndex, QueryValues, TArrayView<float>(), PoseIndexes[Index], SearchFilters, SearchContext, this, DynamicWeightsSqrt, true
 #if WITH_EDITOR && ENABLE_ANIM_DEBUG
 							, SortedResultsIndexes[ResultIndex]
 #endif // WITH_EDITOR && ENABLE_ANIM_DEBUG
@@ -1553,7 +1608,7 @@ UE::PoseSearch::FSearchResult UPoseSearchDatabase::SearchPCAKDTree(UE::PoseSearc
 			check(IsAligned(ReconstructedPoseValuesBuffer.GetData(), alignof(VectorRegister4Float)));
 			for (int32 ResultIndex = 0; ResultIndex < NumResults; ++ResultIndex)
 			{
-				EvaluatePoseKernel<true, false>(Result, SearchIndex, QueryValues, ReconstructedPoseValuesBuffer, Results[ResultIndex].Index, SearchFilters, SearchContext, this, true
+				EvaluatePoseKernel<true, false>(Result, SearchIndex, QueryValues, ReconstructedPoseValuesBuffer, Results[ResultIndex].Index, SearchFilters, SearchContext, this, DynamicWeightsSqrt, true
 #if WITH_EDITOR && ENABLE_ANIM_DEBUG
 					, SortedResultsIndexes[ResultIndex]
 #endif // WITH_EDITOR && ENABLE_ANIM_DEBUG
@@ -1565,7 +1620,7 @@ UE::PoseSearch::FSearchResult UPoseSearchDatabase::SearchPCAKDTree(UE::PoseSearc
 		{
 			for (int32 ResultIndex = 0; ResultIndex < NumResults; ++ResultIndex)
 			{
-				EvaluatePoseKernel<false, true>(Result, SearchIndex, QueryValues, TArrayView<float>(), Results[ResultIndex].Index, SearchFilters, SearchContext, this, true
+				EvaluatePoseKernel<false, true>(Result, SearchIndex, QueryValues, TArrayView<float>(), Results[ResultIndex].Index, SearchFilters, SearchContext, this, DynamicWeightsSqrt, true
 #if WITH_EDITOR && ENABLE_ANIM_DEBUG
 					, SortedResultsIndexes[ResultIndex]
 #endif // WITH_EDITOR && ENABLE_ANIM_DEBUG
@@ -1577,7 +1632,7 @@ UE::PoseSearch::FSearchResult UPoseSearchDatabase::SearchPCAKDTree(UE::PoseSearc
 		{
 			for (int32 ResultIndex = 0; ResultIndex < NumResults; ++ResultIndex)
 			{
-				EvaluatePoseKernel<false, false>(Result, SearchIndex, QueryValues, TArrayView<float>(), Results[ResultIndex].Index, SearchFilters, SearchContext, this, true
+				EvaluatePoseKernel<false, false>(Result, SearchIndex, QueryValues, TArrayView<float>(), Results[ResultIndex].Index, SearchFilters, SearchContext, this, DynamicWeightsSqrt, true
 #if WITH_EDITOR && ENABLE_ANIM_DEBUG
 					, SortedResultsIndexes[ResultIndex]
 #endif // WITH_EDITOR && ENABLE_ANIM_DEBUG
@@ -1588,10 +1643,14 @@ UE::PoseSearch::FSearchResult UPoseSearchDatabase::SearchPCAKDTree(UE::PoseSearc
 	else
 	{
 #if UE_POSE_SEARCH_TRACE_ENABLED
+		// @todo: perhaps store the DynamicWeightsSqrt into the GetOrBuildQuery to share between the continuing pose and the full search
+		TArrayView<float> DynamicWeightsSqrtBuffer((float*)FMemory_Alloca(NumDimensions * sizeof(float)), NumDimensions);
+		const TConstArrayView<float> DynamicWeightsSqrt = CalculateDynamicWeightsSqrt(DynamicWeightsSqrtBuffer);
+
 		// calling just for reporting non selectable poses
 		TConstArrayView<float> QueryValues = SearchContext.GetOrBuildQuery(Schema);
 		FNonSelectableIdx NonSelectableIdx;
-		PopulateNonSelectableIdx(NonSelectableIdx, SearchContext, QueryValues);
+		PopulateNonSelectableIdx(NonSelectableIdx, SearchContext, QueryValues, DynamicWeightsSqrt);
 #endif // UE_POSE_SEARCH_TRACE_ENABLED
 	}
 
@@ -1619,7 +1678,13 @@ UE::PoseSearch::FSearchResult UPoseSearchDatabase::SearchVPTree(UE::PoseSearch::
 	// there's no point in performing the search if CurrentBestTotalCost is already better than that
 	if (!GetSkipSearchIfPossible() || SearchContext.GetCurrentBestTotalCost() > SearchIndex.MinCostAddend)
 	{
+		// @todo: perhaps store the DynamicWeightsSqrt into the GetOrBuildQuery to share between the continuing pose and the full search
+		const int32 NumDimensions = Schema->SchemaCardinality;
+		TArrayView<float> DynamicWeightsSqrtBuffer((float*)FMemory_Alloca(NumDimensions * sizeof(float)), NumDimensions);
+		const TConstArrayView<float> DynamicWeightsSqrt = CalculateDynamicWeightsSqrt(DynamicWeightsSqrtBuffer);
+
 		TConstArrayView<float> QueryValues = SearchContext.GetOrBuildQuery(Schema);
+		check(QueryValues.Num() == NumDimensions);
 
 		FSelectableAssetIdx SelectableAssetIdx;
 		PopulateSelectableAssetIdx(SelectableAssetIdx, SearchContext.GetAssetsToConsider());
@@ -1628,12 +1693,9 @@ UE::PoseSearch::FSearchResult UPoseSearchDatabase::SearchVPTree(UE::PoseSearch::
 		FNonSelectableIdx NonSelectableIdx;
 		PopulateNonSelectableIdx(NonSelectableIdx, SearchContext
 #if UE_POSE_SEARCH_TRACE_ENABLED
-			, QueryValues
+			, QueryValues, DynamicWeightsSqrt
 #endif // UE_POSE_SEARCH_TRACE_ENABLED
 		);
-
-		const int32 NumDimensions = Schema->SchemaCardinality;
-		check(QueryValues.Num() == NumDimensions);
 		
 		const FSearchFilters SearchFilters(Schema, NonSelectableIdx, SelectableAssetIdx, SearchIndex.bAnyBlockTransition);
 
@@ -1657,7 +1719,7 @@ UE::PoseSearch::FSearchResult UPoseSearchDatabase::SearchVPTree(UE::PoseSearch::
 				const TConstArrayView<int32> PoseIndexes = SearchIndex.ValuesVectorToPoseIndexes[IndexDistance.Index];
 				for (int32 Index = 0; Index < PoseIndexes.Num() && NumEvaluatePoseKernelCalls < MaxNumEvaluatePoseKernelCalls; ++Index, ++NumEvaluatePoseKernelCalls)
 				{
-					EvaluatePoseKernel<false, false>(Result, SearchIndex, QueryValues, TArrayView<float>(), PoseIndexes[Index], SearchFilters, SearchContext, this, true
+					EvaluatePoseKernel<false, false>(Result, SearchIndex, QueryValues, TArrayView<float>(), PoseIndexes[Index], SearchFilters, SearchContext, this, DynamicWeightsSqrt, true
 #if WITH_EDITOR && ENABLE_ANIM_DEBUG
 						, ResultIndex
 #endif // WITH_EDITOR && ENABLE_ANIM_DEBUG
@@ -1672,7 +1734,7 @@ UE::PoseSearch::FSearchResult UPoseSearchDatabase::SearchVPTree(UE::PoseSearch::
 				const FIndexDistance& IndexDistance = UnsortedResults[ResultIndex];
 
 				// @todo: IndexDistance.Distance is the Sqrt(DissimilarityCost), so there's no need to calculate it again in SearchIndex.ComparePoses
-				EvaluatePoseKernel<false, false>(Result, SearchIndex, QueryValues, TArrayView<float>(), IndexDistance.Index, SearchFilters, SearchContext, this, true
+				EvaluatePoseKernel<false, false>(Result, SearchIndex, QueryValues, TArrayView<float>(), IndexDistance.Index, SearchFilters, SearchContext, this, DynamicWeightsSqrt, true
 #if WITH_EDITOR && ENABLE_ANIM_DEBUG
 					, ResultIndex
 #endif // WITH_EDITOR && ENABLE_ANIM_DEBUG
@@ -1683,10 +1745,15 @@ UE::PoseSearch::FSearchResult UPoseSearchDatabase::SearchVPTree(UE::PoseSearch::
 	else
 	{
 #if UE_POSE_SEARCH_TRACE_ENABLED
+		// @todo: perhaps store the DynamicWeightsSqrt into the GetOrBuildQuery to share between the continuing pose and the full search
+		const int32 NumDimensions = Schema->SchemaCardinality;
+		TArrayView<float> DynamicWeightsSqrtBuffer((float*)FMemory_Alloca(NumDimensions * sizeof(float)), NumDimensions);
+		const TConstArrayView<float> DynamicWeightsSqrt = CalculateDynamicWeightsSqrt(DynamicWeightsSqrtBuffer);
+
 		// calling just for reporting non selectable poses
 		TConstArrayView<float> QueryValues = SearchContext.GetOrBuildQuery(Schema);
 		FNonSelectableIdx NonSelectableIdx;
-		PopulateNonSelectableIdx(NonSelectableIdx, SearchContext, QueryValues);
+		PopulateNonSelectableIdx(NonSelectableIdx, SearchContext, QueryValues, DynamicWeightsSqrt);
 #endif // UE_POSE_SEARCH_TRACE_ENABLED
 	}
 
@@ -1714,6 +1781,11 @@ UE::PoseSearch::FSearchResult UPoseSearchDatabase::SearchBruteForce(UE::PoseSear
 	// there's no point in performing the search if CurrentBestTotalCost is already better than that
 	if (!GetSkipSearchIfPossible() || SearchContext.GetCurrentBestTotalCost() > SearchIndex.MinCostAddend)
 	{
+		// @todo: perhaps store the DynamicWeightsSqrt into the GetOrBuildQuery to share between the continuing pose and the full search
+		const int32 NumDimensions = Schema->SchemaCardinality;
+		TArrayView<float> DynamicWeightsSqrtBuffer((float*)FMemory_Alloca(NumDimensions * sizeof(float)), NumDimensions);
+		const TConstArrayView<float> DynamicWeightsSqrt = CalculateDynamicWeightsSqrt(DynamicWeightsSqrtBuffer);
+
 		TConstArrayView<float> QueryValues = SearchContext.GetOrBuildQuery(Schema);
 
 		FSelectableAssetIdx SelectableAssetIdx;
@@ -1722,11 +1794,10 @@ UE::PoseSearch::FSearchResult UPoseSearchDatabase::SearchBruteForce(UE::PoseSear
 		FNonSelectableIdx NonSelectableIdx;
 		PopulateNonSelectableIdx(NonSelectableIdx, SearchContext
 #if UE_POSE_SEARCH_TRACE_ENABLED
-			, QueryValues
+			, QueryValues, DynamicWeightsSqrt
 #endif // UE_POSE_SEARCH_TRACE_ENABLED
 		);
 
-		const int32 NumDimensions = Schema->SchemaCardinality;
 		const bool bUpdateBestCandidates = PoseSearchMode == EPoseSearchMode::BruteForce;
 
 		const FSearchFilters SearchFilters(Schema, NonSelectableIdx, FSelectableAssetIdx(), SearchIndex.bAnyBlockTransition);
@@ -1741,7 +1812,7 @@ UE::PoseSearch::FSearchResult UPoseSearchDatabase::SearchBruteForce(UE::PoseSear
 				check(IsAligned(ReconstructedPoseValuesBuffer.GetData(), alignof(VectorRegister4Float)));
 				for (int32 PoseIdx = 0; PoseIdx < SearchIndex.GetNumPoses(); ++PoseIdx)
 				{
-					EvaluatePoseKernel<true, false>(Result, SearchIndex, QueryValues, ReconstructedPoseValuesBuffer, PoseIdx, SearchFilters, SearchContext, this, bUpdateBestCandidates
+					EvaluatePoseKernel<true, false>(Result, SearchIndex, QueryValues, ReconstructedPoseValuesBuffer, PoseIdx, SearchFilters, SearchContext, this, DynamicWeightsSqrt, bUpdateBestCandidates
 #if WITH_EDITOR && ENABLE_ANIM_DEBUG
 						, PoseIdx
 #endif // WITH_EDITOR && ENABLE_ANIM_DEBUG
@@ -1753,7 +1824,7 @@ UE::PoseSearch::FSearchResult UPoseSearchDatabase::SearchBruteForce(UE::PoseSear
 			{
 				for (int32 PoseIdx = 0; PoseIdx < SearchIndex.GetNumPoses(); ++PoseIdx)
 				{
-					EvaluatePoseKernel<false, true>(Result, SearchIndex, QueryValues, TArrayView<float>(), PoseIdx, SearchFilters, SearchContext, this, bUpdateBestCandidates
+					EvaluatePoseKernel<false, true>(Result, SearchIndex, QueryValues, TArrayView<float>(), PoseIdx, SearchFilters, SearchContext, this, DynamicWeightsSqrt, bUpdateBestCandidates
 #if WITH_EDITOR && ENABLE_ANIM_DEBUG
 						, PoseIdx
 #endif // WITH_EDITOR && ENABLE_ANIM_DEBUG
@@ -1765,7 +1836,7 @@ UE::PoseSearch::FSearchResult UPoseSearchDatabase::SearchBruteForce(UE::PoseSear
 			{
 				for (int32 PoseIdx = 0; PoseIdx < SearchIndex.GetNumPoses(); ++PoseIdx)
 				{
-					EvaluatePoseKernel<false, false>(Result, SearchIndex, QueryValues, TArrayView<float>(), PoseIdx, SearchFilters, SearchContext, this, bUpdateBestCandidates
+					EvaluatePoseKernel<false, false>(Result, SearchIndex, QueryValues, TArrayView<float>(), PoseIdx, SearchFilters, SearchContext, this, DynamicWeightsSqrt, bUpdateBestCandidates
 #if WITH_EDITOR && ENABLE_ANIM_DEBUG
 						, PoseIdx
 #endif // WITH_EDITOR && ENABLE_ANIM_DEBUG
@@ -1791,7 +1862,7 @@ UE::PoseSearch::FSearchResult UPoseSearchDatabase::SearchBruteForce(UE::PoseSear
 					const int32 LastPoseIdx = FirstPoseIdx + SearchIndexAsset.GetNumPoses();
 					for (int32 PoseIdx = FirstPoseIdx; PoseIdx < LastPoseIdx; ++PoseIdx)
 					{
-						EvaluatePoseKernel<true, false>(Result, SearchIndex, QueryValues, ReconstructedPoseValuesBuffer, PoseIdx, SearchFilters, SearchContext, this, bUpdateBestCandidates
+						EvaluatePoseKernel<true, false>(Result, SearchIndex, QueryValues, ReconstructedPoseValuesBuffer, PoseIdx, SearchFilters, SearchContext, this, DynamicWeightsSqrt, bUpdateBestCandidates
 #if WITH_EDITOR && ENABLE_ANIM_DEBUG
 							, ++ResultIndex
 #endif // WITH_EDITOR && ENABLE_ANIM_DEBUG
@@ -1809,7 +1880,7 @@ UE::PoseSearch::FSearchResult UPoseSearchDatabase::SearchBruteForce(UE::PoseSear
 					const int32 LastPoseIdx = FirstPoseIdx + SearchIndexAsset.GetNumPoses();
 					for (int32 PoseIdx = FirstPoseIdx; PoseIdx < LastPoseIdx; ++PoseIdx)
 					{
-						EvaluatePoseKernel<false, true>(Result, SearchIndex, QueryValues, TArrayView<float>(), PoseIdx, SearchFilters, SearchContext, this, bUpdateBestCandidates
+						EvaluatePoseKernel<false, true>(Result, SearchIndex, QueryValues, TArrayView<float>(), PoseIdx, SearchFilters, SearchContext, this, DynamicWeightsSqrt, bUpdateBestCandidates
 #if WITH_EDITOR && ENABLE_ANIM_DEBUG
 							, ++ResultIndex
 #endif // WITH_EDITOR && ENABLE_ANIM_DEBUG
@@ -1827,7 +1898,7 @@ UE::PoseSearch::FSearchResult UPoseSearchDatabase::SearchBruteForce(UE::PoseSear
 					const int32 LastPoseIdx = FirstPoseIdx + SearchIndexAsset.GetNumPoses();
 					for (int32 PoseIdx = FirstPoseIdx; PoseIdx < LastPoseIdx; ++PoseIdx)
 					{
-						EvaluatePoseKernel<false, false>(Result, SearchIndex, QueryValues, TArrayView<float>(), PoseIdx, SearchFilters, SearchContext, this, bUpdateBestCandidates
+						EvaluatePoseKernel<false, false>(Result, SearchIndex, QueryValues, TArrayView<float>(), PoseIdx, SearchFilters, SearchContext, this, DynamicWeightsSqrt, bUpdateBestCandidates
 #if WITH_EDITOR && ENABLE_ANIM_DEBUG
 							, ++ResultIndex
 #endif // WITH_EDITOR && ENABLE_ANIM_DEBUG
@@ -1840,10 +1911,15 @@ UE::PoseSearch::FSearchResult UPoseSearchDatabase::SearchBruteForce(UE::PoseSear
 	else
 	{
 #if UE_POSE_SEARCH_TRACE_ENABLED
+		// @todo: perhaps store the DynamicWeightsSqrt into the GetOrBuildQuery to share between the continuing pose and the full search
+		const int32 NumDimensions = Schema->SchemaCardinality;
+		TArrayView<float> DynamicWeightsSqrtBuffer((float*)FMemory_Alloca(NumDimensions * sizeof(float)), NumDimensions);
+		const TConstArrayView<float> DynamicWeightsSqrt = CalculateDynamicWeightsSqrt(DynamicWeightsSqrtBuffer);
+
 		// calling just for reporting non selectable poses
 		TConstArrayView<float> QueryValues = SearchContext.GetOrBuildQuery(Schema);
 		FNonSelectableIdx NonSelectableIdx;
-		PopulateNonSelectableIdx(NonSelectableIdx, SearchContext, QueryValues);
+		PopulateNonSelectableIdx(NonSelectableIdx, SearchContext, QueryValues, DynamicWeightsSqrt);
 #endif // UE_POSE_SEARCH_TRACE_ENABLED
 	}
 
diff --git a/Engine/Plugins/Animation/PoseSearch/Source/Runtime/Private/PoseSearchFeatureChannel_FilterCrashingLegs.cpp b/Engine/Plugins/Animation/PoseSearch/Source/Runtime/Private/PoseSearchFeatureChannel_FilterCrashingLegs.cpp
--- a/Engine/Plugins/Animation/PoseSearch/Source/Runtime/Private/PoseSearchFeatureChannel_FilterCrashingLegs.cpp
+++ b/Engine/Plugins/Animation/PoseSearch/Source/Runtime/Private/PoseSearchFeatureChannel_FilterCrashingLegs.cpp
@@ -124,35 +124,38 @@ void UPoseSearchFeatureChannel_FilterCrashingLegs::DebugDraw(const UE::PoseSearc
 {
 	using namespace UE::PoseSearch;
 
-	const float CrashingLegsValue = FFeatureVectorHelper::DecodeFloat(PoseVector, ChannelDataOffset);
-
-	const FVector LeftThighPosition = DrawParams.ExtractPosition(PoseVector, 0.f, LeftThighIdx, SampleRole);
-	const FVector RightThighPosition = DrawParams.ExtractPosition(PoseVector, 0.f, RightThighIdx, SampleRole);
-	const FVector LeftFootPosition = DrawParams.ExtractPosition(PoseVector, 0.f, LeftFootIdx, SampleRole);
-	const FVector RightFootPosition = DrawParams.ExtractPosition(PoseVector, 0.f, RightFootIdx, SampleRole);
-
-	const float FeetDistance = (RightFootPosition - LeftFootPosition).Length();
-
-	FVector ThighsSide;
-	FVector ThighsForward;
-	ComputeThighsSideAndForward(RightThighPosition, LeftThighPosition, ThighsSide, ThighsForward);
-	const FVector CrossingLegsVector = (FMath::Cos(CrashingLegsValue * UE_PI) * ThighsSide + FMath::Sin(CrashingLegsValue * UE_PI) * ThighsForward) * FeetDistance;
-
-	auto LerpColor = [](FColor A, FColor B, float T) -> FColor
+	if (DrawParams.IsAnyWeightRelevant(this))
 	{
-		return FColor(
-			FMath::RoundToInt(float(A.R) * (1.f - T) + float(B.R) * T),
-			FMath::RoundToInt(float(A.G) * (1.f - T) + float(B.G) * T),
-			FMath::RoundToInt(float(A.B) * (1.f - T) + float(B.B) * T),
-			FMath::RoundToInt(float(A.A) * (1.f - T) + float(B.A) * T));
-	};
-
-	// when CrossingLegsValue is greater than .5 or lesser than -.5 we draw in Red or Orange
-	const float LerpValue = FMath::Min(CrashingLegsValue * 2.f, 1.f);
-	const FColor Color = LerpValue >= 0.f ? LerpColor(FColor::Green, FColor::Red, LerpValue) : LerpColor(FColor::Green, FColor::Orange, -LerpValue);
-
-	DrawParams.DrawLine(LeftFootPosition, LeftFootPosition + CrossingLegsVector, Color);
-	DrawParams.DrawLine(RightFootPosition, RightFootPosition - CrossingLegsVector, Color);
+		const float CrashingLegsValue = FFeatureVectorHelper::DecodeFloat(PoseVector, ChannelDataOffset);
+
+		const FVector LeftThighPosition = DrawParams.ExtractPosition(PoseVector, 0.f, LeftThighIdx, SampleRole);
+		const FVector RightThighPosition = DrawParams.ExtractPosition(PoseVector, 0.f, RightThighIdx, SampleRole);
+		const FVector LeftFootPosition = DrawParams.ExtractPosition(PoseVector, 0.f, LeftFootIdx, SampleRole);
+		const FVector RightFootPosition = DrawParams.ExtractPosition(PoseVector, 0.f, RightFootIdx, SampleRole);
+
+		const float FeetDistance = (RightFootPosition - LeftFootPosition).Length();
+
+		FVector ThighsSide;
+		FVector ThighsForward;
+		ComputeThighsSideAndForward(RightThighPosition, LeftThighPosition, ThighsSide, ThighsForward);
+		const FVector CrossingLegsVector = (FMath::Cos(CrashingLegsValue * UE_PI) * ThighsSide + FMath::Sin(CrashingLegsValue * UE_PI) * ThighsForward) * FeetDistance;
+
+		auto LerpColor = [](FColor A, FColor B, float T) -> FColor
+			{
+				return FColor(
+					FMath::RoundToInt(float(A.R) * (1.f - T) + float(B.R) * T),
+					FMath::RoundToInt(float(A.G) * (1.f - T) + float(B.G) * T),
+					FMath::RoundToInt(float(A.B) * (1.f - T) + float(B.B) * T),
+					FMath::RoundToInt(float(A.A) * (1.f - T) + float(B.A) * T));
+			};
+
+		// when CrossingLegsValue is greater than .5 or lesser than -.5 we draw in Red or Orange
+		const float LerpValue = FMath::Min(CrashingLegsValue * 2.f, 1.f);
+		const FColor Color = LerpValue >= 0.f ? LerpColor(FColor::Green, FColor::Red, LerpValue) : LerpColor(FColor::Green, FColor::Orange, -LerpValue);
+
+		DrawParams.DrawLine(LeftFootPosition, LeftFootPosition + CrossingLegsVector, Color);
+		DrawParams.DrawLine(RightFootPosition, RightFootPosition - CrossingLegsVector, Color);
+	}
 }
 #endif // ENABLE_DRAW_DEBUG
 
diff --git a/Engine/Plugins/Animation/PoseSearch/Source/Runtime/Private/PoseSearchFeatureChannel_Group.cpp b/Engine/Plugins/Animation/PoseSearch/Source/Runtime/Private/PoseSearchFeatureChannel_Group.cpp
--- a/Engine/Plugins/Animation/PoseSearch/Source/Runtime/Private/PoseSearchFeatureChannel_Group.cpp
+++ b/Engine/Plugins/Animation/PoseSearch/Source/Runtime/Private/PoseSearchFeatureChannel_Group.cpp
@@ -122,6 +122,10 @@ UE::PoseSearch::TLabelBuilder& UPoseSearchFeatureChannel_Group::GetLabel(UE::Pos
 	GetOuterLabel(LabelBuilder, LabelFormat);
 	AppendLabelSeparator(LabelBuilder, UE::PoseSearch::ELabelFormat::Full_Horizontal);
 	LabelBuilder.Append(TEXT("Group"));
+	if (DebugWeightGroupID != INDEX_NONE)
+	{
+		LabelBuilder.Appendf(TEXT("_%d"), DebugWeightGroupID);
+	}
 	return LabelBuilder;
 }
 #endif // WITH_EDITOR
diff --git a/Engine/Plugins/Animation/PoseSearch/Source/Runtime/Private/PoseSearchFeatureChannel_Group.h b/Engine/Plugins/Animation/PoseSearch/Source/Runtime/Private/PoseSearchFeatureChannel_Group.h
--- a/Engine/Plugins/Animation/PoseSearch/Source/Runtime/Private/PoseSearchFeatureChannel_Group.h
+++ b/Engine/Plugins/Animation/PoseSearch/Source/Runtime/Private/PoseSearchFeatureChannel_Group.h
@@ -16,6 +16,12 @@ public:
 	UPROPERTY(EditAnywhere, Category = "Settings")
 	FName SampleRole = UE::PoseSearch::DefaultRole;
 
+	// Experimental, this feature might be removed without warning, not for production use
+	// identifier / metadata to categorize this group base channel weights to be able to zero them via 
+	// UPoseSearchDatabase::CalculateDynamicWeightsSqrt
+	UPROPERTY(EditAnywhere, Category = "Settings")
+	int32 DebugWeightGroupID = INDEX_NONE;
+
 	// UPoseSearchFeatureChannel interface
 	virtual bool Finalize(UPoseSearchSchema* Schema) override;
 	virtual void BuildQuery(UE::PoseSearch::FSearchContext& SearchContext) const override;
diff --git a/Engine/Plugins/Animation/PoseSearch/Source/Runtime/Private/PoseSearchFeatureChannel_Heading.cpp b/Engine/Plugins/Animation/PoseSearch/Source/Runtime/Private/PoseSearchFeatureChannel_Heading.cpp
--- a/Engine/Plugins/Animation/PoseSearch/Source/Runtime/Private/PoseSearchFeatureChannel_Heading.cpp
+++ b/Engine/Plugins/Animation/PoseSearch/Source/Runtime/Private/PoseSearchFeatureChannel_Heading.cpp
@@ -188,37 +188,36 @@ void UPoseSearchFeatureChannel_Heading::DebugDraw(const UE::PoseSearch::FDebugDr
 {
 	using namespace UE::PoseSearch;
 
-#if WITH_EDITORONLY_DATA
 	bool bDrawInjectAdditionalDebugChannels = false;
+#if WITH_EDITORONLY_DATA
 	if (const UPoseSearchSchema* Schema = GetSchema())
 	{
 		bDrawInjectAdditionalDebugChannels = Schema->bDrawInjectAdditionalDebugChannels;
 	}
-	if (!bDrawInjectAdditionalDebugChannels && Weight <= 0.f)
-	{
-		return;
-	}
 #endif // WITH_EDITORONLY_DATA
 
-	FColor Color;
+	if (bDrawInjectAdditionalDebugChannels || DrawParams.IsAnyWeightRelevant(this))
+	{
+		FColor Color;
 #if WITH_EDITORONLY_DATA
-	Color = DebugColor.ToFColor(true);
+		Color = DebugColor.ToFColor(true);
 #else // WITH_EDITORONLY_DATA
-	Color = FLinearColor::White.ToFColor(true);
+		Color = FLinearColor::White.ToFColor(true);
 #endif // WITH_EDITORONLY_DATA
 
-	float PermutationSampleTimeOffset = 0.f;
-	float PermutationOriginTimeOffset = 0.f;
-	UPoseSearchFeatureChannel::GetPermutationTimeOffsets(PermutationTimeType, DrawParams.ExtractPermutationTime(PoseVector), PermutationSampleTimeOffset, PermutationOriginTimeOffset);
-	const EPermutationTimeType SamplePermutationTimeType = PermutationTimeType != EPermutationTimeType::UseSampleTime ? EPermutationTimeType::UseSampleToPermutationTime : EPermutationTimeType::UseSampleTime;
-	const EPermutationTimeType OriginPermutationTimeType = PermutationTimeType == EPermutationTimeType::UsePermutationTime ? EPermutationTimeType::UseSampleToPermutationTime : EPermutationTimeType::UseSampleTime;
+		float PermutationSampleTimeOffset = 0.f;
+		float PermutationOriginTimeOffset = 0.f;
+		UPoseSearchFeatureChannel::GetPermutationTimeOffsets(PermutationTimeType, DrawParams.ExtractPermutationTime(PoseVector), PermutationSampleTimeOffset, PermutationOriginTimeOffset);
+		const EPermutationTimeType SamplePermutationTimeType = PermutationTimeType != EPermutationTimeType::UseSampleTime ? EPermutationTimeType::UseSampleToPermutationTime : EPermutationTimeType::UseSampleTime;
+		const EPermutationTimeType OriginPermutationTimeType = PermutationTimeType == EPermutationTimeType::UsePermutationTime ? EPermutationTimeType::UseSampleToPermutationTime : EPermutationTimeType::UseSampleTime;
 
-	const FVector FeaturesVector = FFeatureVectorHelper::DecodeVector(PoseVector, ChannelDataOffset, ComponentStripping);
-	const FVector BoneHeading = DrawParams.ExtractRotation(PoseVector, OriginTimeOffset, SchemaOriginBoneIdx, OriginRole, OriginPermutationTimeType, INDEX_NONE, PermutationOriginTimeOffset).RotateVector(FeaturesVector);
-	const FVector BonePos = DrawParams.ExtractPosition(PoseVector, SampleTimeOffset, SchemaBoneIdx, SampleRole, SamplePermutationTimeType, SamplingAttributeId, PermutationSampleTimeOffset);
+		const FVector FeaturesVector = FFeatureVectorHelper::DecodeVector(PoseVector, ChannelDataOffset, ComponentStripping);
+		const FVector BoneHeading = DrawParams.ExtractRotation(PoseVector, OriginTimeOffset, SchemaOriginBoneIdx, OriginRole, OriginPermutationTimeType, INDEX_NONE, PermutationOriginTimeOffset).RotateVector(FeaturesVector);
+		const FVector BonePos = DrawParams.ExtractPosition(PoseVector, SampleTimeOffset, SchemaBoneIdx, SampleRole, SamplePermutationTimeType, SamplingAttributeId, PermutationSampleTimeOffset);
 
-	DrawParams.DrawPoint(BonePos, Color, 3.f);
-	DrawParams.DrawLine(BonePos + BoneHeading * 4.f, BonePos + BoneHeading * 15.f, Color);
+		DrawParams.DrawPoint(BonePos, Color, 3.f);
+		DrawParams.DrawLine(BonePos + BoneHeading * 4.f, BonePos + BoneHeading * 15.f, Color);
+	}
 }
 #endif // ENABLE_DRAW_DEBUG
 
@@ -335,9 +334,16 @@ USkeleton* UPoseSearchFeatureChannel_Heading::GetSkeleton(bool& bInvalidSkeleton
 	if (const UPoseSearchSchema* Schema = GetSchema())
 	{
 		bInvalidSkeletonIsError = false;
-		if (PropertyHandle && PropertyHandle->GetProperty()->GetFName() == GET_MEMBER_NAME_CHECKED(UPoseSearchFeatureChannel_Heading, OriginBone))
+		if (PropertyHandle)
 		{
-			return Schema->GetSkeleton(OriginRole);
+			if (PropertyHandle->GetProperty()->GetFName() == GET_MEMBER_NAME_CHECKED(UPoseSearchFeatureChannel_Heading, Bone))
+			{
+				return Schema->GetSkeleton(SampleRole);
+			}
+			if (PropertyHandle->GetProperty()->GetFName() == GET_MEMBER_NAME_CHECKED(UPoseSearchFeatureChannel_Heading, OriginBone))
+			{
+				return Schema->GetSkeleton(OriginRole);
+			}
 		}
 	}
 
diff --git a/Engine/Plugins/Animation/PoseSearch/Source/Runtime/Private/PoseSearchFeatureChannel_Phase.cpp b/Engine/Plugins/Animation/PoseSearch/Source/Runtime/Private/PoseSearchFeatureChannel_Phase.cpp
--- a/Engine/Plugins/Animation/PoseSearch/Source/Runtime/Private/PoseSearchFeatureChannel_Phase.cpp
+++ b/Engine/Plugins/Animation/PoseSearch/Source/Runtime/Private/PoseSearchFeatureChannel_Phase.cpp
@@ -360,29 +360,32 @@ void UPoseSearchFeatureChannel_Phase::DebugDraw(const UE::PoseSearch::FDebugDraw
 	using namespace UE::PoseSearch;
 
 	static float ScaleFactor = 1.f;
+	static int32 Segments = 32;
 
-	FColor Color;
+	if (DrawParams.IsAnyWeightRelevant(this))
+	{
+		FColor Color;
 #if WITH_EDITORONLY_DATA
-	Color = DebugColor.ToFColor(true);
+		Color = DebugColor.ToFColor(true);
 #else // WITH_EDITORONLY_DATA
-	Color = FLinearColor::Yellow.ToFColor(true);
+		Color = FLinearColor::Yellow.ToFColor(true);
 #endif // WITH_EDITORONLY_DATA
 
-	const FVector2D Phase = FFeatureVectorHelper::DecodeVector2D(PoseVector, ChannelDataOffset);
-	const FVector BonePos = DrawParams.ExtractPosition(PoseVector, 0.f, SchemaBoneIdx, SampleRole);
+		const FVector2D Phase = FFeatureVectorHelper::DecodeVector2D(PoseVector, ChannelDataOffset);
+		const FVector BonePos = DrawParams.ExtractPosition(PoseVector, 0.f, SchemaBoneIdx, SampleRole);
 
-	const FTransform RootBoneTransform = DrawParams.GetRootBoneTransform(SampleRole);
-	const FVector TransformXAxisVector = RootBoneTransform.TransformVector(FVector::XAxisVector);
-	const FVector TransformYAxisVector = RootBoneTransform.TransformVector(FVector::YAxisVector);
-	const FVector TransformZAxisVector = RootBoneTransform.TransformVector(FVector::ZAxisVector);
+		const FTransform RootBoneTransform = DrawParams.GetRootBoneTransform(SampleRole);
+		const FVector TransformXAxisVector = RootBoneTransform.TransformVector(FVector::XAxisVector);
+		const FVector TransformYAxisVector = RootBoneTransform.TransformVector(FVector::YAxisVector);
+		const FVector TransformZAxisVector = RootBoneTransform.TransformVector(FVector::ZAxisVector);
 
-	const FVector PhaseVector = (TransformZAxisVector * Phase.X + TransformYAxisVector * Phase.Y) * ScaleFactor;
-	DrawParams.DrawLine(BonePos, BonePos + PhaseVector, Color);
+		const FVector PhaseVector = (TransformZAxisVector * Phase.X + TransformYAxisVector * Phase.Y) * ScaleFactor;
+		DrawParams.DrawLine(BonePos, BonePos + PhaseVector, Color);
 
-	static int32 Segments = 32;
-	FMatrix CircleTransform;
-	CircleTransform.SetAxes(&TransformXAxisVector, &TransformYAxisVector, &TransformZAxisVector, &BonePos);
-	DrawParams.DrawCircle(CircleTransform, PhaseVector.Length(), Segments, Color);
+		FMatrix CircleTransform;
+		CircleTransform.SetAxes(&TransformXAxisVector, &TransformYAxisVector, &TransformZAxisVector, &BonePos);
+		DrawParams.DrawCircle(CircleTransform, PhaseVector.Length(), Segments, Color);
+	}
 }
 #endif // ENABLE_DRAW_DEBUG
 
diff --git a/Engine/Plugins/Animation/PoseSearch/Source/Runtime/Private/PoseSearchFeatureChannel_Pose.cpp b/Engine/Plugins/Animation/PoseSearch/Source/Runtime/Private/PoseSearchFeatureChannel_Pose.cpp
--- a/Engine/Plugins/Animation/PoseSearch/Source/Runtime/Private/PoseSearchFeatureChannel_Pose.cpp
+++ b/Engine/Plugins/Animation/PoseSearch/Source/Runtime/Private/PoseSearchFeatureChannel_Pose.cpp
@@ -126,6 +126,10 @@ UE::PoseSearch::TLabelBuilder& UPoseSearchFeatureChannel_Pose::GetLabel(UE::Pose
 	GetOuterLabel(LabelBuilder, LabelFormat);
 	AppendLabelSeparator(LabelBuilder, UE::PoseSearch::ELabelFormat::Full_Horizontal);
 	LabelBuilder.Append(TEXT("Pose"));
+	if (DebugWeightGroupID != INDEX_NONE)
+	{
+		LabelBuilder.Appendf(TEXT("_%d"), DebugWeightGroupID);
+	}
 	return LabelBuilder;
 }
 #endif // WITH_EDITOR
diff --git a/Engine/Plugins/Animation/PoseSearch/Source/Runtime/Private/PoseSearchFeatureChannel_Position.cpp b/Engine/Plugins/Animation/PoseSearch/Source/Runtime/Private/PoseSearchFeatureChannel_Position.cpp
--- a/Engine/Plugins/Animation/PoseSearch/Source/Runtime/Private/PoseSearchFeatureChannel_Position.cpp
+++ b/Engine/Plugins/Animation/PoseSearch/Source/Runtime/Private/PoseSearchFeatureChannel_Position.cpp
@@ -181,42 +181,41 @@ void UPoseSearchFeatureChannel_Position::DebugDraw(const UE::PoseSearch::FDebugD
 {
 	using namespace UE::PoseSearch;
 
-#if WITH_EDITORONLY_DATA
 	bool bDrawInjectAdditionalDebugChannels = false;
+#if WITH_EDITORONLY_DATA
 	if (const UPoseSearchSchema* Schema = GetSchema())
 	{
 		bDrawInjectAdditionalDebugChannels = Schema->bDrawInjectAdditionalDebugChannels;
 	}
-	if (!bDrawInjectAdditionalDebugChannels && Weight <= 0.f)
-	{
-		return;
-	}
 #endif // WITH_EDITORONLY_DATA
 
-	FColor Color;
+	if (bDrawInjectAdditionalDebugChannels || DrawParams.IsAnyWeightRelevant(this))
+	{
+		FColor Color;
 #if WITH_EDITORONLY_DATA
-	Color = DebugColor.ToFColor(true);
+		Color = DebugColor.ToFColor(true);
 #else // WITH_EDITORONLY_DATA
-	Color = FLinearColor::Blue.ToFColor(true);
+		Color = FLinearColor::Blue.ToFColor(true);
 #endif // WITH_EDITORONLY_DATA
 
-	float PermutationSampleTimeOffset = 0.f;
-	float PermutationOriginTimeOffset = 0.f;
-	UPoseSearchFeatureChannel::GetPermutationTimeOffsets(PermutationTimeType, DrawParams.ExtractPermutationTime(PoseVector), PermutationSampleTimeOffset, PermutationOriginTimeOffset);
-	const EPermutationTimeType OriginPermutationTimeType = PermutationTimeType == EPermutationTimeType::UsePermutationTime ? EPermutationTimeType::UseSampleToPermutationTime : EPermutationTimeType::UseSampleTime;
+		float PermutationSampleTimeOffset = 0.f;
+		float PermutationOriginTimeOffset = 0.f;
+		UPoseSearchFeatureChannel::GetPermutationTimeOffsets(PermutationTimeType, DrawParams.ExtractPermutationTime(PoseVector), PermutationSampleTimeOffset, PermutationOriginTimeOffset);
+		const EPermutationTimeType OriginPermutationTimeType = PermutationTimeType == EPermutationTimeType::UsePermutationTime ? EPermutationTimeType::UseSampleToPermutationTime : EPermutationTimeType::UseSampleTime;
 
-	const FVector FeaturesVector = FFeatureVectorHelper::DecodeVector(PoseVector, ChannelDataOffset, ComponentStripping);
-	const FVector OriginBonePos = DrawParams.ExtractPosition(PoseVector, OriginTimeOffset, SchemaOriginBoneIdx, OriginRole, OriginPermutationTimeType, INDEX_NONE, PermutationOriginTimeOffset);
-	const FVector DeltaPos = DrawParams.ExtractRotation(PoseVector, OriginTimeOffset, RootSchemaBoneIdx, OriginRole, OriginPermutationTimeType, INDEX_NONE, PermutationOriginTimeOffset).RotateVector(FeaturesVector);
-	const FVector BonePos = OriginBonePos + DeltaPos;
-	DrawParams.DrawPoint(BonePos, Color);
+		const FVector FeaturesVector = FFeatureVectorHelper::DecodeVector(PoseVector, ChannelDataOffset, ComponentStripping);
+		const FVector OriginBonePos = DrawParams.ExtractPosition(PoseVector, OriginTimeOffset, SchemaOriginBoneIdx, OriginRole, OriginPermutationTimeType, INDEX_NONE, PermutationOriginTimeOffset);
+		const FVector DeltaPos = DrawParams.ExtractRotation(PoseVector, OriginTimeOffset, RootSchemaBoneIdx, OriginRole, OriginPermutationTimeType, INDEX_NONE, PermutationOriginTimeOffset).RotateVector(FeaturesVector);
+		const FVector BonePos = OriginBonePos + DeltaPos;
+		DrawParams.DrawPoint(BonePos, Color);
 
-	const bool bDrawOrigin = !DeltaPos.IsNearlyZero() && (SchemaOriginBoneIdx != RootSchemaBoneIdx || !FMath::IsNearlyZero(OriginTimeOffset) ||
-							 SampleRole != OriginRole || PermutationTimeType != EPermutationTimeType::UseSampleTime || bUseBlueprintQueryOverride);
-	if (bDrawOrigin)
-	{
-		DrawParams.DrawPoint(OriginBonePos, Color);
-		DrawParams.DrawLine(OriginBonePos, BonePos, Color);
+		const bool bDrawOrigin = !DeltaPos.IsNearlyZero() && (SchemaOriginBoneIdx != RootSchemaBoneIdx || !FMath::IsNearlyZero(OriginTimeOffset) ||
+			SampleRole != OriginRole || PermutationTimeType != EPermutationTimeType::UseSampleTime || bUseBlueprintQueryOverride);
+		if (bDrawOrigin)
+		{
+			DrawParams.DrawPoint(OriginBonePos, Color);
+			DrawParams.DrawLine(OriginBonePos, BonePos, Color);
+		}
 	}
 }
 #endif // ENABLE_DRAW_DEBUG
@@ -322,9 +321,16 @@ USkeleton* UPoseSearchFeatureChannel_Position::GetSkeleton(bool& bInvalidSkeleto
 	if (const UPoseSearchSchema* Schema = GetSchema())
 	{
 		bInvalidSkeletonIsError = false;
-		if (PropertyHandle && PropertyHandle->GetProperty()->GetFName() == GET_MEMBER_NAME_CHECKED(UPoseSearchFeatureChannel_Position, OriginBone))
+		if (PropertyHandle)
 		{
-			return Schema->GetSkeleton(OriginRole);
+			if (PropertyHandle->GetProperty()->GetFName() == GET_MEMBER_NAME_CHECKED(UPoseSearchFeatureChannel_Position, Bone))
+			{
+				return Schema->GetSkeleton(SampleRole);
+			}
+			if (PropertyHandle->GetProperty()->GetFName() == GET_MEMBER_NAME_CHECKED(UPoseSearchFeatureChannel_Position, OriginBone))
+			{
+				return Schema->GetSkeleton(OriginRole);
+			}
 		}
 	}
 
diff --git a/Engine/Plugins/Animation/PoseSearch/Source/Runtime/Private/PoseSearchFeatureChannel_Trajectory.cpp b/Engine/Plugins/Animation/PoseSearch/Source/Runtime/Private/PoseSearchFeatureChannel_Trajectory.cpp
--- a/Engine/Plugins/Animation/PoseSearch/Source/Runtime/Private/PoseSearchFeatureChannel_Trajectory.cpp
+++ b/Engine/Plugins/Animation/PoseSearch/Source/Runtime/Private/PoseSearchFeatureChannel_Trajectory.cpp
@@ -139,53 +139,56 @@ void UPoseSearchFeatureChannel_Trajectory::DebugDraw(const UE::PoseSearch::FDebu
 {
 	using namespace UE::PoseSearch;
 
-	TArray<const UPoseSearchFeatureChannel_Position*, TInlineAllocator<32>> Positions;
-	for (const TObjectPtr<UPoseSearchFeatureChannel>& SubChannelPtr : GetSubChannels())
+	if (DrawParams.IsAnyWeightRelevant(this))
 	{
-		if (const UPoseSearchFeatureChannel_Position* Position = Cast<UPoseSearchFeatureChannel_Position>(SubChannelPtr.Get()))
+		TArray<const UPoseSearchFeatureChannel_Position*, TInlineAllocator<32>> Positions;
+		for (const TObjectPtr<UPoseSearchFeatureChannel>& SubChannelPtr : GetSubChannels())
 		{
-			Positions.Add(Position);
+			if (const UPoseSearchFeatureChannel_Position* Position = Cast<UPoseSearchFeatureChannel_Position>(SubChannelPtr.Get()))
+			{
+				Positions.Add(Position);
+			}
 		}
-	}
 
-	if (!Positions.IsEmpty())
-	{
-		Positions.Sort([](const UPoseSearchFeatureChannel_Position& a, const UPoseSearchFeatureChannel_Position& b)
-			{
-				return a.SampleTimeOffset < b.SampleTimeOffset;
-			});
-
-		// big enough negative number that prevents PrevTimeOffset * CurrTimeOffset being infinite (there will never be UPoseSearchFeatureChannel_Trajectory trying to match 1000 seconds in the past)
-		float PrevTimeOffset = -1000.f;
-		TArray<FVector, TInlineAllocator<32>> TrajSplinePos;
-		TArray<FColor, TInlineAllocator<32>> TrajSplineColor;
-		for (int32 i = 0; i < Positions.Num(); ++i)
+		if (!Positions.IsEmpty())
 		{
-			const float CurrTimeOffset = Positions[i]->SampleTimeOffset;
-			FColor Color;
+			Positions.Sort([](const UPoseSearchFeatureChannel_Position& a, const UPoseSearchFeatureChannel_Position& b)
+				{
+					return a.SampleTimeOffset < b.SampleTimeOffset;
+				});
+
+			// big enough negative number that prevents PrevTimeOffset * CurrTimeOffset being infinite (there will never be UPoseSearchFeatureChannel_Trajectory trying to match 1000 seconds in the past)
+			float PrevTimeOffset = -1000.f;
+			TArray<FVector, TInlineAllocator<32>> TrajSplinePos;
+			TArray<FColor, TInlineAllocator<32>> TrajSplineColor;
+			for (int32 i = 0; i < Positions.Num(); ++i)
+			{
+				const float CurrTimeOffset = Positions[i]->SampleTimeOffset;
+				FColor Color;
 #if WITH_EDITORONLY_DATA
-			Color = Positions[i]->DebugColor.ToFColor(true);
+				Color = Positions[i]->DebugColor.ToFColor(true);
 #else // WITH_EDITORONLY_DATA
-			Color = FLinearColor::Blue.ToFColor(true);
+				Color = FLinearColor::Blue.ToFColor(true);
 #endif // WITH_EDITORONLY_DATA
 
-			if (Positions.Num() == 1 || PrevTimeOffset * CurrTimeOffset < UE_KINDA_SMALL_NUMBER)
-			{
-				// we jumped from negative to positive time offset without having a zero time offset, or we only have one position point. so we add the zero
-				TrajSplinePos.Add(DrawParams.ExtractPosition(PoseVector, 0.f, RootSchemaBoneIdx, Positions[i]->OriginRole));
+				if (Positions.Num() == 1 || PrevTimeOffset * CurrTimeOffset < UE_KINDA_SMALL_NUMBER)
+				{
+					// we jumped from negative to positive time offset without having a zero time offset, or we only have one position point. so we add the zero
+					TrajSplinePos.Add(DrawParams.ExtractPosition(PoseVector, 0.f, RootSchemaBoneIdx, Positions[i]->OriginRole));
+					TrajSplineColor.Add(Color);
+				}
+
+				TrajSplinePos.Add(DrawParams.ExtractPosition(PoseVector, CurrTimeOffset, RootSchemaBoneIdx, Positions[i]->OriginRole));
 				TrajSplineColor.Add(Color);
-			}
 
-			TrajSplinePos.Add(DrawParams.ExtractPosition(PoseVector, CurrTimeOffset, RootSchemaBoneIdx, Positions[i]->OriginRole));
-			TrajSplineColor.Add(Color);
+				PrevTimeOffset = CurrTimeOffset;
+			}
 
-			PrevTimeOffset = CurrTimeOffset;
+			DrawParams.DrawCentripetalCatmullRomSpline(TrajSplinePos, TrajSplineColor, 0.5f, 8);
 		}
 
-		DrawParams.DrawCentripetalCatmullRomSpline(TrajSplinePos, TrajSplineColor, 0.5f, 8);
+		Super::DebugDraw(DrawParams, PoseVector);
 	}
-
-	Super::DebugDraw(DrawParams, PoseVector);
 }
 #endif // ENABLE_DRAW_DEBUG
 
@@ -195,6 +198,10 @@ UE::PoseSearch::TLabelBuilder& UPoseSearchFeatureChannel_Trajectory::GetLabel(UE
 	GetOuterLabel(LabelBuilder, LabelFormat);
 	AppendLabelSeparator(LabelBuilder, UE::PoseSearch::ELabelFormat::Full_Horizontal);
 	LabelBuilder.Append(TEXT("Traj"));
+	if (DebugWeightGroupID != INDEX_NONE)
+	{
+		LabelBuilder.Appendf(TEXT("_%d"), DebugWeightGroupID);
+	}
 	return LabelBuilder;
 }
 #endif // WITH_EDITOR
diff --git a/Engine/Plugins/Animation/PoseSearch/Source/Runtime/Private/PoseSearchFeatureChannel_Velocity.cpp b/Engine/Plugins/Animation/PoseSearch/Source/Runtime/Private/PoseSearchFeatureChannel_Velocity.cpp
--- a/Engine/Plugins/Animation/PoseSearch/Source/Runtime/Private/PoseSearchFeatureChannel_Velocity.cpp
+++ b/Engine/Plugins/Animation/PoseSearch/Source/Runtime/Private/PoseSearchFeatureChannel_Velocity.cpp
@@ -143,25 +143,28 @@ void UPoseSearchFeatureChannel_Velocity::DebugDraw(const UE::PoseSearch::FDebugD
 {
 	using namespace UE::PoseSearch;
 
-	FColor Color;
+	if (DrawParams.IsAnyWeightRelevant(this))
+	{
+		FColor Color;
 #if WITH_EDITORONLY_DATA
-	Color = DebugColor.ToFColor(true);
+		Color = DebugColor.ToFColor(true);
 #else // WITH_EDITORONLY_DATA
-	Color = FLinearColor::Green.ToFColor(true);
+		Color = FLinearColor::Green.ToFColor(true);
 #endif // WITH_EDITORONLY_DATA
 
-	const float LinearVelocityScale = bNormalize ? 15.f : 0.08f;
+		const float LinearVelocityScale = bNormalize ? 15.f : 0.08f;
 
-	float PermutationSampleTimeOffset = 0.f;
-	float PermutationOriginTimeOffset = 0.f;
-	UPoseSearchFeatureChannel::GetPermutationTimeOffsets(PermutationTimeType, DrawParams.ExtractPermutationTime(PoseVector), PermutationSampleTimeOffset, PermutationOriginTimeOffset);
-	const EPermutationTimeType SamplePermutationTimeType = PermutationTimeType == EPermutationTimeType::UsePermutationTime ? EPermutationTimeType::UseSampleToPermutationTime : EPermutationTimeType::UseSampleTime;
-	
-	const FVector FeaturesVector = FFeatureVectorHelper::DecodeVector(PoseVector, ChannelDataOffset, ComponentStripping);
-	const FVector LinearVelocity = DrawParams.ExtractRotation(PoseVector, SampleTimeOffset, RootSchemaBoneIdx, SampleRole, SamplePermutationTimeType, SamplingAttributeId, PermutationSampleTimeOffset).RotateVector(FeaturesVector);
-	const FVector BonePos = DrawParams.ExtractPosition(PoseVector, SampleTimeOffset, SchemaBoneIdx, SampleRole, SamplePermutationTimeType, SamplingAttributeId, PermutationSampleTimeOffset);
+		float PermutationSampleTimeOffset = 0.f;
+		float PermutationOriginTimeOffset = 0.f;
+		UPoseSearchFeatureChannel::GetPermutationTimeOffsets(PermutationTimeType, DrawParams.ExtractPermutationTime(PoseVector), PermutationSampleTimeOffset, PermutationOriginTimeOffset);
+		const EPermutationTimeType SamplePermutationTimeType = PermutationTimeType == EPermutationTimeType::UsePermutationTime ? EPermutationTimeType::UseSampleToPermutationTime : EPermutationTimeType::UseSampleTime;
 
-	DrawParams.DrawLine(BonePos, BonePos + LinearVelocity * LinearVelocityScale, Color);
+		const FVector FeaturesVector = FFeatureVectorHelper::DecodeVector(PoseVector, ChannelDataOffset, ComponentStripping);
+		const FVector LinearVelocity = DrawParams.ExtractRotation(PoseVector, SampleTimeOffset, RootSchemaBoneIdx, SampleRole, SamplePermutationTimeType, SamplingAttributeId, PermutationSampleTimeOffset).RotateVector(FeaturesVector);
+		const FVector BonePos = DrawParams.ExtractPosition(PoseVector, SampleTimeOffset, SchemaBoneIdx, SampleRole, SamplePermutationTimeType, SamplingAttributeId, PermutationSampleTimeOffset);
+
+		DrawParams.DrawLine(BonePos, BonePos + LinearVelocity * LinearVelocityScale, Color);
+	}
 }
 #endif // ENABLE_DRAW_DEBUG
 
@@ -274,9 +277,16 @@ USkeleton* UPoseSearchFeatureChannel_Velocity::GetSkeleton(bool& bInvalidSkeleto
 	if (const UPoseSearchSchema* Schema = GetSchema())
 	{
 		bInvalidSkeletonIsError = false;
-		if (PropertyHandle && PropertyHandle->GetProperty()->GetFName() == GET_MEMBER_NAME_CHECKED(UPoseSearchFeatureChannel_Velocity, OriginBone))
+		if (PropertyHandle)
 		{
-			return Schema->GetSkeleton(OriginRole);
+			if (PropertyHandle->GetProperty()->GetFName() == GET_MEMBER_NAME_CHECKED(UPoseSearchFeatureChannel_Velocity, Bone))
+			{
+				return Schema->GetSkeleton(SampleRole);
+			}
+			if (PropertyHandle->GetProperty()->GetFName() == GET_MEMBER_NAME_CHECKED(UPoseSearchFeatureChannel_Velocity, OriginBone))
+			{
+				return Schema->GetSkeleton(OriginRole);
+			}
 		}
 	}
 
diff --git a/Engine/Plugins/Animation/PoseSearch/Source/Runtime/Private/PoseSearchIndex.cpp b/Engine/Plugins/Animation/PoseSearch/Source/Runtime/Private/PoseSearchIndex.cpp
--- a/Engine/Plugins/Animation/PoseSearch/Source/Runtime/Private/PoseSearchIndex.cpp
+++ b/Engine/Plugins/Animation/PoseSearch/Source/Runtime/Private/PoseSearchIndex.cpp
@@ -722,19 +722,29 @@ TConstArrayView<float> FSearchIndex::GetPCAPoseValues(int32 PCAValuesVectorIdx)
 }
 
 FPoseSearchCost FSearchIndex::ComparePoses(int32 PoseIdx, float ContinuingPoseCostBias, TConstArrayView<float> PoseValues, TConstArrayView<float> QueryValues) const
+{
+	return ComparePoses(PoseIdx, ContinuingPoseCostBias, PoseValues, QueryValues, WeightsSqrt);
+}
+
+FPoseSearchCost FSearchIndex::CompareAlignedPoses(int32 PoseIdx, float ContinuingPoseCostBias, TConstArrayView<float> PoseValues, TConstArrayView<float> QueryValues) const
+{
+	return CompareAlignedPoses(PoseIdx, ContinuingPoseCostBias, PoseValues, QueryValues, WeightsSqrt);
+}
+
+FPoseSearchCost FSearchIndex::ComparePoses(int32 PoseIdx, float ContinuingPoseCostBias, TConstArrayView<float> PoseValues, TConstArrayView<float> QueryValues, TConstArrayView<float> DynamicWeightsSqrt) const
 {
 	// base dissimilarity cost representing how the associated PoseIdx differ, in a weighted way, from the query pose (QueryValues)
-	const float DissimilarityCost = CompareFeatureVectors(PoseValues, QueryValues, WeightsSqrt);
+	const float DissimilarityCost = CompareFeatureVectors(PoseValues, QueryValues, DynamicWeightsSqrt);
 
 	// cost addend associated to Schema->BaseCostBias or overriden by UAnimNotifyState_PoseSearchModifyCost
 	const float NotifyAddend = PoseMetadata[PoseIdx].GetCostAddend();
 	return FPoseSearchCost(DissimilarityCost, NotifyAddend, ContinuingPoseCostBias);
 }
 
-FPoseSearchCost FSearchIndex::CompareAlignedPoses(int32 PoseIdx, float ContinuingPoseCostBias, TConstArrayView<float> PoseValues, TConstArrayView<float> QueryValues) const
+FPoseSearchCost FSearchIndex::CompareAlignedPoses(int32 PoseIdx, float ContinuingPoseCostBias, TConstArrayView<float> PoseValues, TConstArrayView<float> QueryValues, TConstArrayView<float> DynamicWeightsSqrt) const
 {
 	// base dissimilarity cost representing how the associated PoseIdx differ, in a weighted way, from the query pose (QueryValues)
-	const float DissimilarityCost = CompareAlignedFeatureVectors(PoseValues, QueryValues, WeightsSqrt);
+	const float DissimilarityCost = CompareAlignedFeatureVectors(PoseValues, QueryValues, DynamicWeightsSqrt);
 
 	// cost addend associated to Schema->BaseCostBias or overriden by UAnimNotifyState_PoseSearchModifyCost
 	const float NotifyAddend = PoseMetadata[PoseIdx].GetCostAddend();
diff --git a/Engine/Plugins/Animation/PoseSearch/Source/Runtime/Private/PoseSearchLibrary.cpp b/Engine/Plugins/Animation/PoseSearch/Source/Runtime/Private/PoseSearchLibrary.cpp
--- a/Engine/Plugins/Animation/PoseSearch/Source/Runtime/Private/PoseSearchLibrary.cpp
+++ b/Engine/Plugins/Animation/PoseSearch/Source/Runtime/Private/PoseSearchLibrary.cpp
@@ -30,8 +30,11 @@
 #define LOCTEXT_NAMESPACE "PoseSearchLibrary"
 
 #if ENABLE_DRAW_DEBUG && ENABLE_ANIM_DEBUG
-TAutoConsoleVariable<bool> CVarAnimMotionMatchDrawQueryEnable(TEXT("a.MotionMatch.DrawQuery.Enable"), false, TEXT("Enable / Disable MotionMatch Draw Query"));
-TAutoConsoleVariable<bool> CVarAnimMotionMatchDrawMatchEnable(TEXT("a.MotionMatch.DrawMatch.Enable"), false, TEXT("Enable / Disable MotionMatch Draw Match"));
+static bool GVarAnimMotionMatchDrawQueryEnable = false;
+static FAutoConsoleVariableRef CVarAnimMotionMatchDrawQueryEnable(TEXT("a.MotionMatch.DrawQuery.Enable"), GVarAnimMotionMatchDrawQueryEnable, TEXT("Enable / Disable MotionMatch Draw Query"));
+
+static bool GVarAnimMotionMatchDrawMatchEnable = false;
+static FAutoConsoleVariableRef CVarAnimMotionMatchDrawMatchEnable(TEXT("a.MotionMatch.DrawMatch.Enable"), GVarAnimMotionMatchDrawMatchEnable, TEXT("Enable / Disable MotionMatch Draw Match"));
 #endif
 
 namespace UE::PoseSearch
@@ -741,15 +744,14 @@ void UPoseSearchLibrary::UpdateMotionMatchingState(
 		if (EAsyncBuildIndexResult::Success == FAsyncPoseSearchDatabasesManagement::RequestAsyncBuildIndex(CurResultDatabase, ERequestAsyncBuildFlag::ContinueRequest))
 #endif // WITH_EDITOR
 		{
+			FDebugDrawParams DrawParams(SearchContext.GetAnimContexts(), SearchContext.GetPoseHistories(), SearchContext.GetRoleToIndex(), CurResultDatabase);
 			if (bDebugDrawCurResult)
 			{
-				FDebugDrawParams DrawParams(SearchContext.GetAnimContexts(), SearchContext.GetPoseHistories(), SearchContext.GetRoleToIndex(), CurResultDatabase);
 				DrawParams.DrawFeatureVector(CurResult.PoseIdx);
 			}
 
 			if (bDebugDrawQuery)
 			{
-				FDebugDrawParams DrawParams(SearchContext.GetAnimContexts(), SearchContext.GetPoseHistories(), SearchContext.GetRoleToIndex(), CurResultDatabase, EDebugDrawFlags::DrawQuery);
 				DrawParams.DrawFeatureVector(SearchContext.GetOrBuildQuery(CurResultDatabase->Schema));
 			}
 		}
@@ -1173,20 +1175,19 @@ UE::PoseSearch::FSearchResult UPoseSearchLibrary::MotionMatch(
 #if ENABLE_DRAW_DEBUG && ENABLE_ANIM_DEBUG
 	if (SearchResult.IsValid())
 	{
-		const bool bDrawMatch = CVarAnimMotionMatchDrawMatchEnable.GetValueOnAnyThread();
-		const bool bDrawquery = CVarAnimMotionMatchDrawQueryEnable.GetValueOnAnyThread();
+		const bool bDrawMatch = GVarAnimMotionMatchDrawMatchEnable;
+		const bool bDrawquery = GVarAnimMotionMatchDrawQueryEnable;
 
 		if (bDrawMatch || bDrawquery)
 		{
+			FDebugDrawParams DrawParams(SearchContext.GetAnimContexts(), SearchContext.GetPoseHistories(), SearchContext.GetRoleToIndex(), SearchResult.Database.Get());
 			if (bDrawMatch)
 			{
-				FDebugDrawParams DrawParams(SearchContext.GetAnimContexts(), SearchContext.GetPoseHistories(), SearchContext.GetRoleToIndex(), SearchResult.Database.Get());
 				DrawParams.DrawFeatureVector(SearchResult.PoseIdx);
 			}
 
 			if (bDrawquery)
 			{
-				FDebugDrawParams DrawParams(SearchContext.GetAnimContexts(), SearchContext.GetPoseHistories(), SearchContext.GetRoleToIndex(), SearchResult.Database.Get(), EDebugDrawFlags::DrawQuery);
 				DrawParams.DrawFeatureVector(SearchContext.GetOrBuildQuery(SearchResult.Database->Schema));
 			}
 		}
diff --git a/Engine/Plugins/Animation/PoseSearch/Source/Runtime/Public/PoseSearch/PoseSearchContext.h b/Engine/Plugins/Animation/PoseSearch/Source/Runtime/Public/PoseSearch/PoseSearchContext.h
--- a/Engine/Plugins/Animation/PoseSearch/Source/Runtime/Public/PoseSearch/PoseSearchContext.h
+++ b/Engine/Plugins/Animation/PoseSearch/Source/Runtime/Public/PoseSearch/PoseSearchContext.h
@@ -19,14 +19,16 @@ namespace UE::PoseSearch
 struct FPoseIndicesHistory;
 struct IPoseHistory;
 
-enum class EDebugDrawFlags : uint32
+enum class UE_DEPRECATED(5.6, "EDebugDrawFlags is no longer used") EDebugDrawFlags : uint32
 {
 	None = 0,
 
 	// used to differenciate channels debug drawing of the query
 	DrawQuery = 1 << 0,
 };
+PRAGMA_DISABLE_DEPRECATION_WARNINGS
 ENUM_CLASS_FLAGS(EDebugDrawFlags);
+PRAGMA_ENABLE_DEPRECATION_WARNINGS
 
 enum class EPoseCandidateFlags : uint32
 {
@@ -53,11 +55,13 @@ ENUM_CLASS_FLAGS(EPoseCandidateFlags);
 #if ENABLE_DRAW_DEBUG
 struct POSESEARCH_API FDebugDrawParams
 {
+	PRAGMA_DISABLE_DEPRECATION_WARNINGS
 	UE_DEPRECATED(5.6, "Use other FDebugDrawParams signatures instead")
 	FDebugDrawParams(TArrayView<FAnimInstanceProxy*> InAnimInstanceProxies, TConstArrayView<const IPoseHistory*> InPoseHistories, const FRoleToIndex& InRoleToIndex, const UPoseSearchDatabase* InDatabase, EDebugDrawFlags InFlags = EDebugDrawFlags::None);
 	UE_DEPRECATED(5.6, "Use other FDebugDrawParams signatures instead")
 	FDebugDrawParams(TArrayView<const USkinnedMeshComponent*> InMeshes, TConstArrayView<const IPoseHistory*> InPoseHistories, const FRoleToIndex& InRoleToIndex, const UPoseSearchDatabase* InDatabase, EDebugDrawFlags InFlags = EDebugDrawFlags::None);
-	FDebugDrawParams(TConstArrayView<const UObject*> InAnimContexts, TConstArrayView<const IPoseHistory*> InPoseHistories, const FRoleToIndex& InRoleToIndex, const UPoseSearchDatabase* InDatabase, EDebugDrawFlags InFlags = EDebugDrawFlags::None);
+	PRAGMA_ENABLE_DEPRECATION_WARNINGS
+	FDebugDrawParams(TConstArrayView<const UObject*> InAnimContexts, TConstArrayView<const IPoseHistory*> InPoseHistories, const FRoleToIndex& InRoleToIndex, const UPoseSearchDatabase* InDatabase);
 	
 	const FSearchIndex* GetSearchIndex() const;
 	const UPoseSearchSchema* GetSchema() const;
@@ -75,6 +79,8 @@ struct POSESEARCH_API FDebugDrawParams
 	void DrawFeatureVector(TConstArrayView<float> PoseVector);
 	void DrawFeatureVector(int32 PoseIdx);
 
+	bool IsAnyWeightRelevant(const UPoseSearchFeatureChannel* Channel) const;
+
 private:
 	bool CanDraw() const;
 
@@ -86,7 +92,8 @@ private:
 	const FRoleToIndex& RoleToIndex;
 
 	const UPoseSearchDatabase* Database = nullptr;
-	EDebugDrawFlags Flags = EDebugDrawFlags::None;
+	TAlignedArray<float> DynamicWeightsSqrtBuffer;
+	TConstArrayView<float> DynamicWeightsSqrt;
 };
 
 #endif // ENABLE_DRAW_DEBUG
diff --git a/Engine/Plugins/Animation/PoseSearch/Source/Runtime/Public/PoseSearch/PoseSearchDatabase.h b/Engine/Plugins/Animation/PoseSearch/Source/Runtime/Public/PoseSearch/PoseSearchDatabase.h
--- a/Engine/Plugins/Animation/PoseSearch/Source/Runtime/Public/PoseSearch/PoseSearchDatabase.h
+++ b/Engine/Plugins/Animation/PoseSearch/Source/Runtime/Public/PoseSearch/PoseSearchDatabase.h
@@ -584,6 +584,8 @@ public:
 #endif // WITH_EDITOR && ENABLE_ANIM_DEBUG
 
 	TConstArrayView<int32> GetAssetIndexesForSourceAsset(const UObject* SourceAsset) const;
+	// Experimental, this feature might be removed without warning, not for production use
+	TConstArrayView<float> CalculateDynamicWeightsSqrt(TArrayView<float> DynamicWeightsSqrtBuffer) const;
 
 private:
 	UE::PoseSearch::FSearchResult SearchPCAKDTree(UE::PoseSearch::FSearchContext& SearchContext) const;
@@ -596,7 +598,7 @@ private:
 	typedef TArray<int32, TInlineAllocator<256, TMemStackAllocator<>>> FNonSelectableIdx;
 	void PopulateNonSelectableIdx(FNonSelectableIdx& NonSelectableIdx, UE::PoseSearch::FSearchContext& SearchContext
 #if UE_POSE_SEARCH_TRACE_ENABLED
-		, TConstArrayView<float> QueryValues
+		, TConstArrayView<float> QueryValues, TConstArrayView<float> DynamicWeightsSqrt
 #endif //UE_POSE_SEARCH_TRACE_ENABLED
 	) const;
 
diff --git a/Engine/Plugins/Animation/PoseSearch/Source/Runtime/Public/PoseSearch/PoseSearchIndex.h b/Engine/Plugins/Animation/PoseSearch/Source/Runtime/Public/PoseSearch/PoseSearchIndex.h
--- a/Engine/Plugins/Animation/PoseSearch/Source/Runtime/Public/PoseSearch/PoseSearchIndex.h
+++ b/Engine/Plugins/Animation/PoseSearch/Source/Runtime/Public/PoseSearch/PoseSearchIndex.h
@@ -457,9 +457,14 @@ struct FSearchIndex : public FSearchIndexBase
 		return PCAValues.Num() / DataCardinality;
 	}
 
+	UE_DEPRECATED(5.6, "Use other ComparePoses signatures instead")
 	POSESEARCH_API FPoseSearchCost ComparePoses(int32 PoseIdx, float ContinuingPoseCostBias, TConstArrayView<float> PoseValues, TConstArrayView<float> QueryValues) const;
+	UE_DEPRECATED(5.6, "Use other CompareAlignedPoses signatures instead")
 	POSESEARCH_API FPoseSearchCost CompareAlignedPoses(int32 PoseIdx, float ContinuingPoseCostBias, TConstArrayView<float> PoseValues, TConstArrayView<float> QueryValues) const;
 
+	POSESEARCH_API FPoseSearchCost ComparePoses(int32 PoseIdx, float ContinuingPoseCostBias, TConstArrayView<float> PoseValues, TConstArrayView<float> QueryValues, TConstArrayView<float> DynamicWeightsSqrt) const;
+	POSESEARCH_API FPoseSearchCost CompareAlignedPoses(int32 PoseIdx, float ContinuingPoseCostBias, TConstArrayView<float> PoseValues, TConstArrayView<float> QueryValues, TConstArrayView<float> DynamicWeightsSqrt) const;
+
 	void PruneDuplicatePCAValues(float SimilarityThreshold, int32 NumberOfPrincipalComponents);
 	void PrunePCAValuesFromBlockTransitionPoses(int32 NumberOfPrincipalComponents);
 
