PoseSearch - converting TAutoConsoleVariable to FAutoConsoleVariableRef

diff --git a/Engine/Plugins/Animation/PoseSearch/Source/Editor/Private/PoseSearchDatabaseViewModel.cpp b/Engine/Plugins/Animation/PoseSearch/Source/Editor/Private/PoseSearchDatabaseViewModel.cpp
--- a/Engine/Plugins/Animation/PoseSearch/Source/Editor/Private/PoseSearchDatabaseViewModel.cpp
+++ b/Engine/Plugins/Animation/PoseSearch/Source/Editor/Private/PoseSearchDatabaseViewModel.cpp
@@ -32,9 +32,14 @@
 namespace UE::PoseSearch
 {
 #if ENABLE_ANIM_DEBUG
-static TAutoConsoleVariable<float> CVarDatabasePreviewDebugDrawSamplerSize(TEXT("a.DatabasePreview.DebugDrawSamplerSize"), 0.f, TEXT("Debug Draw Sampler Positions Size"));
-static TAutoConsoleVariable<float> CVarDatabasePreviewDebugDrawSamplerTimeOffset(TEXT("a.DatabasePreview.DebugDrawSamplerTimeOffset"), 0.f, TEXT("Debug Draw Sampler Positions At Time Offset"));
-static TAutoConsoleVariable<float> CVarDatabasePreviewDebugDrawSamplerRootAxisLength(TEXT("a.DatabasePreview.DebugDrawSamplerRootAxisLength"), 0.f, TEXT("Debug Draw Sampler Root Axis Length"));
+static float GVarDatabasePreviewDebugDrawSamplerSize = 0.f;
+static FAutoConsoleVariableRef CVarDatabasePreviewDebugDrawSamplerSize(TEXT("a.DatabasePreview.DebugDrawSamplerSize"), GVarDatabasePreviewDebugDrawSamplerSize, TEXT("Debug Draw Sampler Positions Size"));
+
+static float GVarDatabasePreviewDebugDrawSamplerTimeOffset = 0.f;
+static FAutoConsoleVariableRef CVarDatabasePreviewDebugDrawSamplerTimeOffset(TEXT("a.DatabasePreview.DebugDrawSamplerTimeOffset"), GVarDatabasePreviewDebugDrawSamplerTimeOffset, TEXT("Debug Draw Sampler Positions At Time Offset"));
+
+static float GVarDatabasePreviewDebugDrawSamplerRootAxisLength = 0.f;
+static FAutoConsoleVariableRef CVarDatabasePreviewDebugDrawSamplerRootAxisLength(TEXT("a.DatabasePreview.DebugDrawSamplerRootAxisLength"), GVarDatabasePreviewDebugDrawSamplerRootAxisLength, TEXT("Debug Draw Sampler Root Axis Length"));
 #endif
 
 // FDatabasePreviewActor
@@ -364,10 +369,10 @@ bool FDatabasePreviewActor::DrawPreviewActors(TConstArrayView<FDatabasePreviewAc
 		}
 
 #if ENABLE_ANIM_DEBUG
-		const float DebugDrawSamplerSize = CVarDatabasePreviewDebugDrawSamplerSize.GetValueOnAnyThread();
+		const float DebugDrawSamplerSize = GVarDatabasePreviewDebugDrawSamplerSize;
 		if (DebugDrawSamplerSize > UE_KINDA_SMALL_NUMBER)
 		{
-			const float DebugDrawSamplerTimeOffset = CVarDatabasePreviewDebugDrawSamplerTimeOffset.GetValueOnAnyThread();
+			const float DebugDrawSamplerTimeOffset = GVarDatabasePreviewDebugDrawSamplerTimeOffset;
 			
 			const int32 NumDrawPasses = FMath::IsNearlyZero(DebugDrawSamplerTimeOffset) ? 1 : 2;
 
@@ -390,7 +395,7 @@ bool FDatabasePreviewActor::DrawPreviewActors(TConstArrayView<FDatabasePreviewAc
 				ComponentSpacePose.InitPose(MoveTemp(Pose));
 
 				const FTransform RootTransform = MirrorDataCache.MirrorTransform(PreviewActor.Sampler.ExtractRootTransform(SamplerTime));
-				const float DebugDrawSamplerRootAxisLength = CVarDatabasePreviewDebugDrawSamplerRootAxisLength.GetValueOnAnyThread();
+				const float DebugDrawSamplerRootAxisLength = GVarDatabasePreviewDebugDrawSamplerRootAxisLength;
 				if (DebugDrawSamplerRootAxisLength > 0.f)
 				{
 					DrawParams.DrawLine(RootTransform.GetTranslation(), RootTransform.GetTranslation() + RootTransform.GetScaledAxis(EAxis::X) * DebugDrawSamplerRootAxisLength, FColor::Red);
diff --git a/Engine/Plugins/Animation/PoseSearch/Source/Runtime/Private/AnimNode_MotionMatching.cpp b/Engine/Plugins/Animation/PoseSearch/Source/Runtime/Private/AnimNode_MotionMatching.cpp
--- a/Engine/Plugins/Animation/PoseSearch/Source/Runtime/Private/AnimNode_MotionMatching.cpp
+++ b/Engine/Plugins/Animation/PoseSearch/Source/Runtime/Private/AnimNode_MotionMatching.cpp
@@ -21,11 +21,20 @@
 #define LOCTEXT_NAMESPACE "AnimNode_MotionMatching"
 
 #if ENABLE_ANIM_DEBUG
-static TAutoConsoleVariable<bool> CVarAnimNodeMotionMatchingDrawQuery(TEXT("a.AnimNode.MotionMatching.DebugDrawQuery"), false, TEXT("Draw input query"));
-static TAutoConsoleVariable<bool> CVarAnimNodeMotionMatchingDrawCurResult(TEXT("a.AnimNode.MotionMatching.DebugDrawCurResult"), false, TEXT("Draw current result"));
-static TAutoConsoleVariable<bool> CVarAnimNodeMotionMatchingDrawInfo(TEXT("a.AnimNode.MotionMatching.DebugDrawInfo"), false, TEXT("Draw info like current databases and asset"));
-static TAutoConsoleVariable<bool> CVarAnimNodeMotionMatchingDrawInfoVerbose(TEXT("a.AnimNode.MotionMatching.DebugDrawInfoVerbose"), true, TEXT("Draw additional info like blend stack"));
-static TAutoConsoleVariable<float> CVarAnimNodeMotionMatchingDrawInfoHeight(TEXT("a.AnimNode.MotionMatching.DebugDrawInfoHeight"), 50.f, TEXT("Vertical offset for DebugDrawInfo"));
+static bool GVarAnimNodeMotionMatchingDrawQuery = false;
+static FAutoConsoleVariableRef CVarAnimNodeMotionMatchingDrawQuery(TEXT("a.AnimNode.MotionMatching.DebugDrawQuery"), GVarAnimNodeMotionMatchingDrawQuery, TEXT("Draw input query"));
+
+static bool GVarAnimNodeMotionMatchingDrawCurResult = false;
+static FAutoConsoleVariableRef CVarAnimNodeMotionMatchingDrawCurResult(TEXT("a.AnimNode.MotionMatching.DebugDrawCurResult"), GVarAnimNodeMotionMatchingDrawCurResult, TEXT("Draw current result"));
+
+static bool GVarAnimNodeMotionMatchingDrawInfo = false;
+static FAutoConsoleVariableRef CVarAnimNodeMotionMatchingDrawInfo(TEXT("a.AnimNode.MotionMatching.DebugDrawInfo"), GVarAnimNodeMotionMatchingDrawInfo, TEXT("Draw info like current databases and asset"));
+
+static bool GVarAnimNodeMotionMatchingDrawInfoVerbose = true;
+static FAutoConsoleVariableRef CVarAnimNodeMotionMatchingDrawInfoVerbose(TEXT("a.AnimNode.MotionMatching.DebugDrawInfoVerbose"), GVarAnimNodeMotionMatchingDrawInfoVerbose, TEXT("Draw additional info like blend stack"));
+
+static float GVarAnimNodeMotionMatchingDrawInfoHeight = 50.f;
+static FAutoConsoleVariableRef CVarAnimNodeMotionMatchingDrawInfoHeight(TEXT("a.AnimNode.MotionMatching.DebugDrawInfoHeight"), GVarAnimNodeMotionMatchingDrawInfoHeight, TEXT("Vertical offset for DebugDrawInfo"));
 #endif
 
 /////////////////////////////////////////////////////
@@ -144,7 +153,7 @@ void FAnimNode_MotionMatching::UpdateAssetPlayer(const FAnimationUpdateContext&
 	}
 
 #if ENABLE_ANIM_DEBUG
-	if (CVarAnimNodeMotionMatchingDrawInfo.GetValueOnAnyThread())
+	if (GVarAnimNodeMotionMatchingDrawInfo)
 	{
 		const UPoseSearchDatabase* CurrentDatabase = MotionMatchingState.CurrentSearchResult.Database.Get();
 		const UAnimationAsset* CurrentAnimationAsset = AnimPlayers.IsEmpty() ? nullptr : AnimPlayers[0].GetAnimationAsset();
@@ -152,7 +161,7 @@ void FAnimNode_MotionMatching::UpdateAssetPlayer(const FAnimationUpdateContext&
 		FString DebugInfo = FString::Printf(TEXT("NextUpdateInterruptMode(%s)\n"), *UEnum::GetValueAsString(NextUpdateInterruptMode));
 		DebugInfo += FString::Printf(TEXT("Current Database(%s)\n"), *GetNameSafe(CurrentDatabase));
 		DebugInfo += FString::Printf(TEXT("Current Asset(%s)\n"), *GetNameSafe(CurrentAnimationAsset));
-		if (CVarAnimNodeMotionMatchingDrawInfoVerbose.GetValueOnAnyThread())
+		if (GVarAnimNodeMotionMatchingDrawInfoVerbose)
 		{
 			DebugInfo += FString::Printf(TEXT("Databases to search:\n"));
 			for (const UPoseSearchDatabase* DatabaseToSearch : DatabasesToSearch)
@@ -165,7 +174,7 @@ void FAnimNode_MotionMatching::UpdateAssetPlayer(const FAnimationUpdateContext&
 				DebugInfo += FString::Printf(TEXT("  %s [time:%.2f|playrate:%.2f]\n"), *GetNameSafe(AnimPlayer.GetAnimationAsset()), AnimPlayer.GetAccumulatedTime(), AnimPlayer.GetPlayRate());
 			}
 		}
-		Context.AnimInstanceProxy->AnimDrawDebugInWorldMessage(DebugInfo, FVector::UpVector * CVarAnimNodeMotionMatchingDrawInfoHeight.GetValueOnAnyThread(), FColor::Yellow, 1.f /*TextScale*/);
+		Context.AnimInstanceProxy->AnimDrawDebugInWorldMessage(DebugInfo, FVector::UpVector * GVarAnimNodeMotionMatchingDrawInfoHeight, FColor::Yellow, 1.f /*TextScale*/);
 	}
 #endif // ENABLE_ANIM_DEBUG
 
@@ -198,8 +207,8 @@ void FAnimNode_MotionMatching::UpdateAssetPlayer(const FAnimationUpdateContext&
 			NextUpdateInterruptMode,
 			bShouldUseCachedChannelData
 #if ENABLE_ANIM_DEBUG
-			, CVarAnimNodeMotionMatchingDrawQuery.GetValueOnAnyThread()
-			, CVarAnimNodeMotionMatchingDrawCurResult.GetValueOnAnyThread()
+			, GVarAnimNodeMotionMatchingDrawQuery
+			, GVarAnimNodeMotionMatchingDrawCurResult
 #endif // ENABLE_ANIM_DEBUG
 		);
 	}
diff --git a/Engine/Plugins/Animation/PoseSearch/Source/Runtime/Private/AnimNode_MotionMatchingInteraction.cpp b/Engine/Plugins/Animation/PoseSearch/Source/Runtime/Private/AnimNode_MotionMatchingInteraction.cpp
--- a/Engine/Plugins/Animation/PoseSearch/Source/Runtime/Private/AnimNode_MotionMatchingInteraction.cpp
+++ b/Engine/Plugins/Animation/PoseSearch/Source/Runtime/Private/AnimNode_MotionMatchingInteraction.cpp
@@ -13,7 +13,8 @@
 #include "HAL/IConsoleManager.h"
 
 #if ENABLE_ANIM_DEBUG
-TAutoConsoleVariable<bool> CVarAnimNodeMotionMatchingInteractionDebug(TEXT("a.AnimNode.MotionMatchingInteraction.Debug"), false, TEXT("Turn on visualization debugging for AnimNode Motion Matching Interaction"));
+static bool GVarAnimNodeMotionMatchingInteractionDebug = false;
+static FAutoConsoleVariableRef CVarAnimNodeMotionMatchingInteractionDebug(TEXT("a.AnimNode.MotionMatchingInteraction.Debug"), GVarAnimNodeMotionMatchingInteractionDebug, TEXT("Turn on visualization debugging for AnimNode Motion Matching Interaction"));
 #endif // ENABLE_ANIM_DEBUG
 
 void FAnimNode_MotionMatchingInteraction::GatherDebugData(FNodeDebugData& DebugData)
@@ -139,7 +140,7 @@ void FAnimNode_MotionMatchingInteraction::UpdateAssetPlayer(const FAnimationUpda
 				FVector::ZeroVector);
 
 #if ENABLE_ANIM_DEBUG
-			if (CVarAnimNodeMotionMatchingInteractionDebug.GetValueOnAnyThread())
+			if (GVarAnimNodeMotionMatchingInteractionDebug)
 			{
 				Context.AnimInstanceProxy->AnimDrawDebugCoordinateSystem(RootBoneTrasform.GetLocation(), RootBoneTrasform.Rotator(), 25.f, false, 0.f, 0.f, SDPG_Foreground);
 				Context.AnimInstanceProxy->AnimDrawDebugCoordinateSystem(FullAlignedActorRootBoneTransform.GetLocation(), FullAlignedActorRootBoneTransform.Rotator(), 50.f, false, 0.f, 0.f, SDPG_Foreground);
@@ -223,7 +224,7 @@ void FAnimNode_MotionMatchingInteraction::Evaluate_AnyThread(FPoseContext& Outpu
 			const FTransform FullAlignedActorRootBoneLocalTransform = FullAlignedActorRootBoneTransform.GetRelativeTransform(ComponentTransform);
 
 #if ENABLE_ANIM_DEBUG
-			if (CVarAnimNodeMotionMatchingInteractionDebug.GetValueOnAnyThread())
+			if (GVarAnimNodeMotionMatchingInteractionDebug)
 			{
 				Output.AnimInstanceProxy->AnimDrawDebugCoordinateSystem(FullAlignedActorRootBoneTransform.GetLocation(), FullAlignedActorRootBoneTransform.Rotator(), 15.f, false, 0.f, 0.f, SDPG_Foreground);
 				Output.AnimInstanceProxy->AnimDrawDebugCoordinateSystem(ComponentTransform.GetLocation(), ComponentTransform.Rotator(), 5.f, false, 0.f, 0.f, SDPG_Foreground);
diff --git a/Engine/Plugins/Animation/PoseSearch/Source/Runtime/Private/PoseSearchAssetIndexer.cpp b/Engine/Plugins/Animation/PoseSearch/Source/Runtime/Private/PoseSearchAssetIndexer.cpp
--- a/Engine/Plugins/Animation/PoseSearch/Source/Runtime/Private/PoseSearchAssetIndexer.cpp
+++ b/Engine/Plugins/Animation/PoseSearch/Source/Runtime/Private/PoseSearchAssetIndexer.cpp
@@ -18,8 +18,11 @@ namespace UE::PoseSearch
 {
 
 #if ENABLE_ANIM_DEBUG
-static TAutoConsoleVariable<bool> CVarMotionMatchTestDisableIndexerCaching(TEXT("a.MotionMatch.TestDisableIndexerCaching"), false, TEXT("Disable Motion Matching Indexer Caching"));
-static TAutoConsoleVariable<int32> CVarMotionMatchTestExtractPoseDeterminismNumIterations(TEXT("a.MotionMatch.TestExtractPoseDeterminismNumIterations"), 0, TEXT("Test Motion Matching ExtractPose Determinism via this NumIterations retries"));
+static bool GVarMotionMatchTestDisableIndexerCaching = false;
+static FAutoConsoleVariableRef CVarMotionMatchTestDisableIndexerCaching(TEXT("a.MotionMatch.TestDisableIndexerCaching"), GVarMotionMatchTestDisableIndexerCaching, TEXT("Disable Motion Matching Indexer Caching"));
+
+static int32 GVarMotionMatchTestExtractPoseDeterminismNumIterations = 0;
+static FAutoConsoleVariableRef CVarMotionMatchTestExtractPoseDeterminismNumIterations(TEXT("a.MotionMatch.TestExtractPoseDeterminismNumIterations"), GVarMotionMatchTestExtractPoseDeterminismNumIterations, TEXT("Test Motion Matching ExtractPose Determinism via this NumIterations retries"));
 #endif // ENABLE_ANIM_DEBUG
 
 //////////////////////////////////////////////////////////////////////////
@@ -409,7 +412,7 @@ FAssetIndexer::FCachedEntry& FAssetIndexer::GetEntry(float SampleTime)
 	
 	bool bDisableCaching = false;
 #if ENABLE_ANIM_DEBUG
-	bDisableCaching = CVarMotionMatchTestDisableIndexerCaching.GetValueOnAnyThread();
+	bDisableCaching = GVarMotionMatchTestDisableIndexerCaching;
 #endif // ENABLE_ANIM_DEBUG
 
 	SampleTime = FMath::Clamp(SampleTime, ExtrapolationTimeInterval.Min, ExtrapolationTimeInterval.Max);
@@ -459,7 +462,7 @@ FAssetIndexer::FCachedEntry& FAssetIndexer::GetEntry(float SampleTime)
 			AssetSamplers.ExtractPose(CurrentTime, Pose, Curve, RoleIndex);
 
 #if ENABLE_ANIM_DEBUG
-			const int32 NumIterations = CVarMotionMatchTestExtractPoseDeterminismNumIterations.GetValueOnAnyThread();
+			const int32 NumIterations = GVarMotionMatchTestExtractPoseDeterminismNumIterations;
 			for (int32 IterationIndex = 0; IterationIndex < NumIterations; ++IterationIndex)
 			{
 				FCompactPose TestPose;
diff --git a/Engine/Plugins/Animation/PoseSearch/Source/Runtime/Private/PoseSearchDerivedData.cpp b/Engine/Plugins/Animation/PoseSearch/Source/Runtime/Private/PoseSearchDerivedData.cpp
--- a/Engine/Plugins/Animation/PoseSearch/Source/Runtime/Private/PoseSearchDerivedData.cpp
+++ b/Engine/Plugins/Animation/PoseSearch/Source/Runtime/Private/PoseSearchDerivedData.cpp
@@ -87,16 +87,25 @@ enum EMotionMatchTestFlags
 	// test DDC key generation determinism
 	TestDDCKeyDeterminism = 1 << 14,
 };
-static TAutoConsoleVariable<int32> CVarMotionMatchTestFlags(TEXT("a.MotionMatch.TestFlags"), EMotionMatchTestFlags::None, TEXT("Test Motion Matching using EMotionMatchTestFlags"));
-static TAutoConsoleVariable<int32> CVarMotionMatchTestNumIterations(TEXT("a.MotionMatch.TestNumIterations"), 10, TEXT("Test Motion Matching Num Iterations"));
-static bool AnyTestFlags(int32 Flags) { return (CVarMotionMatchTestFlags.GetValueOnAnyThread() & Flags) != 0; }
+
+static int32 GVarMotionMatchTestFlags = EMotionMatchTestFlags::None;
+static FAutoConsoleVariableRef CVarMotionMatchTestFlags(TEXT("a.MotionMatch.TestFlags"), GVarMotionMatchTestFlags, TEXT("Test Motion Matching using EMotionMatchTestFlags"));
+
+static int32 GVarMotionMatchTestNumIterations = 10;
+static FAutoConsoleVariableRef CVarMotionMatchTestNumIterations(TEXT("a.MotionMatch.TestNumIterations"), GVarMotionMatchTestNumIterations, TEXT("Test Motion Matching Num Iterations"));
+static bool AnyTestFlags(int32 Flags) { return (GVarMotionMatchTestFlags & Flags) != 0; }
 #endif // ENABLE_ANIM_DEBUG
 
-static TAutoConsoleVariable<bool> CVarMotionMatchReindexCancelledDatabases(TEXT("a.MotionMatch.ReindexCancelledDatabases"), false, TEXT("Reindex Cancelled Databases"));
+static bool GVarMotionMatchReindexCancelledDatabases = false;
+static FAutoConsoleVariableRef CVarMotionMatchReindexCancelledDatabases(TEXT("a.MotionMatch.ReindexCancelledDatabases"), GVarMotionMatchReindexCancelledDatabases, TEXT("Reindex Cancelled Databases"));
+
 // Experimental, this feature might be removed without warning, not for production use
-static TAutoConsoleVariable<bool> CVarMotionMatchReindexAllReferencedDatabases(TEXT("a.MotionMatch.ReindexAllReferencedDatabases"), true, TEXT("Reindex All Referenced Databases"));
+static bool GVarMotionMatchReindexAllReferencedDatabases = true;
+static FAutoConsoleVariableRef CVarMotionMatchReindexAllReferencedDatabases(TEXT("a.MotionMatch.ReindexAllReferencedDatabases"), GVarMotionMatchReindexAllReferencedDatabases, TEXT("Reindex All Referenced Databases"));
+
 // Experimental, this feature might be removed without warning, not for production use
-static TAutoConsoleVariable<int32> CVarMotionMatchPartialKeyHashesMode(TEXT("a.MotionMatch.PartialKeyHashesMode"), 0, TEXT("0: use partial key hashes, 1: do not use partial key hashes, 2: do not use and validate partial key hashes"));
+static int32 GVarMotionMatchPartialKeyHashesMode = 0;
+static FAutoConsoleVariableRef CVarMotionMatchPartialKeyHashesMode(TEXT("a.MotionMatch.PartialKeyHashesMode"), GVarMotionMatchPartialKeyHashesMode, TEXT("0: use partial key hashes, 1: do not use partial key hashes, 2: do not use and validate partial key hashes"));
 
 static const UE::DerivedData::FValueId Id(UE::DerivedData::FValueId::FromName("Data"));
 static const UE::DerivedData::FCacheBucket Bucket("PoseSearchDatabase");
@@ -737,7 +746,7 @@ static void PreprocessSearchIndexKDTree(FSearchIndex& SearchIndex, const UPoseSe
 		// testing kdtree Construct determinism
 		if (AnyTestFlags(EMotionMatchTestFlags::TestKDTreeConstructDeterminism))
 		{
-			const int32 NumIterations = CVarMotionMatchTestNumIterations.GetValueOnAnyThread();
+			const int32 NumIterations = GVarMotionMatchTestNumIterations;
 			for (int32 Iteration = 0; Iteration < NumIterations; ++Iteration)
 			{
 				// copy PCAValues in a different container to ensure input data has different memory addresses
@@ -849,7 +858,7 @@ static void PreprocessSearchIndexVPTree(FSearchIndex& SearchIndex, const UPoseSe
 				UE_LOG(LogPoseSearch, Error, TEXT("PreprocessSearchIndexVPTree - FVPTree construction failed"));
 			}
 
-			const int32 NumIterations = CVarMotionMatchTestNumIterations.GetValueOnAnyThread();
+			const int32 NumIterations = GVarMotionMatchTestNumIterations;
 			for (int32 Iteration = 0; Iteration < NumIterations; ++Iteration)
 			{
 				FVPTree VPTreeTest;
@@ -1404,7 +1413,7 @@ void FPoseSearchDatabaseAsyncCacheTask::StartNewRequestIfNeeded(bool bPerformCon
 	Owner.Cancel();
 
 	FKeyBuilder::EDebugPartialKeyHashesMode DebugPartialKeyHashesMode;
-	switch (CVarMotionMatchPartialKeyHashesMode.GetValueOnAnyThread())
+	switch (GVarMotionMatchPartialKeyHashesMode)
 	{
 	case 0:
 		DebugPartialKeyHashesMode = FKeyBuilder::EDebugPartialKeyHashesMode::Use;
@@ -1741,7 +1750,7 @@ void FPoseSearchDatabaseAsyncCacheTask::OnGetComplete(UE::DerivedData::FCacheGet
 #if ENABLE_ANIM_DEBUG
 					if (AnyTestFlags(EMotionMatchTestFlags::TestIndexDatabaseDeterminism))
 					{
-						const int32 NumIterations = CVarMotionMatchTestNumIterations.GetValueOnAnyThread();
+						const int32 NumIterations = GVarMotionMatchTestNumIterations;
 						for (int32 Iteration = 0; Iteration < NumIterations; ++Iteration)
 						{
 							FSearchIndexBase TestSearchIndexBase = SearchIndexBase;
@@ -1765,7 +1774,7 @@ void FPoseSearchDatabaseAsyncCacheTask::OnGetComplete(UE::DerivedData::FCacheGet
 				// testing PruneDuplicateValues determinism
 				if (AnyTestFlags(EMotionMatchTestFlags::TestPruneDuplicateValuesDeterminism))
 				{
-					const int32 NumIterations = CVarMotionMatchTestNumIterations.GetValueOnAnyThread();
+					const int32 NumIterations = GVarMotionMatchTestNumIterations;
 
 					FSearchIndex TestSearchIndexA = SearchIndex;
 					TestSearchIndexA.PruneDuplicateValues(MainDatabase->PosePruningSimilarityThreshold, MainDatabaseSchema->SchemaCardinality, false);
@@ -1830,7 +1839,7 @@ void FPoseSearchDatabaseAsyncCacheTask::OnGetComplete(UE::DerivedData::FCacheGet
 				// testing PruneDuplicatePCAValues determinism
 				if (AnyTestFlags(EMotionMatchTestFlags::TestPruneDuplicatePCAValuesDeterminism))
 				{
-					const int32 NumIterations = CVarMotionMatchTestNumIterations.GetValueOnAnyThread();
+					const int32 NumIterations = GVarMotionMatchTestNumIterations;
 
 					FSearchIndex TestSearchIndexA = SearchIndex;
 					TestSearchIndexA.PruneDuplicatePCAValues(MainDatabase->PCAValuesPruningSimilarityThreshold, MainDatabase->GetNumberOfPrincipalComponents());
@@ -2188,7 +2197,7 @@ void FAsyncPoseSearchDatabasesManagement::Tick(float DeltaTime)
 	const bool bTestAssetSamplerDeterminismFromPreviousExecution = AnyTestFlags(EMotionMatchTestFlags::TestAssetSamplerDeterminismFromPreviousExecution);
 	if (bTestAssetSamplerDeterminism || bTestAssetSamplerDeterminismFromPreviousExecution)
 	{
-		const int32 NumIterations = CVarMotionMatchTestNumIterations.GetValueOnAnyThread();
+		const int32 NumIterations = GVarMotionMatchTestNumIterations;
 
 		struct FTestSample
 		{
@@ -2393,7 +2402,7 @@ void FAsyncPoseSearchDatabasesManagement::Tick(float DeltaTime)
 
 #endif // ENABLE_ANIM_DEBUG
 	
-	const bool bReindexCancelledDatabases = CVarMotionMatchReindexCancelledDatabases.GetValueOnAnyThread();
+	const bool bReindexCancelledDatabases = GVarMotionMatchReindexCancelledDatabases;
 
 	// iterating backwards because of the possible RemoveAtSwap 
 	for (int32 TaskIndex = Tasks.Num() - 1; TaskIndex >= 0; --TaskIndex)
@@ -2423,7 +2432,7 @@ void FAsyncPoseSearchDatabasesManagement::Tick(float DeltaTime)
 #if ENABLE_ANIM_DEBUG
 	if (AnyTestFlags(EMotionMatchTestFlags::TestDDCKeyDeterminism))
 	{
-		const int32 NumIterations = CVarMotionMatchTestNumIterations.GetValueOnAnyThread();
+		const int32 NumIterations = GVarMotionMatchTestNumIterations;
 		for (int32 TaskIndex = 0; TaskIndex < Tasks.Num(); ++TaskIndex)
 		{
 			if (const UPoseSearchDatabase* Database = Tasks[TaskIndex]->GetDatabase())
@@ -2537,7 +2546,7 @@ EAsyncBuildIndexResult FAsyncPoseSearchDatabasesManagement::RequestAsyncBuildInd
 
 EAsyncBuildIndexResult FAsyncPoseSearchDatabasesManagement::RequestAsyncBuildIndex(const UPoseSearchDatabase* Database, ERequestAsyncBuildFlag Flag)
 {
-	if (CVarMotionMatchReindexAllReferencedDatabases.GetValueOnAnyThread())
+	if (GVarMotionMatchReindexAllReferencedDatabases)
 	{
 		FDatabaseSet DatabaseSet;
 		RecursivePopulateDependentDatabases(Database, DatabaseSet);
diff --git a/Engine/Plugins/Animation/PoseSearch/Source/Runtime/Private/PoseSearchHistory.cpp b/Engine/Plugins/Animation/PoseSearch/Source/Runtime/Private/PoseSearchHistory.cpp
--- a/Engine/Plugins/Animation/PoseSearch/Source/Runtime/Private/PoseSearchHistory.cpp
+++ b/Engine/Plugins/Animation/PoseSearch/Source/Runtime/Private/PoseSearchHistory.cpp
@@ -19,11 +19,20 @@ namespace UE::PoseSearch
 {
 
 #if ENABLE_DRAW_DEBUG && ENABLE_ANIM_DEBUG
-TAutoConsoleVariable<bool> CVarAnimPoseHistoryDebugDrawPose(TEXT("a.AnimNode.PoseHistory.DebugDrawPose"), false, TEXT("Enable / Disable Pose History Pose DebugDraw"));
-TAutoConsoleVariable<bool> CVarAnimPoseHistoryDebugDrawTrajectory(TEXT("a.AnimNode.PoseHistory.DebugDrawTrajectory"), false, TEXT("Enable / Disable Pose History Trajectory DebugDraw"));
-TAutoConsoleVariable<float> CVarAnimPoseHistoryDebugDrawTrajectoryThickness(TEXT("a.AnimNode.PoseHistory.DebugDrawTrajectoryThickness"), 0.0f, TEXT("Thickness of the trajectory debug draw (Default 0.0f)"));
-TAutoConsoleVariable<int> CVarAnimPoseHistoryDebugDrawTrajectoryMaxNumOfHistorySamples(TEXT("a.AnimNode.PoseHistory.DebugDrawMaxNumOfHistorySamples"), -1, TEXT("Max number of history samples to debug draw. All history samples will be drawn if value is negative. (Default -1)"));
-TAutoConsoleVariable<int> CVarAnimPoseHistoryDebugDrawTrajectoryMaxNumOfPredictionSamples(TEXT("a.AnimNode.PoseHistory.DebugDrawMaxNumOfPredictionSamples"), -1, TEXT("Max number of prediction samples to debug draw. All prediction samples will be drawn if value is negative. (Default -1)"));
+static bool GVarAnimPoseHistoryDebugDrawPose = false;
+static FAutoConsoleVariableRef CVarAnimPoseHistoryDebugDrawPose(TEXT("a.AnimNode.PoseHistory.DebugDrawPose"), GVarAnimPoseHistoryDebugDrawPose, TEXT("Enable / Disable Pose History Pose DebugDraw"));
+
+static bool GVarAnimPoseHistoryDebugDrawTrajectory = false;
+static FAutoConsoleVariableRef CVarAnimPoseHistoryDebugDrawTrajectory(TEXT("a.AnimNode.PoseHistory.DebugDrawTrajectory"), GVarAnimPoseHistoryDebugDrawTrajectory, TEXT("Enable / Disable Pose History Trajectory DebugDraw"));
+
+static float GVarAnimPoseHistoryDebugDrawTrajectoryThickness = 0.f;
+static FAutoConsoleVariableRef CVarAnimPoseHistoryDebugDrawTrajectoryThickness(TEXT("a.AnimNode.PoseHistory.DebugDrawTrajectoryThickness"), GVarAnimPoseHistoryDebugDrawTrajectoryThickness, TEXT("Thickness of the trajectory debug draw (Default 0.0f)"));
+
+static int32 GVarAnimPoseHistoryDebugDrawTrajectoryMaxNumOfHistorySamples = -1;
+static FAutoConsoleVariableRef CVarAnimPoseHistoryDebugDrawTrajectoryMaxNumOfHistorySamples(TEXT("a.AnimNode.PoseHistory.DebugDrawMaxNumOfHistorySamples"), GVarAnimPoseHistoryDebugDrawTrajectoryMaxNumOfHistorySamples, TEXT("Max number of history samples to debug draw. All history samples will be drawn if value is negative. (Default -1)"));
+
+static int32 GVarAnimPoseHistoryDebugDrawTrajectoryMaxNumOfPredictionSamples = -1;
+static FAutoConsoleVariableRef CVarAnimPoseHistoryDebugDrawTrajectoryMaxNumOfPredictionSamples(TEXT("a.AnimNode.PoseHistory.DebugDrawMaxNumOfPredictionSamples"), GVarAnimPoseHistoryDebugDrawTrajectoryMaxNumOfPredictionSamples, TEXT("Max number of prediction samples to debug draw. All prediction samples will be drawn if value is negative. (Default -1)"));
 #endif
 
 /**
@@ -987,15 +996,15 @@ void FPoseHistory::DebugDraw(FAnimInstanceProxy& AnimInstanceProxy, FColor Color
 {
 	CheckThreadSafetyRead(ReadPoseDataThreadSafeCounter);
 
-	if (CVarAnimPoseHistoryDebugDrawTrajectory.GetValueOnAnyThread())
+	if (GVarAnimPoseHistoryDebugDrawTrajectory)
 	{
-		const float DebugThickness = CVarAnimPoseHistoryDebugDrawTrajectoryThickness.GetValueOnAnyThread();
-		const int MaxHistorySamples = CVarAnimPoseHistoryDebugDrawTrajectoryMaxNumOfHistorySamples.GetValueOnAnyThread();
-		const int MaxPredictionSamples = CVarAnimPoseHistoryDebugDrawTrajectoryMaxNumOfPredictionSamples.GetValueOnAnyThread();
+		const float DebugThickness = GVarAnimPoseHistoryDebugDrawTrajectoryThickness;
+		const int MaxHistorySamples = GVarAnimPoseHistoryDebugDrawTrajectoryMaxNumOfHistorySamples;
+		const int MaxPredictionSamples = GVarAnimPoseHistoryDebugDrawTrajectoryMaxNumOfPredictionSamples;
 		Trajectory.DebugDrawTrajectory(AnimInstanceProxy, DebugThickness, 0, MaxHistorySamples, MaxPredictionSamples);
 	}
 
-	if (Color.A > 0 && CVarAnimPoseHistoryDebugDrawPose.GetValueOnAnyThread())
+	if (Color.A > 0 && GVarAnimPoseHistoryDebugDrawPose)
 	{
 		const bool bValidTrajectory = !Trajectory.Samples.IsEmpty();
 		TArray<FTransform, TInlineAllocator<128>> PrevGlobalTransforms;
@@ -1140,7 +1149,7 @@ void FMemStackPoseHistory::DebugDraw(FAnimInstanceProxy& AnimInstanceProxy, FCol
 {
 	check(PoseHistory);
 
-	if (Color.A > 0 && !FutureEntries.IsEmpty() && CVarAnimPoseHistoryDebugDrawPose.GetValueOnAnyThread())
+	if (Color.A > 0 && !FutureEntries.IsEmpty() && GVarAnimPoseHistoryDebugDrawPose)
 	{
 		const FPoseSearchQueryTrajectory& Trajectory = GetTrajectory();
 		const bool bValidTrajectory = !Trajectory.Samples.IsEmpty();
diff --git a/Engine/Plugins/Animation/PoseSearch/Source/Runtime/Private/PoseSearchInteractionIsland.cpp b/Engine/Plugins/Animation/PoseSearch/Source/Runtime/Private/PoseSearchInteractionIsland.cpp
--- a/Engine/Plugins/Animation/PoseSearch/Source/Runtime/Private/PoseSearchInteractionIsland.cpp
+++ b/Engine/Plugins/Animation/PoseSearch/Source/Runtime/Private/PoseSearchInteractionIsland.cpp
@@ -17,7 +17,8 @@ namespace UE::PoseSearch
 {
 
 #if ENABLE_DRAW_DEBUG && ENABLE_ANIM_DEBUG
-static TAutoConsoleVariable<bool> CVarPoseSearchInteractionShowIslands(TEXT("a.PoseSearchInteraction.ShowIslands"), false, TEXT("Show Pose Search Interaction Islands"));
+static bool GVarPoseSearchInteractionShowIslands = false;
+static FAutoConsoleVariableRef CVarPoseSearchInteractionShowIslands(TEXT("a.PoseSearchInteraction.ShowIslands"), GVarPoseSearchInteractionShowIslands, TEXT("Show Pose Search Interaction Islands"));
 #endif // ENABLE_DRAW_DEBUG && ENABLE_ANIM_DEBUG
 
 // Utils functions
@@ -327,7 +328,7 @@ void FInteractionIsland::DebugDraw(const FColor& Color) const
 #if ENABLE_DRAW_DEBUG && ENABLE_ANIM_DEBUG
 	check(IsInGameThread());
 
-	if (CVarPoseSearchInteractionShowIslands.GetValueOnAnyThread())
+	if (GVarPoseSearchInteractionShowIslands)
 	{
 		for (const FInteractionSearchContext& SearchContext : SearchContexts)
 		{
diff --git a/Engine/Plugins/Animation/PoseSearch/Source/Runtime/Public/PoseSearch/PoseSearchHistory.h b/Engine/Plugins/Animation/PoseSearch/Source/Runtime/Public/PoseSearch/PoseSearchHistory.h
--- a/Engine/Plugins/Animation/PoseSearch/Source/Runtime/Public/PoseSearch/PoseSearchHistory.h
+++ b/Engine/Plugins/Animation/PoseSearch/Source/Runtime/Public/PoseSearch/PoseSearchHistory.h
@@ -17,11 +17,6 @@ class UWorld;
 namespace UE::PoseSearch
 {
 
-#if ENABLE_DRAW_DEBUG && ENABLE_ANIM_DEBUG
-extern POSESEARCH_API TAutoConsoleVariable<bool> CVarAnimPoseHistoryDebugDrawPose;
-extern POSESEARCH_API TAutoConsoleVariable<bool> CVarAnimPoseHistoryDebugDrawTrajectory;
-#endif
-
 struct FSearchResult;
 typedef uint16 FComponentSpaceTransformIndex;
 typedef TPair<FBoneIndexType, FComponentSpaceTransformIndex> FBoneToTransformPair;
