PoseSearch: Fix for TimeToLand

#rnx

diff --git a/Engine/Plugins/Animation/PoseSearch/Source/Runtime/Private/PoseSearchTrajectoryLibrary.cpp b/Engine/Plugins/Animation/PoseSearch/Source/Runtime/Private/PoseSearchTrajectoryLibrary.cpp
--- a/Engine/Plugins/Animation/PoseSearch/Source/Runtime/Private/PoseSearchTrajectoryLibrary.cpp
+++ b/Engine/Plugins/Animation/PoseSearch/Source/Runtime/Private/PoseSearchTrajectoryLibrary.cpp
@@ -509,7 +509,7 @@ void UPoseSearchTrajectoryLibrary::HandleTrajectoryWorldCollisionsWithGravity(co
 		GravityZ = -GravityZ;
 		const FVector VelocityOnGravityAxis = StartingVelocity.ProjectOnTo(GravityDirection);
 		
-		InitialVelocityZ = VelocityOnGravityAxis.Length() * FMath::Sign(GravityDirection.Dot(VelocityOnGravityAxis));
+		InitialVelocityZ = VelocityOnGravityAxis.Length() * -FMath::Sign(GravityDirection.Dot(VelocityOnGravityAxis));
 	}
 
 	CollisionResult.TimeToLand = OutTrajectory.Samples.Last().AccumulatedSeconds;
