PoseSearch - fix for FSearchContext::GetWorldRootBoneTransformAtTime not leveraging a valid trajectory in case of invalid pose entries

diff --git a/Engine/Plugins/Animation/PoseSearch/Source/Runtime/Private/PoseSearchContext.cpp b/Engine/Plugins/Animation/PoseSearch/Source/Runtime/Private/PoseSearchContext.cpp
--- a/Engine/Plugins/Animation/PoseSearch/Source/Runtime/Private/PoseSearchContext.cpp
+++ b/Engine/Plugins/Animation/PoseSearch/Source/Runtime/Private/PoseSearchContext.cpp
@@ -645,6 +645,7 @@ FTransform FSearchContext::GetWorldRootBoneTransformAtTime(float SampleTime, con
 	const UPoseSearchSchema* Schema = CachedQueries.Last().GetSchema();
 	check(Schema);
 
+	FTransform PoseHistoryTransformAtTime;
 	const IPoseHistory* PoseHistory = GetPoseHistory(SampleRole);
 	#if WITH_EDITOR
 	if (!PoseHistory)
@@ -654,11 +655,14 @@ FTransform FSearchContext::GetWorldRootBoneTransformAtTime(float SampleTime, con
 	else
 	#endif // WITH_EDITOR
 	{
-		FTransform WorldRootBoneTransform;
-		if (ensure(PoseHistory) && PoseHistory->GetTransformAtTime(SampleTime, WorldRootBoneTransform, Schema->GetSkeleton(SampleRole), RootBoneIndexType, WorldSpaceIndexType))
+		if (ensure(PoseHistory) && PoseHistory->GetTransformAtTime(SampleTime, PoseHistoryTransformAtTime, Schema->GetSkeleton(SampleRole), RootBoneIndexType, WorldSpaceIndexType))
 		{
-			return WorldRootBoneTransform;
+			// PoseHistoryTransformAtTime contains the WorldRootBoneTransform (RootBoneTransform from the blended entry in the pose history times ComponentToWorldTransform from the trajectory)
+			return PoseHistoryTransformAtTime;
 		}
+
+		// NoTe: when GetTransformAtTime returns false PoseHistoryTransformAtTime doesn't contain a valid root bone transform in world space, BUT a valid ComponentToWorldTransform!
+		// So we use PoseHistoryTransformAtTime as such below (in the form of "const FTransform& ComponentToWorldTransform = PoseHistory ? PoseHistoryTransformAtTime : ...)
 	}
 
 	const UObject* AnimContext = GetAnimContext(SampleRole);
@@ -667,7 +671,7 @@ FTransform FSearchContext::GetWorldRootBoneTransformAtTime(float SampleTime, con
 		if (AnimInstance && AnimInstance->CurrentSkeleton)
 		{
 			const FTransform& RootBoneTransform = AnimInstance->CurrentSkeleton->GetReferenceSkeleton().GetRefBonePose()[RootSchemaBoneIdx];
-			const FTransform& ComponentToWorldTransform = AnimInstance->GetSkelMeshComponent()->GetComponentTransform();
+			const FTransform& ComponentToWorldTransform = PoseHistory ? PoseHistoryTransformAtTime : AnimInstance->GetSkelMeshComponent()->GetComponentTransform();
 			return RootBoneTransform * ComponentToWorldTransform;
 		}
 	}
@@ -681,7 +685,7 @@ FTransform FSearchContext::GetWorldRootBoneTransformAtTime(float SampleTime, con
 			if (const USkeletalMesh* SkeletalMesh = SkeletalMeshComponent->GetSkeletalMeshAsset())
 			{
 				const FTransform& RootBoneTransform = SkeletalMesh->GetSkeleton()->GetReferenceSkeleton().GetRefBonePose()[RootSchemaBoneIdx];
-				const FTransform& ComponentToWorldTransform = SkeletalMeshComponent->GetComponentTransform();
+				const FTransform& ComponentToWorldTransform = PoseHistory ? PoseHistoryTransformAtTime : SkeletalMeshComponent->GetComponentTransform();
 				return RootBoneTransform * ComponentToWorldTransform;
 			}
 		}
