PoseSearch - relying on previous frame trajectory for MotionMatching Interaction in case we're not able to generate it. VLOG debugging improvements

diff --git a/Engine/Plugins/Animation/PoseSearch/Source/Runtime/Private/AnimNode_MotionMatchingInteraction.cpp b/Engine/Plugins/Animation/PoseSearch/Source/Runtime/Private/AnimNode_MotionMatchingInteraction.cpp
--- a/Engine/Plugins/Animation/PoseSearch/Source/Runtime/Private/AnimNode_MotionMatchingInteraction.cpp
+++ b/Engine/Plugins/Animation/PoseSearch/Source/Runtime/Private/AnimNode_MotionMatchingInteraction.cpp
@@ -4,13 +4,14 @@
 #include "Animation/AnimInertializationSyncScope.h"
 #include "Animation/AnimInstanceProxy.h"
 #include "Animation/AnimRootMotionProvider.h"
+#include "HAL/IConsoleManager.h"
 #include "PoseSearch/PoseSearchInteractionAsset.h"
 #include "PoseSearch/PoseSearchInteractionLibrary.h"
 #include "PoseSearch/PoseHistoryProvider.h"
 #include "PoseSearch/PoseSearchDatabase.h"
 #include "PoseSearch/PoseSearchLibrary.h"
 #include "PoseSearch/PoseSearchSchema.h"
-#include "HAL/IConsoleManager.h"
+#include "VisualLogger/VisualLogger.h"
 
 #if ENABLE_ANIM_DEBUG
 static bool GVarAnimNodeMotionMatchingInteractionDebug = false;
@@ -149,6 +150,19 @@ void FAnimNode_MotionMatchingInteraction::UpdateAssetPlayer(const FAnimationUpda
 
 			WarpRootMotionTransform = LerpedAlignedActorRootBoneTransform.GetRelativeTransform(RootBoneTrasform);
 			check(WarpRootMotionTransform.IsRotationNormalized());
+
+#if ENABLE_VISUAL_LOG
+			if (FVisualLogger::IsRecording())
+			{
+				static const TCHAR* LogName = TEXT("MotionMatchingInteraction");
+
+				UObject* LogOwner = Context.AnimInstanceProxy->GetAnimInstanceObject();
+
+				UE_VLOG_ARROW(LogOwner, LogName, Display, FullAlignedActorRootBoneTransform.GetLocation(), FullAlignedActorRootBoneTransform.GetLocation() + FullAlignedActorRootBoneTransform.GetRotation().GetForwardVector() * 100, FColor::Red, TEXT(""));
+				UE_VLOG_ARROW(LogOwner, LogName, Display, LerpedAlignedActorRootBoneTransform.GetLocation(), LerpedAlignedActorRootBoneTransform.GetLocation() + LerpedAlignedActorRootBoneTransform.GetRotation().GetForwardVector() * 100, FColor::Blue, TEXT(""));
+				UE_VLOG_ARROW(LogOwner, LogName, Display, RootBoneTrasform.GetLocation(), RootBoneTrasform.GetLocation() + RootBoneTrasform.GetRotation().GetForwardVector() * 100, FColor::Green, TEXT(""));
+			}
+#endif
 		}
 	}
 	else
@@ -209,6 +223,16 @@ void FAnimNode_MotionMatchingInteraction::Evaluate_AnyThread(FPoseContext& Outpu
 	{
 		if (bUseAnimRootMotionProvider)
 		{
+#if ENABLE_ANIM_DEBUG
+			if (GVarAnimNodeMotionMatchingInteractionDebug)
+			{
+				const FTransform ComponentTransform = Output.AnimInstanceProxy->GetComponentTransform();
+				Output.AnimInstanceProxy->AnimDrawDebugCoordinateSystem(FullAlignedActorRootBoneTransform.GetLocation(), FullAlignedActorRootBoneTransform.Rotator(), 15.f, false, 0.f, 0.f, SDPG_Foreground);
+				Output.AnimInstanceProxy->AnimDrawDebugCoordinateSystem(ComponentTransform.GetLocation(), ComponentTransform.Rotator(), 5.f, false, 0.f, 0.f, SDPG_Foreground);
+			}
+
+#endif // ENABLE_ANIM_DEBUG
+		
 			if (const UE::Anim::IAnimRootMotionProvider* RootMotionProvider = UE::Anim::IAnimRootMotionProvider::Get())
 			{
 				RootMotionProvider->OverrideRootMotion(WarpRootMotionTransform, Output.CustomAttributes);
diff --git a/Engine/Plugins/Animation/PoseSearch/Source/Runtime/Private/PoseSearchHistory.cpp b/Engine/Plugins/Animation/PoseSearch/Source/Runtime/Private/PoseSearchHistory.cpp
--- a/Engine/Plugins/Animation/PoseSearch/Source/Runtime/Private/PoseSearchHistory.cpp
+++ b/Engine/Plugins/Animation/PoseSearch/Source/Runtime/Private/PoseSearchHistory.cpp
@@ -824,13 +824,15 @@ void FPoseHistory::GenerateTrajectory(const UAnimInstance* AnimInstance, float D
 {
 	// @todo: Synchronize the FPoseSearchQueryTrajectorySample::AccumulatedSeconds of the generated trajectory with the FPoseHistoryEntry::AccumulatedSeconds of the captured poses
 	FPoseSearchTrajectoryData::FDerived TrajectoryDataDerived;
-	TrajectoryData.UpdateData(DeltaTime, AnimInstance, TrajectoryDataDerived, TrajectoryDataState);
-	UPoseSearchTrajectoryLibrary::InitTrajectorySamples(Trajectory, TrajectoryData, TrajectoryDataDerived.Position, TrajectoryDataDerived.Facing, TrajectoryDataSampling, DeltaTime);
-	UPoseSearchTrajectoryLibrary::UpdateHistory_TransformHistory(Trajectory, TrajectoryData, TrajectoryDataDerived.Position, TrajectoryDataDerived.Velocity, TrajectoryDataSampling, DeltaTime);
-	UPoseSearchTrajectoryLibrary::UpdatePrediction_SimulateCharacterMovement(Trajectory, TrajectoryData, TrajectoryDataDerived, TrajectoryDataSampling, DeltaTime);
+	if (TrajectoryData.UpdateData(DeltaTime, AnimInstance, TrajectoryDataDerived, TrajectoryDataState))
+	{
+		UPoseSearchTrajectoryLibrary::InitTrajectorySamples(Trajectory, TrajectoryData, TrajectoryDataDerived.Position, TrajectoryDataDerived.Facing, TrajectoryDataSampling, DeltaTime);
+		UPoseSearchTrajectoryLibrary::UpdateHistory_TransformHistory(Trajectory, TrajectoryData, TrajectoryDataDerived.Position, TrajectoryDataDerived.Velocity, TrajectoryDataSampling, DeltaTime);
+		UPoseSearchTrajectoryLibrary::UpdatePrediction_SimulateCharacterMovement(Trajectory, TrajectoryData, TrajectoryDataDerived, TrajectoryDataSampling, DeltaTime);
 
-	// @todo: support TrajectorySpeedMultiplier
-	//TrajectorySpeedMultiplier = 1.f;
+		// @todo: support TrajectorySpeedMultiplier
+		//TrajectorySpeedMultiplier = 1.f;
+	}
 }
 
 void FPoseHistory::PreUpdate()
diff --git a/Engine/Plugins/Animation/PoseSearch/Source/Runtime/Private/PoseSearchTrajectoryLibrary.cpp b/Engine/Plugins/Animation/PoseSearch/Source/Runtime/Private/PoseSearchTrajectoryLibrary.cpp
--- a/Engine/Plugins/Animation/PoseSearch/Source/Runtime/Private/PoseSearchTrajectoryLibrary.cpp
+++ b/Engine/Plugins/Animation/PoseSearch/Source/Runtime/Private/PoseSearchTrajectoryLibrary.cpp
@@ -10,16 +10,16 @@
 #include "PoseSearch/PoseSearchDefines.h"
 #include "Kismet/KismetMathLibrary.h"
 
-void FPoseSearchTrajectoryData::UpdateData(
+bool FPoseSearchTrajectoryData::UpdateData(
 	float DeltaTime,
 	const FAnimInstanceProxy& AnimInstanceProxy,
 	FDerived& TrajectoryDataDerived,
 	FState& TrajectoryDataState) const
 {
-	UpdateData(DeltaTime, AnimInstanceProxy.GetAnimInstanceObject(), TrajectoryDataDerived, TrajectoryDataState);
+	return UpdateData(DeltaTime, AnimInstanceProxy.GetAnimInstanceObject(), TrajectoryDataDerived, TrajectoryDataState);
 }
 
-void FPoseSearchTrajectoryData::UpdateData(
+bool FPoseSearchTrajectoryData::UpdateData(
 	float DeltaTime,
 	const UObject* Context,
 	FDerived& TrajectoryDataDerived,
@@ -28,7 +28,7 @@ void FPoseSearchTrajectoryData::UpdateData(
 	// An AnimInstance might call this during an AnimBP recompile with 0 delta time.
 	if (DeltaTime <= 0.f)
 	{
-		return;
+		return false;
 	}
 
 	const ACharacter* Character = Cast<ACharacter>(Context);
@@ -38,61 +38,62 @@ void FPoseSearchTrajectoryData::UpdateData(
 		Character = Cast<ACharacter>(AnimInstance->GetOwningActor());
 		if (!Character)
 		{
-			return;
+			return false;
 		}
 	}
 
-	if (const UCharacterMovementComponent* MoveComp = Character->GetCharacterMovement())
+	const UCharacterMovementComponent* MoveComp = Character->GetCharacterMovement();
+	const USkeletalMeshComponent* MeshComp = Character->GetMesh();
+	if (!MoveComp || !MeshComp)
 	{
-		TrajectoryDataDerived.MaxSpeed = FMath::Max(MoveComp->GetMaxSpeed() * MoveComp->GetAnalogInputModifier(), MoveComp->GetMinAnalogSpeed());
-		TrajectoryDataDerived.BrakingDeceleration = FMath::Max(0.f, MoveComp->GetMaxBrakingDeceleration());
-		TrajectoryDataDerived.BrakingSubStepTime = MoveComp->BrakingSubStepTime;
-		TrajectoryDataDerived.bOrientRotationToMovement = MoveComp->bOrientRotationToMovement;
+		return false;
+	}
 
-		TrajectoryDataDerived.Velocity = MoveComp->Velocity;
-		TrajectoryDataDerived.Acceleration = MoveComp->GetCurrentAcceleration();
+	TrajectoryDataDerived.MaxSpeed = FMath::Max(MoveComp->GetMaxSpeed() * MoveComp->GetAnalogInputModifier(), MoveComp->GetMinAnalogSpeed());
+	TrajectoryDataDerived.BrakingDeceleration = FMath::Max(0.f, MoveComp->GetMaxBrakingDeceleration());
+	TrajectoryDataDerived.BrakingSubStepTime = MoveComp->BrakingSubStepTime;
+	TrajectoryDataDerived.bOrientRotationToMovement = MoveComp->bOrientRotationToMovement;
+
+	TrajectoryDataDerived.Velocity = MoveComp->Velocity;
+	TrajectoryDataDerived.Acceleration = MoveComp->GetCurrentAcceleration();
 		
-		TrajectoryDataDerived.bStepGroundPrediction = !MoveComp->IsFalling() && !MoveComp->IsFlying();
+	TrajectoryDataDerived.bStepGroundPrediction = !MoveComp->IsFalling() && !MoveComp->IsFlying();
 
-		if (TrajectoryDataDerived.Acceleration.IsZero())
-		{
-			TrajectoryDataDerived.Friction = MoveComp->bUseSeparateBrakingFriction ? MoveComp->BrakingFriction : MoveComp->GroundFriction;
-			const float FrictionFactor = FMath::Max(0.f, MoveComp->BrakingFrictionFactor);
-			TrajectoryDataDerived.Friction = FMath::Max(0.f, TrajectoryDataDerived.Friction * FrictionFactor);
-		}
-		else
-		{
-			TrajectoryDataDerived.Friction = MoveComp->GroundFriction;
-		}
+	if (TrajectoryDataDerived.Acceleration.IsZero())
+	{
+		TrajectoryDataDerived.Friction = MoveComp->bUseSeparateBrakingFriction ? MoveComp->BrakingFriction : MoveComp->GroundFriction;
+		const float FrictionFactor = FMath::Max(0.f, MoveComp->BrakingFrictionFactor);
+		TrajectoryDataDerived.Friction = FMath::Max(0.f, TrajectoryDataDerived.Friction * FrictionFactor);
 	}
-
+	else
 	{
-		const float DesiredControllerYaw = Character->GetViewRotation().Yaw;
+		TrajectoryDataDerived.Friction = MoveComp->GroundFriction;
+	}
+
+	const float DesiredControllerYaw = Character->GetViewRotation().Yaw;
 		
-		const float DesiredYawDelta = DesiredControllerYaw - TrajectoryDataState.DesiredControllerYawLastUpdate;
-		TrajectoryDataState.DesiredControllerYawLastUpdate = DesiredControllerYaw;
+	const float DesiredYawDelta = DesiredControllerYaw - TrajectoryDataState.DesiredControllerYawLastUpdate;
+	TrajectoryDataState.DesiredControllerYawLastUpdate = DesiredControllerYaw;
 
-		TrajectoryDataDerived.ControllerYawRate = FRotator::NormalizeAxis(DesiredYawDelta) * (1.f / DeltaTime);
-		if (MaxControllerYawRate >= 0.f)
-		{
-			TrajectoryDataDerived.ControllerYawRate = FMath::Sign(TrajectoryDataDerived.ControllerYawRate) * FMath::Min(FMath::Abs(TrajectoryDataDerived.ControllerYawRate), MaxControllerYawRate);
-		}
+	TrajectoryDataDerived.ControllerYawRate = FRotator::NormalizeAxis(DesiredYawDelta) * (1.f / DeltaTime);
+	if (MaxControllerYawRate >= 0.f)
+	{
+		TrajectoryDataDerived.ControllerYawRate = FMath::Sign(TrajectoryDataDerived.ControllerYawRate) * FMath::Min(FMath::Abs(TrajectoryDataDerived.ControllerYawRate), MaxControllerYawRate);
 	}
 
-	if (const USkeletalMeshComponent* MeshComp = Character->GetMesh())
-	{
-		TrajectoryDataDerived.Position = MeshComp->GetComponentLocation();
-		TrajectoryDataDerived.MeshCompRelativeRotation = MeshComp->GetRelativeRotation().Quaternion();
+	TrajectoryDataDerived.Position = MeshComp->GetComponentLocation();
+	TrajectoryDataDerived.MeshCompRelativeRotation = MeshComp->GetRelativeRotation().Quaternion();
 		
-		if (TrajectoryDataDerived.bOrientRotationToMovement)
-		{
-			TrajectoryDataDerived.Facing = MeshComp->GetComponentRotation().Quaternion();
-		}
-		else
-		{
-			TrajectoryDataDerived.Facing = FQuat::MakeFromRotator(FRotator(0,TrajectoryDataState.DesiredControllerYawLastUpdate,0)) * TrajectoryDataDerived.MeshCompRelativeRotation;
-		}
+	if (TrajectoryDataDerived.bOrientRotationToMovement)
+	{
+		TrajectoryDataDerived.Facing = MeshComp->GetComponentRotation().Quaternion();
+	}
+	else
+	{
+		TrajectoryDataDerived.Facing = FQuat::MakeFromRotator(FRotator(0,TrajectoryDataState.DesiredControllerYawLastUpdate,0)) * TrajectoryDataDerived.MeshCompRelativeRotation;
 	}
+
+	return true;
 }
 
 FVector FPoseSearchTrajectoryData::StepCharacterMovementGroundPrediction(
diff --git a/Engine/Plugins/Animation/PoseSearch/Source/Runtime/Public/PoseSearch/PoseSearchTrajectoryLibrary.h b/Engine/Plugins/Animation/PoseSearch/Source/Runtime/Public/PoseSearch/PoseSearchTrajectoryLibrary.h
--- a/Engine/Plugins/Animation/PoseSearch/Source/Runtime/Public/PoseSearch/PoseSearchTrajectoryLibrary.h
+++ b/Engine/Plugins/Animation/PoseSearch/Source/Runtime/Public/PoseSearch/PoseSearchTrajectoryLibrary.h
@@ -50,8 +50,8 @@ public:
 		float SecondsPerPredictionSample = 0.f;
 	};
 
-	void UpdateData(float DeltaTime, const FAnimInstanceProxy& AnimInstanceProxy, FDerived& TrajectoryDataDerived, FState& TrajectoryDataState) const;
-	void UpdateData(float DeltaTime, const UObject* Context, FDerived& TrajectoryDataDerived, FState& TrajectoryDataState) const;
+	bool UpdateData(float DeltaTime, const FAnimInstanceProxy& AnimInstanceProxy, FDerived& TrajectoryDataDerived, FState& TrajectoryDataState) const;
+	bool UpdateData(float DeltaTime, const UObject* Context, FDerived& TrajectoryDataDerived, FState& TrajectoryDataState) const;
 	FVector StepCharacterMovementGroundPrediction(float DeltaTime, const FVector& InVelocity, const FVector& InAcceleration, const FDerived& TrajectoryDataDerived) const;
 	
 	// If the character is forward facing (i.e. bOrientRotationToMovement is true), this controls how quickly the trajectory will rotate
