PoseSearch - UPoseSearchInteractionAsset warp banking weight

diff --git a/Engine/Plugins/Animation/PoseSearch/Source/Runtime/Private/PoseSearchInteractionAsset.cpp b/Engine/Plugins/Animation/PoseSearch/Source/Runtime/Private/PoseSearchInteractionAsset.cpp
--- a/Engine/Plugins/Animation/PoseSearch/Source/Runtime/Private/PoseSearchInteractionAsset.cpp
+++ b/Engine/Plugins/Animation/PoseSearch/Source/Runtime/Private/PoseSearchInteractionAsset.cpp
@@ -129,23 +129,42 @@ int32 UPoseSearchInteractionAsset::GetFrameAtTime(float Time) const
 }
 #endif // WITH_EDITOR
 
-FQuat UPoseSearchInteractionAsset::FindReferenceOrientation(const TArrayView<const FTransform> Transforms, const TArrayView<int32> SortedByWarpingWeightRotationItemIndex) const
+FQuat UPoseSearchInteractionAsset::FindReferenceOrientationNoBanking(const TArrayView<const FTransform> Transforms, const TArrayView<int32> SortedByWarpingWeightRotationItemIndex, const TArrayView<float> NormalizedWarpingWeightRotation) const
+{
+	check(Items.Num() > 0);
+	check(Items.Num() == Transforms.Num());
+	check(Items.Num() == SortedByWarpingWeightRotationItemIndex.Num());
+	check(Items.Num() == NormalizedWarpingWeightRotation.Num());
+
+	// @todo: use a Slerp or a proper FastLerp
+	// for now we don't account for the shortest path while FastLerping those queternion together
+	FQuat WeightedQuaternion = FQuat::Identity * 0.f;
+	for (int32 ItemIndex : SortedByWarpingWeightRotationItemIndex)
+	{
+		WeightedQuaternion += Transforms[ItemIndex].GetRotation() * NormalizedWarpingWeightRotation[ItemIndex];
+	}
+	WeightedQuaternion.Normalize();
+	return WeightedQuaternion;
+}
+
+FQuat UPoseSearchInteractionAsset::FindReferenceOrientationFullBanking(const TArrayView<const FTransform> Transforms, const TArrayView<int32> SortedByWarpingWeightRotationItemIndex, const TArrayView<float> NormalizedWarpingWeightRotation) const
 {
 	const int32 ItemsNum = Items.Num();
 
 	check(ItemsNum > 0);
-	check(Transforms.Num() == ItemsNum);
-	check(Transforms.Num() == SortedByWarpingWeightRotationItemIndex.Num());
+	check(ItemsNum == Transforms.Num());
+	check(ItemsNum == SortedByWarpingWeightRotationItemIndex.Num());
+	check(ItemsNum == NormalizedWarpingWeightRotation.Num());
 
-	const int32 LastItemIndex = ItemsNum - 1;
 	if (ItemsNum > 1)
 	{
+		const int32 LastItemIndex = ItemsNum - 1;
 		FVector OtherItemsPositionsSum = Transforms[SortedByWarpingWeightRotationItemIndex[0]].GetTranslation();
 		for (int32 ItemIndex = 1; ItemIndex < LastItemIndex; ++ItemIndex)
 		{
 			OtherItemsPositionsSum += Transforms[SortedByWarpingWeightRotationItemIndex[ItemIndex]].GetTranslation();
 		}
-		
+
 		const FVector OtherItemsPositionAverage = OtherItemsPositionsSum / LastItemIndex;
 		const FVector DeltaPosition = OtherItemsPositionAverage - Transforms[SortedByWarpingWeightRotationItemIndex[LastItemIndex]].GetTranslation();
 
@@ -155,7 +174,25 @@ FQuat UPoseSearchInteractionAsset::FindReferenceOrientation(const TArrayView<con
 		}
 	}
 
-	return Transforms[SortedByWarpingWeightRotationItemIndex[LastItemIndex]].GetRotation();
+	return FindReferenceOrientationNoBanking(Transforms, SortedByWarpingWeightRotationItemIndex, NormalizedWarpingWeightRotation);
+}
+
+FQuat UPoseSearchInteractionAsset::FindReferenceOrientation(const TArrayView<const FTransform> Transforms, const TArrayView<int32> SortedByWarpingWeightRotationItemIndex, const TArrayView<float> NormalizedWarpingWeightRotation) const
+{
+	if (WarpingBankingWeight < UE_KINDA_SMALL_NUMBER)
+	{
+		return FindReferenceOrientationNoBanking(Transforms, SortedByWarpingWeightRotationItemIndex, NormalizedWarpingWeightRotation);
+	}
+
+	if (WarpingBankingWeight > 1.f - UE_KINDA_SMALL_NUMBER)
+	{
+		return FindReferenceOrientationFullBanking(Transforms, SortedByWarpingWeightRotationItemIndex, NormalizedWarpingWeightRotation);
+	}
+
+	return FQuat::Slerp(
+		FindReferenceOrientationNoBanking(Transforms, SortedByWarpingWeightRotationItemIndex, NormalizedWarpingWeightRotation),
+		FindReferenceOrientationFullBanking(Transforms, SortedByWarpingWeightRotationItemIndex, NormalizedWarpingWeightRotation),
+		WarpingBankingWeight);
 }
 
 FVector UPoseSearchInteractionAsset::FindReferencePosition(const TArrayView<const FTransform> Transforms, const TArrayView<float> NormalizedWarpingWeightTranslation) const
@@ -276,8 +313,10 @@ void UPoseSearchInteractionAsset::CalculateWarpTransforms(float Time, const TArr
 		}
 
 		TArray<int32, TInlineAllocator<UE::PoseSearch::PreallocatedRolesNum>> SortedByWarpingWeightRotationItemIndex;
+		TArray<float, TInlineAllocator<UE::PoseSearch::PreallocatedRolesNum>> NormalizedWarpingWeightRotation;
 		TArray<float, TInlineAllocator<UE::PoseSearch::PreallocatedRolesNum>> NormalizedWarpingWeightTranslation;
 		SortedByWarpingWeightRotationItemIndex.SetNum(ItemsNum);
+		NormalizedWarpingWeightRotation.SetNum(ItemsNum);
 		NormalizedWarpingWeightTranslation.SetNum(ItemsNum);
 
 		float WarpingWeightTranslationSum = 0.f;
@@ -311,10 +350,22 @@ void UPoseSearchInteractionAsset::CalculateWarpTransforms(float Time, const TArr
 				{
 					return Items[A].WarpingWeightRotation < Items[B].WarpingWeightRotation;
 				});
+
+			for (int32 ItemIndex = 0; ItemIndex < ItemsNum; ++ItemIndex)
+			{
+				NormalizedWarpingWeightRotation[ItemIndex] = Items[ItemIndex].WarpingWeightRotation / WarpingWeightRotationSum;
+			}
+		}
+		else
+		{
+			for (int32 ItemIndex = 0; ItemIndex < ItemsNum; ++ItemIndex)
+			{
+				NormalizedWarpingWeightRotation[ItemIndex] = NormalizedHomogeneousWeight;
+			}
 		}
 
-		const FQuat AssetReferenceOrientation = FindReferenceOrientation(AssetRootBoneTransforms, SortedByWarpingWeightRotationItemIndex);
-		const FQuat ActorsReferenceOrientation = FindReferenceOrientation(ActorRootBoneTransforms, SortedByWarpingWeightRotationItemIndex);
+		const FQuat AssetReferenceOrientation = FindReferenceOrientation(AssetRootBoneTransforms, SortedByWarpingWeightRotationItemIndex, NormalizedWarpingWeightRotation);
+		const FQuat ActorsReferenceOrientation = FindReferenceOrientation(ActorRootBoneTransforms, SortedByWarpingWeightRotationItemIndex, NormalizedWarpingWeightRotation);
 
 		FQuat WeightedActorsReferenceOrientation = ActorsReferenceOrientation;
 		if (WarpingWeightRotationSum > UE_KINDA_SMALL_NUMBER)
@@ -323,15 +374,14 @@ void UPoseSearchInteractionAsset::CalculateWarpTransforms(float Time, const TArr
 			for (int32 ItemIndex : SortedByWarpingWeightRotationItemIndex)
 			{
 				const FPoseSearchInteractionAssetItem& Item = Items[ItemIndex];
-				const float NormalizedWarpingWeightRotation = Item.WarpingWeightRotation / WarpingWeightRotationSum;
-				if (NormalizedWarpingWeightRotation > NormalizedHomogeneousWeight)
+				if (NormalizedWarpingWeightRotation[ItemIndex] > NormalizedHomogeneousWeight)
 				{
 					// NormalizedHomogeneousWeight is one only if ItemsNum is one, 
-					// BUT NormalizedWarpingWeightRotation > NormalizedHomogeneousWeight should always be false
+					// BUT NormalizedWarpingWeightRotation[ItemIndex] > NormalizedHomogeneousWeight should always be false
 					check(!FMath::IsNearlyEqual(NormalizedHomogeneousWeight, 1.f));
 
 					// how much this item wants to reorient the ReferenceOrientation from the homogeneous "fair" value
-					const float SlerpParam = (NormalizedWarpingWeightRotation - NormalizedHomogeneousWeight) / (1.f - NormalizedHomogeneousWeight);
+					const float SlerpParam = (NormalizedWarpingWeightRotation[ItemIndex] - NormalizedHomogeneousWeight) / (1.f - NormalizedHomogeneousWeight);
 
 					// calculating the reference orientation relative to the character
 					// AssetReferenceOrientation in actor world orientation
diff --git a/Engine/Plugins/Animation/PoseSearch/Source/Runtime/Public/PoseSearch/PoseSearchInteractionAsset.h b/Engine/Plugins/Animation/PoseSearch/Source/Runtime/Public/PoseSearch/PoseSearchInteractionAsset.h
--- a/Engine/Plugins/Animation/PoseSearch/Source/Runtime/Public/PoseSearch/PoseSearchInteractionAsset.h
+++ b/Engine/Plugins/Animation/PoseSearch/Source/Runtime/Public/PoseSearch/PoseSearchInteractionAsset.h
@@ -45,6 +45,9 @@ class POSESEARCH_API UPoseSearchInteractionAsset : public UMultiAnimAsset
 	UPROPERTY(EditAnywhere, Category = "Settings")
 	TArray<FPoseSearchInteractionAssetItem> Items;
 
+	UPROPERTY(EditAnywhere, Category = "Settings", meta = (ClampMin = "0", ClampMax = "1"))
+	float WarpingBankingWeight = 0.f;
+
 public:
 #if WITH_EDITORONLY_DATA
 
@@ -83,7 +86,9 @@ public:
 #endif // WITH_EDITOR
 
 	virtual void CalculateWarpTransforms(float Time, const TArrayView<const FTransform> ActorRootBoneTransforms, TArrayView<FTransform> FullAlignedActorRootBoneTransforms) const override;
-
-	FQuat FindReferenceOrientation(const TArrayView<const FTransform> Transforms, const TArrayView<int32> SortedByWarpingWeightRotationItemIndex) const;
+	
 	FVector FindReferencePosition(const TArrayView<const FTransform> Transforms, const TArrayView<float> NormalizedWarpingWeightTranslation) const;
+	FQuat FindReferenceOrientation(const TArrayView<const FTransform> Transforms, const TArrayView<int32> SortedByWarpingWeightRotationItemIndex, const TArrayView<float> NormalizedWarpingWeightRotation) const;
+	FQuat FindReferenceOrientationNoBanking(const TArrayView<const FTransform> Transforms, const TArrayView<int32> SortedByWarpingWeightRotationItemIndex, const TArrayView<float> NormalizedWarpingWeightRotation) const;
+	FQuat FindReferenceOrientationFullBanking(const TArrayView<const FTransform> Transforms, const TArrayView<int32> SortedByWarpingWeightRotationItemIndex, const TArrayView<float> NormalizedWarpingWeightRotation) const;
 };
\ No newline at end of file
