PoseSearch
- fix for FAnimNode_PoseSearchHistoryCollector_Base::Initialize_AnyThread initializing PoseHistory with a valid pose and the proper RequiredBones
- deprecating no longer necessary FAnimNode_PoseSearchHistoryCollector_Base::bInitializeWithRefPose
- removing unnecessary FPoseHistory double buffering of the PoseData and related PreUpdate, since now handled properly by injecting tick dependencies via UE::PoseSearch::FInteractionIsland

diff --git a/Engine/Plugins/Animation/PoseSearch/Source/Runtime/Private/AnimNode_PoseSearchHistoryCollector.cpp b/Engine/Plugins/Animation/PoseSearch/Source/Runtime/Private/AnimNode_PoseSearchHistoryCollector.cpp
--- a/Engine/Plugins/Animation/PoseSearch/Source/Runtime/Private/AnimNode_PoseSearchHistoryCollector.cpp
+++ b/Engine/Plugins/Animation/PoseSearch/Source/Runtime/Private/AnimNode_PoseSearchHistoryCollector.cpp
@@ -19,23 +19,12 @@ void FAnimNode_PoseSearchHistoryCollector_Base::Initialize_AnyThread(const FAnim
 
 	PoseHistory.Initialize_AnyThread(PoseCount, SamplingInterval);
 
-	if (bInitializeWithRefPose)
-	{
-		const FBoneContainer& BoneContainer = Context.AnimInstanceProxy->GetRequiredBones();
-		if (BoneContainer.IsValid())
-		{
-			// initializing PoseHistory with a ref pose at FAnimInstanceProxy location/facing
-			FMemMark Mark(FMemStack::Get());
-			FCSPose<FCompactPose> ComponentSpacePose;
-			FBlendedCurve EmptyCurves;
-			ComponentSpacePose.InitPose(&BoneContainer);
-			UE::PoseSearch::FComponentSpacePoseProvider ComponentSpacePoseProvider(ComponentSpacePose);
+	const TArray<FBoneIndexType> RequiredBones = GetRequiredBones(Context.AnimInstanceProxy);
+	UE::PoseSearch::FAIPComponentSpacePoseProvider ComponentSpacePoseProvider(Context.AnimInstanceProxy);
 			PoseHistory.EvaluateComponentSpace_AnyThread(0.f, ComponentSpacePoseProvider, bStoreScales,
 				RootBoneRecoveryTime, RootBoneTranslationRecoveryRatio, RootBoneRotationRecoveryRatio, true, true, 
-				GetRequiredBones(Context.AnimInstanceProxy), EmptyCurves, MakeConstArrayView(CollectedCurves));
+		RequiredBones, FBlendedCurve(), MakeConstArrayView(CollectedCurves));
 		}
-	}
-}
 
 TArray<FBoneIndexType> FAnimNode_PoseSearchHistoryCollector_Base::GetRequiredBones(const FAnimInstanceProxy* AnimInstanceProxy) const
 {
@@ -86,20 +75,6 @@ void FAnimNode_PoseSearchHistoryCollector_Base::Update_AnyThread(const FAnimatio
 	UpdateCounter.SynchronizeWith(Context.AnimInstanceProxy->GetUpdateCounter());
 }
 
-void FAnimNode_PoseSearchHistoryCollector_Base::PreUpdate(const UAnimInstance* InAnimInstance)
-{
-	Super::PreUpdate(InAnimInstance);
-	
-	if (bGenerateTrajectory)
-	{
-		GenerateTrajectory(InAnimInstance);
-	}
-
-	PoseHistory.PreUpdate();
-
-	bIsTrajectoryGeneratedBeforePreUpdate = false;
-}
-
 void FAnimNode_PoseSearchHistoryCollector_Base::GenerateTrajectory(const UAnimInstance* InAnimInstance)
 {
 	if (!bIsTrajectoryGeneratedBeforePreUpdate)
@@ -175,6 +150,13 @@ void FAnimNode_PoseSearchHistoryCollector::Evaluate_AnyThread(FPoseContext& Outp
 void FAnimNode_PoseSearchHistoryCollector::Update_AnyThread(const FAnimationUpdateContext& Context)
 {
 	DECLARE_SCOPE_HIERARCHICAL_COUNTER_ANIMNODE(Update_AnyThread);
+
+	if (bGenerateTrajectory)
+	{
+		GenerateTrajectory(Cast<UAnimInstance>(Context.AnimInstanceProxy->GetAnimInstanceObject()));
+	}
+	bIsTrajectoryGeneratedBeforePreUpdate = false;
+
 	Super::Update_AnyThread(Context);
 	UE::Anim::TScopedGraphMessage<UE::PoseSearch::FPoseHistoryProvider> ScopedMessage(Context, this);
 	Source.Update(Context);
@@ -243,6 +225,13 @@ void FAnimNode_PoseSearchComponentSpaceHistoryCollector::EvaluateComponentSpace_
 void FAnimNode_PoseSearchComponentSpaceHistoryCollector::Update_AnyThread(const FAnimationUpdateContext& Context)
 {
 	DECLARE_SCOPE_HIERARCHICAL_COUNTER_ANIMNODE(Update_AnyThread);
+
+	if (bGenerateTrajectory)
+	{
+		GenerateTrajectory(Cast<UAnimInstance>(Context.AnimInstanceProxy->GetAnimInstanceObject()));
+	}
+	bIsTrajectoryGeneratedBeforePreUpdate = false;
+
 	Super::Update_AnyThread(Context);
 	UE::Anim::TScopedGraphMessage<UE::PoseSearch::FPoseHistoryProvider> ScopedMessage(Context, this);
 	Source.Update(Context);
diff --git a/Engine/Plugins/Animation/PoseSearch/Source/Runtime/Private/PoseSearchHistory.cpp b/Engine/Plugins/Animation/PoseSearch/Source/Runtime/Private/PoseSearchHistory.cpp
--- a/Engine/Plugins/Animation/PoseSearch/Source/Runtime/Private/PoseSearchHistory.cpp
+++ b/Engine/Plugins/Animation/PoseSearch/Source/Runtime/Private/PoseSearchHistory.cpp
@@ -9,6 +9,7 @@
 #include "Animation/SkeletonRemappingRegistry.h"
 #include "BonePose.h"
 #include "DrawDebugHelpers.h"
+#include "Engine/SkeletalMesh.h"
 #include "PoseSearch/PoseSearchResult.h"
 #include "PoseSearch/PoseSearchDatabase.h"
 #include "PoseSearch/PoseSearchDefines.h"
@@ -252,38 +253,6 @@ static uint32 GetTypeHash(const FBoneToTransformMap& BoneToTransformMap)
 	return TypeHash;
 }
 
-// optimized copy for TRingBuffer<FPoseHistoryEntry> implementing "To = From" to avoid allocations as much as possible
-static void CopyEntries(const TRingBuffer<FPoseHistoryEntry>& From, TRingBuffer<FPoseHistoryEntry>& To)
-{
-	const int32 FromNum = From.Num();
-	const int32 PopCount = To.Num() - FromNum;
-	if (PopCount != 0)
-	{
-		if (PopCount > 0)
-		{
-			To.Pop(PopCount);
-		}
-		else
-		{
-			To.Reserve(FromNum);
-			const int32 AddCount = -PopCount;
-			for (int32 Index = 0; Index < AddCount; ++Index)
-			{
-				To.Add(FPoseHistoryEntry());
-			}
-		}
-	}
-
-	check(To.Num() == FromNum);
-
-	TRingBuffer<FPoseHistoryEntry>::TIterator ToIt = To.begin();
-	const TRingBuffer<FPoseHistoryEntry>::TConstIterator FromEnd = From.end();
-	for (TRingBuffer<FPoseHistoryEntry>::TConstIterator FromIt = From.begin(); FromIt != FromEnd; ++FromIt, ++ToIt)
-	{
-		*ToIt = *FromIt;
-	}
-}
-
 //////////////////////////////////////////////////////////////////////////
 // FComponentSpacePoseProvider
 FComponentSpacePoseProvider::FComponentSpacePoseProvider(FCSPose<FCompactPose>& InComponentSpacePose)
@@ -318,6 +287,25 @@ const USkeleton* FComponentSpacePoseProvider::GetSkeletonAsset() const
 	return ComponentSpacePose.GetPose().GetBoneContainer().GetSkeletonAsset();
 }
 
+//////////////////////////////////////////////////////////////////////////
+// FAIPComponentSpacePoseProvider
+FAIPComponentSpacePoseProvider::FAIPComponentSpacePoseProvider(const FAnimInstanceProxy* InAnimInstanceProxy)
+	: AnimInstanceProxy(InAnimInstanceProxy)
+{
+	check(AnimInstanceProxy && GetSkeletonAsset());
+}
+
+FTransform FAIPComponentSpacePoseProvider::CalculateComponentSpaceTransform(const FSkeletonPoseBoneIndex SkeletonBoneIdx)
+{
+	const FMeshPoseBoneIndex MeshPoseBoneIndex = AnimInstanceProxy->GetRequiredBones().GetMeshPoseIndexFromSkeletonPoseIndex(SkeletonBoneIdx);
+	return AnimInstanceProxy->GetSkelMeshComponent()->GetBoneTransform(MeshPoseBoneIndex.GetInt(), FTransform::Identity);
+}
+
+const USkeleton* FAIPComponentSpacePoseProvider::GetSkeletonAsset() const
+{
+	return AnimInstanceProxy->GetSkelMeshComponent()->GetSkeletalMeshAsset()->GetSkeleton();
+}
+
 //////////////////////////////////////////////////////////////////////////
 // FPoseHistoryEntry
 void FPoseHistoryEntry::Update(float Time, IComponentSpacePoseProvider& ComponentSpacePoseProvider, const FBoneToTransformMap& BoneToTransformMap, bool bStoreScales, const FBlendedCurve& Curves, const TConstArrayView<FName>& CollectedCurves)
@@ -653,8 +641,7 @@ FPoseHistory& FPoseHistory::operator=(const FPoseHistory& Other)
 	TrajectoryDataState = Other.TrajectoryDataState;
 	TrajectorySpeedMultiplier = Other.TrajectorySpeedMultiplier;
 
-	ReadPoseDataIndex = Other.ReadPoseDataIndex;
-	DoubleBufferedPoseData = Other.DoubleBufferedPoseData;
+	PoseData = Other.PoseData;
 	return *this;
 }
 
@@ -677,8 +664,7 @@ FPoseHistory& FPoseHistory::operator=(FPoseHistory&& Other)
 	TrajectoryDataState = MoveTemp(Other.TrajectoryDataState);
 	TrajectorySpeedMultiplier = MoveTemp(Other.TrajectorySpeedMultiplier);
 
-	ReadPoseDataIndex = MoveTemp(Other.ReadPoseDataIndex);
-	DoubleBufferedPoseData = MoveTemp(Other.DoubleBufferedPoseData);
+	PoseData = MoveTemp(Other.PoseData);
 	return *this;
 }
 
@@ -694,8 +680,7 @@ void FPoseHistory::Initialize_AnyThread(int32 InNumPoses, float InSamplingInterv
 	TrajectoryDataState = FPoseSearchTrajectoryData::FState();
 	TrajectorySpeedMultiplier = 1.f;
 
-	ReadPoseDataIndex = 0;
-	DoubleBufferedPoseData = FDoubleBufferedPoseData();
+	PoseData = FPoseData();
 }
 
 bool FPoseHistory::GetTransformAtTime(float Time, FTransform& OutBoneTransform, const USkeleton* BoneIndexSkeleton, FBoneIndexType BoneIndexType, FBoneIndexType ReferenceBoneIndexType, bool bExtrapolate) const
@@ -715,8 +700,7 @@ bool FPoseHistory::GetTransformAtTime(float Time, FTransform& OutBoneTransform,
 		ReferenceBoneIndexType = ComponentSpaceIndexType;
 	}
 
-	const FPoseData& ReadPoseData = GetReadPoseData();
-	const int32 NumEntries = ReadPoseData.Entries.Num();
+	const int32 NumEntries = PoseData.Entries.Num();
 	if (NumEntries > 0)
 	{
 		int32 NextIdx = 0;
@@ -724,15 +708,15 @@ bool FPoseHistory::GetTransformAtTime(float Time, FTransform& OutBoneTransform,
 
 		if (NumEntries > 1)
 		{
-			const int32 LowerBoundIdx = LowerBound(ReadPoseData.Entries.begin(), ReadPoseData.Entries.end(), Time, [](const FPoseHistoryEntry& Entry, float Value) { return Value > Entry.AccumulatedSeconds; });
+			const int32 LowerBoundIdx = LowerBound(PoseData.Entries.begin(), PoseData.Entries.end(), Time, [](const FPoseHistoryEntry& Entry, float Value) { return Value > Entry.AccumulatedSeconds; });
 			NextIdx = FMath::Clamp(LowerBoundIdx, 1, NumEntries - 1);
 			PrevIdx = NextIdx - 1;
 		}
 	
-		const FPoseHistoryEntry& PrevEntry = ReadPoseData.Entries[PrevIdx];
-		const FPoseHistoryEntry& NextEntry = ReadPoseData.Entries[NextIdx];
+		const FPoseHistoryEntry& PrevEntry = PoseData.Entries[PrevIdx];
+		const FPoseHistoryEntry& NextEntry = PoseData.Entries[NextIdx];
 
-		bSuccess = LerpEntries(Time, bExtrapolate, PrevEntry, NextEntry, BoneIndexSkeleton, ReadPoseData.LastUpdateSkeleton.Get(), ReadPoseData.BoneToTransformMap, BoneIndexType, ReferenceBoneIndexType, OutBoneTransform);
+		bSuccess = LerpEntries(Time, bExtrapolate, PrevEntry, NextEntry, BoneIndexSkeleton, PoseData.LastUpdateSkeleton.Get(), PoseData.BoneToTransformMap, BoneIndexType, ReferenceBoneIndexType, OutBoneTransform);
 		if (bApplyComponentToWorld)
 		{
 			OutBoneTransform *= ComponentToWorld;
@@ -751,8 +735,7 @@ bool FPoseHistory::GetCurveValueAtTime(float Time, const FName& CurveName, float
 	CheckThreadSafetyRead(ReadPoseDataThreadSafeCounter);
 
 	bool bSuccess = false;
-	const FPoseData& ReadPoseData = GetReadPoseData();
-	const int32 NumEntries = ReadPoseData.Entries.Num();
+	const int32 NumEntries = PoseData.Entries.Num();
 	if (NumEntries > 0)
 	{
 		int32 NextIdx = 0;
@@ -760,15 +743,15 @@ bool FPoseHistory::GetCurveValueAtTime(float Time, const FName& CurveName, float
 
 		if (NumEntries > 1)
 		{
-			const int32 LowerBoundIdx = LowerBound(ReadPoseData.Entries.begin(), ReadPoseData.Entries.end(), Time, [](const FPoseHistoryEntry& Entry, float Value) { return Value > Entry.AccumulatedSeconds; });
+			const int32 LowerBoundIdx = LowerBound(PoseData.Entries.begin(), PoseData.Entries.end(), Time, [](const FPoseHistoryEntry& Entry, float Value) { return Value > Entry.AccumulatedSeconds; });
 			NextIdx = FMath::Clamp(LowerBoundIdx, 1, NumEntries - 1);
 			PrevIdx = NextIdx - 1;
 		}
 
-		const FPoseHistoryEntry& PrevEntry = ReadPoseData.Entries[PrevIdx];
-		const FPoseHistoryEntry& NextEntry = ReadPoseData.Entries[NextIdx];
+		const FPoseHistoryEntry& PrevEntry = PoseData.Entries[PrevIdx];
+		const FPoseHistoryEntry& NextEntry = PoseData.Entries[NextIdx];
 
-		bSuccess = LerpEntries(Time, bExtrapolate, PrevEntry, NextEntry, CurveName, ReadPoseData.CollectedCurves, OutCurveValue);
+		bSuccess = LerpEntries(Time, bExtrapolate, PrevEntry, NextEntry, CurveName, PoseData.CollectedCurves, OutCurveValue);
 	}
 	else
 	{
@@ -793,31 +776,31 @@ float FPoseHistory::GetTrajectorySpeedMultiplier() const
 bool FPoseHistory::IsEmpty() const
 {
 	CheckThreadSafetyRead(ReadPoseDataThreadSafeCounter);
-	return GetReadPoseData().Entries.IsEmpty();
+	return PoseData.Entries.IsEmpty();
 }
 
 const FBoneToTransformMap& FPoseHistory::GetBoneToTransformMap() const
 {
 	CheckThreadSafetyRead(ReadPoseDataThreadSafeCounter);
-	return GetReadPoseData().BoneToTransformMap;
+	return PoseData.BoneToTransformMap;
 }
 
 const TConstArrayView<FName> FPoseHistory::GetCollectedCurves() const
 {
 	CheckThreadSafetyRead(ReadPoseDataThreadSafeCounter);
-	return MakeConstArrayView(GetReadPoseData().CollectedCurves);
+	return MakeConstArrayView(PoseData.CollectedCurves);
 }
 
 int32 FPoseHistory::GetNumEntries() const
 {
 	CheckThreadSafetyRead(ReadPoseDataThreadSafeCounter);
-	return GetReadPoseData().Entries.Num();
+	return PoseData.Entries.Num();
 }
 
 const FPoseHistoryEntry& FPoseHistory::GetEntry(int32 EntryIndex) const
 {
 	CheckThreadSafetyRead(ReadPoseDataThreadSafeCounter);
-	return GetReadPoseData().Entries[EntryIndex];
+	return PoseData.Entries[EntryIndex];
 }
 
 void FPoseHistory::GenerateTrajectory(const UAnimInstance* AnimInstance, float DeltaTime, const FPoseSearchTrajectoryData& TrajectoryData, const FPoseSearchTrajectoryData::FSampling& TrajectoryDataSampling)
@@ -835,15 +818,6 @@ void FPoseHistory::GenerateTrajectory(const UAnimInstance* AnimInstance, float D
 	}
 }
 
-void FPoseHistory::PreUpdate()
-{
-	// checking for thread safety
-	CheckThreadSafetyWrite(ReadPoseDataThreadSafeCounter);
-	CheckThreadSafetyWrite(WritePoseDataThreadSafeCounter);
-
-	ReadPoseDataIndex = GetWritePoseDataIndex();
-}
-
 void FPoseHistory::SetTrajectory(const FPoseSearchQueryTrajectory& InTrajectory, float InTrajectorySpeedMultiplier)
 {
 	if (!InTrajectory.Samples.IsEmpty())
@@ -877,87 +851,75 @@ void FPoseHistory::EvaluateComponentSpace_AnyThread(float DeltaTime, IComponentS
 	check(MaxNumPoses >= 2);
 
 	const USkeleton* Skeleton = ComponentSpacePoseProvider.GetSkeletonAsset();
-	const FPoseData& ReadPoseData = GetReadPoseData();
-	FPoseData& WritePoseData = EditWritePoseData();
-
-	WritePoseData.LastUpdateSkeleton = ReadPoseData.LastUpdateSkeleton;
 
 	if (bCacheBones)
 	{
-		WritePoseData.BoneToTransformMap.Reset();
-		WritePoseData.CollectedCurves = CollectedCurves;
+		const uint32 OldBoneToTransformMapTypeHash = PoseData.BoneToTransformMapTypeHash;
+
+		PoseData.BoneToTransformMap.Reset();
+		PoseData.CollectedCurves = CollectedCurves;
 		if (!RequiredBones.IsEmpty())
 		{
 			// making sure we always collect the root bone transform (by construction BoneToTransformMap[0] = 0)
 			const FComponentSpaceTransformIndex ComponentSpaceTransformRootBoneIndex = 0;
-			WritePoseData.BoneToTransformMap.Add(RootBoneIndexType) = ComponentSpaceTransformRootBoneIndex;
+			PoseData.BoneToTransformMap.Add(RootBoneIndexType) = ComponentSpaceTransformRootBoneIndex;
 
 			for (int32 i = 0; i < RequiredBones.Num(); ++i)
 			{
 				// adding only unique RequiredBones to avoid oversizing Entries::ComponentSpaceTransforms
-				if (!WritePoseData.BoneToTransformMap.Find(RequiredBones[i]))
+				if (!PoseData.BoneToTransformMap.Find(RequiredBones[i]))
 				{
-					const FComponentSpaceTransformIndex ComponentSpaceTransformIndex = WritePoseData.BoneToTransformMap.Num();
-					WritePoseData.BoneToTransformMap.Add(RequiredBones[i]) = ComponentSpaceTransformIndex;
+					const FComponentSpaceTransformIndex ComponentSpaceTransformIndex = PoseData.BoneToTransformMap.Num();
+					PoseData.BoneToTransformMap.Add(RequiredBones[i]) = ComponentSpaceTransformIndex;
 				}
 			}
 		}
 
-		WritePoseData.BoneToTransformMapTypeHash = GetTypeHash(WritePoseData.BoneToTransformMap);
-		bNeedsReset |= (WritePoseData.BoneToTransformMapTypeHash != ReadPoseData.BoneToTransformMapTypeHash);
-	}
-	else if (WritePoseData.BoneToTransformMapTypeHash != ReadPoseData.BoneToTransformMapTypeHash)
-	{
-		WritePoseData.BoneToTransformMap = ReadPoseData.BoneToTransformMap;
-		WritePoseData.BoneToTransformMapTypeHash = ReadPoseData.BoneToTransformMapTypeHash;
-		bNeedsReset = true;
+		PoseData.BoneToTransformMapTypeHash = GetTypeHash(PoseData.BoneToTransformMap);
+		bNeedsReset |= (OldBoneToTransformMapTypeHash != PoseData.BoneToTransformMapTypeHash);
 	}
 
-	if (WritePoseData.LastUpdateSkeleton != Skeleton)
+	if (PoseData.LastUpdateSkeleton != Skeleton)
 	{
 		bNeedsReset = true;
-		WritePoseData.LastUpdateSkeleton = Skeleton;
+		PoseData.LastUpdateSkeleton = Skeleton;
 	}
 
 	if (bNeedsReset)
 	{
-		WritePoseData.Entries.Reset();
-		WritePoseData.Entries.Reserve(MaxNumPoses);
-	}
-	else
-	{
-		CopyEntries(ReadPoseData.Entries, WritePoseData.Entries);
+		PoseData.Entries.Reset();
+		PoseData.Entries.Reserve(MaxNumPoses);
 	}
 
 	FPoseHistoryEntry FutureEntryTemp;
-	if (!WritePoseData.Entries.IsEmpty() && WritePoseData.Entries.Last().AccumulatedSeconds > 0.f)
+	if (!PoseData.Entries.IsEmpty() && PoseData.Entries.Last().AccumulatedSeconds > 0.f)
 	{
 		// removing the "future" root bone Entry
-		FutureEntryTemp = MoveTemp(WritePoseData.Entries.Last());
-		WritePoseData.Entries.Pop();
+		FutureEntryTemp = MoveTemp(PoseData.Entries.Last());
+		PoseData.Entries.Pop();
 	}
 
 	// Age our elapsed times
-	for (FPoseHistoryEntry& Entry : WritePoseData.Entries)
+	for (FPoseHistoryEntry& Entry : PoseData.Entries)
 	{
 		Entry.AccumulatedSeconds -= DeltaTime;
 	}
 
-	if (WritePoseData.Entries.Num() != MaxNumPoses)
+	if (PoseData.Entries.Num() != MaxNumPoses)
 	{
 		// Consume every pose until the queue is full
-		WritePoseData.Entries.Emplace();
+		PoseData.Entries.Emplace();
 	}
-	else if (SamplingInterval <= 0.f || WritePoseData.Entries[WritePoseData.Entries.Num() - 2].AccumulatedSeconds <= -SamplingInterval)
+	else if (SamplingInterval <= 0.f || PoseData.Entries[PoseData.Entries.Num() - 2].AccumulatedSeconds <= -SamplingInterval)
 	{
-		FPoseHistoryEntry EntryTemp = MoveTemp(WritePoseData.Entries.First());
-		WritePoseData.Entries.PopFront();
-		WritePoseData.Entries.Emplace(MoveTemp(EntryTemp));
+		FPoseHistoryEntry EntryTemp = MoveTemp(PoseData.Entries.First());
+		PoseData.Entries.PopFront();
+		PoseData.Entries.Emplace(MoveTemp(EntryTemp));
 	}
 
 	// Regardless of the retention policy, we always update the most recent Entry
-	FPoseHistoryEntry& MostRecentEntry = WritePoseData.Entries.Last();
-	MostRecentEntry.Update(0.f, ComponentSpacePoseProvider, WritePoseData.BoneToTransformMap, bStoreScales, Curves, WritePoseData.CollectedCurves);
+	FPoseHistoryEntry& MostRecentEntry = PoseData.Entries.Last();
+	MostRecentEntry.Update(0.f, ComponentSpacePoseProvider, PoseData.BoneToTransformMap, bStoreScales, Curves, PoseData.CollectedCurves);
 
 	if (RootBoneRecoveryTime > 0.f && !Trajectory.Samples.IsEmpty())
 	{
@@ -979,7 +941,7 @@ void FPoseHistory::EvaluateComponentSpace_AnyThread(float DeltaTime, IComponentS
 		FutureEntryTemp.SetNum(1, bStoreScales);
 		FutureEntryTemp.SetComponentSpaceTransform(RootBoneIndexType, RootBoneTransformAtRecoveryTime);
 		FutureEntryTemp.AccumulatedSeconds = RootBoneRecoveryTime;
-		WritePoseData.Entries.Emplace(MoveTemp(FutureEntryTemp));
+		PoseData.Entries.Emplace(MoveTemp(FutureEntryTemp));
 	}
 }
 
@@ -1011,10 +973,9 @@ void FPoseHistory::DebugDraw(FAnimInstanceProxy& AnimInstanceProxy, FColor Color
 		const bool bValidTrajectory = !Trajectory.Samples.IsEmpty();
 		TArray<FTransform, TInlineAllocator<128>> PrevGlobalTransforms;
 
-		const FPoseData& ReadPoseData = GetReadPoseData();
-		for (int32 EntryIndex = 0; EntryIndex < ReadPoseData.Entries.Num(); ++EntryIndex)
+		for (int32 EntryIndex = 0; EntryIndex < PoseData.Entries.Num(); ++EntryIndex)
 		{
-			const FPoseHistoryEntry& Entry = ReadPoseData.Entries[EntryIndex];
+			const FPoseHistoryEntry& Entry = PoseData.Entries[EntryIndex];
 
 			const int32 PrevGlobalTransformsNum = PrevGlobalTransforms.Num();
 			const int32 Max = FMath::Max(PrevGlobalTransformsNum, Entry.Num());
diff --git a/Engine/Plugins/Animation/PoseSearch/Source/Runtime/Public/PoseSearch/AnimNode_PoseSearchHistoryCollector.h b/Engine/Plugins/Animation/PoseSearch/Source/Runtime/Public/PoseSearch/AnimNode_PoseSearchHistoryCollector.h
--- a/Engine/Plugins/Animation/PoseSearch/Source/Runtime/Public/PoseSearch/AnimNode_PoseSearchHistoryCollector.h
+++ b/Engine/Plugins/Animation/PoseSearch/Source/Runtime/Public/PoseSearch/AnimNode_PoseSearchHistoryCollector.h
@@ -28,9 +28,8 @@ public:
 	UPROPERTY(EditAnywhere, Category = Settings)
 	TArray<FName> CollectedCurves;
 
-	// if true, the pose history will be initialized with a ref pose at the location and orientation of the AnimInstance.
-	UPROPERTY(EditAnywhere, Category = Settings)
-	bool bInitializeWithRefPose = false;
+	UPROPERTY()
+	bool bInitializeWithRefPose_DEPRECATED = false;
 
 	// Reset the pose history if it has become relevant to the graph after not being updated on previous frames.
 	UPROPERTY(EditAnywhere, Category = Settings)
@@ -98,8 +97,6 @@ public:
 	virtual void Initialize_AnyThread(const FAnimationInitializeContext& Context) override;
 	virtual void CacheBones_AnyThread(const FAnimationCacheBonesContext& Context) override;
 	virtual void Update_AnyThread(const FAnimationUpdateContext& Context) override;
-	virtual bool HasPreUpdate() const { return true; }
-	virtual void PreUpdate(const UAnimInstance* InAnimInstance) override;
 	// End of FAnimNode_Base interface
 
 	const UE::PoseSearch::FPoseHistory& GetPoseHistory() const { return PoseHistory; }
diff --git a/Engine/Plugins/Animation/PoseSearch/Source/Runtime/Public/PoseSearch/PoseSearchHistory.h b/Engine/Plugins/Animation/PoseSearch/Source/Runtime/Public/PoseSearch/PoseSearchHistory.h
--- a/Engine/Plugins/Animation/PoseSearch/Source/Runtime/Public/PoseSearch/PoseSearchHistory.h
+++ b/Engine/Plugins/Animation/PoseSearch/Source/Runtime/Public/PoseSearch/PoseSearchHistory.h
@@ -40,6 +40,17 @@ private:
 	FCSPose<FCompactPose>& ComponentSpacePose;
 };
 
+struct FAIPComponentSpacePoseProvider : public IComponentSpacePoseProvider
+{
+	POSESEARCH_API FAIPComponentSpacePoseProvider(const FAnimInstanceProxy* InAnimInstanceProxy);
+
+	virtual FTransform CalculateComponentSpaceTransform(const FSkeletonPoseBoneIndex SkeletonBoneIdx) override;
+	virtual const USkeleton* GetSkeletonAsset() const override;
+
+private:
+	const FAnimInstanceProxy* AnimInstanceProxy = nullptr;
+};
+
 struct FPoseHistoryEntry
 {
 	// collected bones transforms in component space
@@ -127,7 +138,9 @@ struct POSESEARCH_API FPoseHistory : public IPoseHistory
 	FPoseHistory& operator=(FPoseHistory&& Other);
 
 	void GenerateTrajectory(const UAnimInstance* AnimInstance, float DeltaTime, const FPoseSearchTrajectoryData& TrajectoryData, const FPoseSearchTrajectoryData::FSampling& TrajectoryDataSampling);
-	void PreUpdate();
+	
+	UE_DEPRECATED(5.6, "No longer necessary method")
+	void PreUpdate() {}
 
 	void Initialize_AnyThread(int32 InNumPoses, float InSamplingInterval);
 	void EvaluateComponentSpace_AnyThread(float DeltaTime, IComponentSpacePoseProvider& ComponentSpacePoseProvider, bool bStoreScales,
@@ -192,16 +205,7 @@ private:
 
 		// ring buffer of collected bones
 		TRingBuffer<FPoseHistoryEntry> Entries;
-	};
-
-	
-	typedef TStaticArray<FPoseData, 2> FDoubleBufferedPoseData;
-	FDoubleBufferedPoseData DoubleBufferedPoseData;
-	int32 ReadPoseDataIndex = 0;
-
-	int32 GetWritePoseDataIndex() const { return (ReadPoseDataIndex + 1) % 2; }
-	const FPoseData& GetReadPoseData() const { return DoubleBufferedPoseData[ReadPoseDataIndex]; }
-	FPoseData& EditWritePoseData() { return DoubleBufferedPoseData[GetWritePoseDataIndex()]; }
+	} PoseData;
 
 #if ENABLE_ANIM_DEBUG
 	
diff --git a/Engine/Plugins/Experimental/AnimNextPoseSearch/Source/AnimNextPoseSearch/Private/HistoryCollectorTrait.cpp b/Engine/Plugins/Experimental/AnimNextPoseSearch/Source/AnimNextPoseSearch/Private/HistoryCollectorTrait.cpp
--- a/Engine/Plugins/Experimental/AnimNextPoseSearch/Source/AnimNextPoseSearch/Private/HistoryCollectorTrait.cpp
+++ b/Engine/Plugins/Experimental/AnimNextPoseSearch/Source/AnimNextPoseSearch/Private/HistoryCollectorTrait.cpp
@@ -167,8 +167,6 @@ namespace UE::AnimNext
 		// @todo: implement bResetOnBecomingRelevant
 		//bool bResetOnBecomingRelevant = SharedData->GetbResetOnBecomingRelevant(Binding);
 
-		InstanceData->PoseHistory.PreUpdate();
-
 		if (SharedData->GetbGenerateTrajectory(Binding))
 		{
 			// @todo: WIP trajectory generation
