PoseSearch - interactions warping improvements

diff --git a/Engine/Plugins/Animation/PoseSearch/Source/Editor/Private/PoseSearchInteractionAssetEditor.cpp b/Engine/Plugins/Animation/PoseSearch/Source/Editor/Private/PoseSearchInteractionAssetEditor.cpp
--- a/Engine/Plugins/Animation/PoseSearch/Source/Editor/Private/PoseSearchInteractionAssetEditor.cpp
+++ b/Engine/Plugins/Animation/PoseSearch/Source/Editor/Private/PoseSearchInteractionAssetEditor.cpp
@@ -28,7 +28,7 @@ namespace UE::PoseSearch
 
 /////////////////////////////////////////////////
 // class FInteractionAssetEdMode
-bool FInteractionAssetPreviewActor::SpawnPreviewActor(UWorld* World, const UPoseSearchInteractionAsset* InteractionAsset, const UE::PoseSearch::FRole& Role)
+bool FInteractionAssetPreviewActor::SpawnPreviewActor(UWorld* World, const UPoseSearchInteractionAsset* InteractionAsset, const FRole& Role)
 {
 	UAnimationAsset* PreviewAsset = InteractionAsset->GetAnimationAsset(Role);
 	if (!PreviewAsset)
@@ -43,7 +43,7 @@ bool FInteractionAssetPreviewActor::SpawnPreviewActor(UWorld* World, const UPose
 
 	for (int32 ActorIndex = 0; ActorIndex < NumActors; ++ActorIndex)
 	{
-		UE::PoseSearch::FAnimationAssetSampler& Sampler = Samplers[ActorIndex];
+		FAnimationAssetSampler& Sampler = Samplers[ActorIndex];
 
 		const FTransform Origin = InteractionAsset->GetDebugWarpOrigin(GetRole(), ActorIndex == DebugActor);
 		Sampler.Init(PreviewAsset, Origin, BlendParameters);
@@ -106,7 +106,7 @@ void FInteractionAssetPreviewActor::UpdatePreviewActor(const UPoseSearchInteract
 			const FTransform Origin = InteractionAsset->GetDebugWarpOrigin(GetRole(), ActorIndex == DebugActor);
 
 			bool bSamplerReinitialized = false;
-			UE::PoseSearch::FAnimationAssetSampler& Sampler = Samplers[ActorIndex];
+			FAnimationAssetSampler& Sampler = Samplers[ActorIndex];
 			if (PreviewAsset != Sampler.GetAsset() || !Origin.Equals(Sampler.GetRootTransformOrigin()))
 			{
 				// reinitializing the Sampler if the PreviewAsset or the origin transform changed
@@ -173,7 +173,7 @@ UAnimPreviewInstance* FInteractionAssetPreviewActor::GetAnimPreviewInstance()
 
 FTransform FInteractionAssetPreviewActor::GetDebugActorTransformFromSampler() const
 {
-	const UE::PoseSearch::FAnimationAssetSampler& Sampler = Samplers[DebugActor];
+	const FAnimationAssetSampler& Sampler = Samplers[DebugActor];
 	const FTransform ActorTransform = Sampler.ExtractRootTransform(CurrentTime);
 	return ActorTransform;
 }
@@ -264,14 +264,14 @@ void FInteractionAssetViewModel::Tick(float DeltaSeconds)
 	PlayTime += DeltaSeconds * DeltaTimeMultiplier;
 	PlayTime = FMath::Clamp(PlayTime, 0.f, InteractionAsset->GetPlayLength());
 
-	UE::PoseSearch::FRoleToIndex InteractionAssetRoleToIndex;
+	FRoleToIndex InteractionAssetRoleToIndex;
 	for (int32 RoleIndex = 0; RoleIndex < InteractionAsset->GetNumRoles(); ++RoleIndex)
 	{
 		InteractionAssetRoleToIndex.Add(InteractionAsset->GetRole(RoleIndex)) = RoleIndex;
 	}
 
 	// iterating backwards because of the possible RemoveAtSwap 
-	UE::PoseSearch::FRoleToIndex PreviewActorsRoleToIndex;
+	FRoleToIndex PreviewActorsRoleToIndex;
 	for (int32 ActorIndex = PreviewActors.Num() - 1; ActorIndex >= 0; --ActorIndex)
 	{
 		FInteractionAssetPreviewActor& PreviewActor = PreviewActors[ActorIndex];
@@ -306,7 +306,7 @@ void FInteractionAssetViewModel::Tick(float DeltaSeconds)
 		PreviewActor.UpdatePreviewActor(InteractionAsset, PlayTime);
 	}
 	
-	UE::PoseSearch::FRoleToIndex PreviewActorRoleToIndex;
+	FRoleToIndex PreviewActorRoleToIndex;
 	PreviewActorRoleToIndex.Reserve(PreviewActors.Num());
 	for (int32 PreviewActorIndex = 0; PreviewActorIndex < PreviewActors.Num(); ++PreviewActorIndex)
 	{
@@ -317,11 +317,11 @@ void FInteractionAssetViewModel::Tick(float DeltaSeconds)
 #if WITH_EDITORONLY_DATA
 	if (InteractionAsset->bEnableDebugWarp && PreviewActors.Num() == InteractionAsset->GetNumRoles())
 	{
-		TArray<FTransform, TInlineAllocator<UE::PoseSearch::PreallocatedRolesNum>> ActorTransforms;
+		TArray<FTransform, TInlineAllocator<PreallocatedRolesNum>> ActorTransforms;
 		ActorTransforms.SetNum(InteractionAsset->GetNumRoles());
 		for (int32 InteractionAssetRoleIndex = 0; InteractionAssetRoleIndex < InteractionAsset->GetNumRoles(); ++InteractionAssetRoleIndex)
 		{
-			const UE::PoseSearch::FRole& InteractionAssetRole = InteractionAsset->GetRole(InteractionAssetRoleIndex);
+			const FRole& InteractionAssetRole = InteractionAsset->GetRole(InteractionAssetRoleIndex);
 			const int32 PreviewActorIndex = PreviewActorRoleToIndex[InteractionAssetRole];
 
 			FInteractionAssetPreviewActor& PreviewActor = PreviewActors[PreviewActorIndex];
@@ -330,13 +330,13 @@ void FInteractionAssetViewModel::Tick(float DeltaSeconds)
 			ActorTransforms[InteractionAssetRoleIndex] = PreviewActor.GetDebugActorTransformFromSampler();
 		}
 
-		TArray<FTransform, TInlineAllocator<UE::PoseSearch::PreallocatedRolesNum>> FullAlignedActorTransforms;
+		TArray<FTransform, TInlineAllocator<PreallocatedRolesNum>> FullAlignedActorTransforms;
 		FullAlignedActorTransforms.SetNum(InteractionAsset->GetNumRoles());
 		InteractionAsset->CalculateWarpTransforms(PlayTime, ActorTransforms, FullAlignedActorTransforms);
 
 		for (int32 InteractionAssetRoleIndex = 0; InteractionAssetRoleIndex < InteractionAsset->GetNumRoles(); ++InteractionAssetRoleIndex)
 		{
-			const UE::PoseSearch::FRole& InteractionAssetRole = InteractionAsset->GetRole(InteractionAssetRoleIndex);
+			const FRole& InteractionAssetRole = InteractionAsset->GetRole(InteractionAssetRoleIndex);
 			const int32 PreviewActorIndex = PreviewActorRoleToIndex[InteractionAssetRole];
 
 			FInteractionAssetPreviewActor& PreviewActor = PreviewActors[PreviewActorIndex];
diff --git a/Engine/Plugins/Animation/PoseSearch/Source/Runtime/Private/AnimNode_MotionMatchingInteraction.cpp b/Engine/Plugins/Animation/PoseSearch/Source/Runtime/Private/AnimNode_MotionMatchingInteraction.cpp
--- a/Engine/Plugins/Animation/PoseSearch/Source/Runtime/Private/AnimNode_MotionMatchingInteraction.cpp
+++ b/Engine/Plugins/Animation/PoseSearch/Source/Runtime/Private/AnimNode_MotionMatchingInteraction.cpp
@@ -6,6 +6,7 @@
 #include "Animation/AnimRootMotionProvider.h"
 #include "HAL/IConsoleManager.h"
 #include "PoseSearch/PoseSearchInteractionAsset.h"
+#include "PoseSearch/PoseSearchInteractionIsland.h"
 #include "PoseSearch/PoseSearchInteractionLibrary.h"
 #include "PoseSearch/PoseHistoryProvider.h"
 #include "PoseSearch/PoseSearchDatabase.h"
@@ -13,10 +14,18 @@
 #include "PoseSearch/PoseSearchSchema.h"
 #include "VisualLogger/VisualLogger.h"
 
-#if ENABLE_ANIM_DEBUG
-static bool GVarAnimNodeMotionMatchingInteractionDebug = false;
-static FAutoConsoleVariableRef CVarAnimNodeMotionMatchingInteractionDebug(TEXT("a.AnimNode.MotionMatchingInteraction.Debug"), GVarAnimNodeMotionMatchingInteractionDebug, TEXT("Turn on visualization debugging for AnimNode Motion Matching Interaction"));
-#endif // ENABLE_ANIM_DEBUG
+namespace UE::PoseSearch
+{
+	// Taken from https://theorangeduck.com/page/spring-roll-call#implicitspringdamper
+	static float DamperImplicit(float Halflife, float DeltaTime, float Epsilon = 1e-8f)
+	{
+		// Halflife values very close to 0 approach infinity, and result in big motion spikes when Halflife < DeltaTime
+		// This is a hack, and adds some degree of frame-rate dependency, but it holds up even at lower frame-rates.
+		Halflife = FMath::Max(Halflife, DeltaTime);
+		return FMath::Clamp((1.0f - FMath::InvExpApprox((0.69314718056f * DeltaTime) / (Halflife + Epsilon))), 0.0f, 1.0f);
+	}
+} // UE::PoseSearch
+
 
 void FAnimNode_MotionMatchingInteraction::GatherDebugData(FNodeDebugData& DebugData)
 {
@@ -33,9 +42,17 @@ void FAnimNode_MotionMatchingInteraction::GatherDebugData(FNodeDebugData& DebugD
 void FAnimNode_MotionMatchingInteraction::Reset()
 {
 	Super::Reset();
-	TranslationWarpLerp = 0.f;
-	RotationWarpLerp = 0.f;
-	InteractingRolesNum = 0;
+	CurrentResult = FPoseSearchInteractionBlueprintResult();
+	CachedDeltaTime = 0.f;
+}
+
+bool FAnimNode_MotionMatchingInteraction::NeedsReset(const FAnimationUpdateContext& Context) const
+{
+	const bool bNeedsReset =
+		bResetOnBecomingRelevant &&
+		UpdateCounter.HasEverBeenUpdated() &&
+		!UpdateCounter.WasSynchronizedCounter(Context.AnimInstanceProxy->GetUpdateCounter());
+	return bNeedsReset;
 }
 
 void FAnimNode_MotionMatchingInteraction::UpdateAssetPlayer(const FAnimationUpdateContext& Context)
@@ -48,122 +65,45 @@ void FAnimNode_MotionMatchingInteraction::UpdateAssetPlayer(const FAnimationUpda
 	}
 
 	UpdateCounter.SynchronizeWith(Context.AnimInstanceProxy->GetUpdateCounter());
+	CachedDeltaTime = Context.GetDeltaTime();
 
 	GetEvaluateGraphExposedInputs().Execute(Context);
 
-	bool bBlendToExecuted = ConditionalBlendTo(Context);
-
-	int32 NewInteractingRolesNum = 0;
+	bool bBlendToExecuted = false;
 	const float DeltaTime = Context.GetDeltaTime();
 	if (FPoseHistoryProvider* PoseHistoryProvider = Context.GetMessage<FPoseHistoryProvider>())
 	{
 		// we're NOT providing the UMultiAnimAsset here, relying on the interaction subsystem to resolve the ContinuingProperties properly
 		FPoseSearchContinuingProperties ContinuingProperties;
-		ContinuingProperties.PlayingAsset = GetAnimAsset();
+		ContinuingProperties.PlayingAsset = CurrentResult.SelectedAnimation; // GetAnimAsset();
 		ContinuingProperties.PlayingAssetAccumulatedTime = GetAccumulatedTime();
 
 		check(Context.AnimInstanceProxy);
-		const FPoseSearchInteractionBlueprintResult Result = UPoseSearchInteractionLibrary::MotionMatchInteraction(Availabilities, Context.AnimInstanceProxy->GetAnimInstanceObject(), ContinuingProperties, PoseHistoryProvider->GetHistoryCollector(), bValidateResultAgainstAvailabilities);
+		CurrentResult = UPoseSearchInteractionLibrary::MotionMatchInteraction(Availabilities, Context.AnimInstanceProxy->GetAnimInstanceObject(), ContinuingProperties, PoseHistoryProvider->GetHistoryCollector(), bValidateResultAgainstAvailabilities);
 
 		UAnimationAsset* RoledAnimAsset = nullptr;
-		if (const UMultiAnimAsset* MultiAnimAsset = Cast<UMultiAnimAsset>(Result.SelectedAnimation))
+		if (const UMultiAnimAsset* MultiAnimAsset = Cast<UMultiAnimAsset>(CurrentResult.SelectedAnimation))
 		{
-			check(Result.SelectedDatabase != nullptr);
-
-			RoledAnimAsset = MultiAnimAsset->GetAnimationAsset(Result.Role);
-			NewInteractingRolesNum = MultiAnimAsset->GetNumRoles();
+			RoledAnimAsset = MultiAnimAsset->GetAnimationAsset(CurrentResult.Role);
 		}
-		else if (UAnimationAsset* SingleAnimAsset = Cast<UAnimationAsset>(Result.SelectedAnimation))
+		else if (UAnimationAsset* SingleAnimAsset = Cast<UAnimationAsset>(CurrentResult.SelectedAnimation))
 		{
 			RoledAnimAsset = SingleAnimAsset;
-			NewInteractingRolesNum = 1;
 		}
 
 		if (RoledAnimAsset)
 		{
-			bool bExecuteBlendTo = false;
-			bool bUpdatePropertiesFromResult = false;
-			if (InteractingRolesNum == 0 || AnimPlayers.IsEmpty())
-			{
-				bExecuteBlendTo = true;
-				bUpdatePropertiesFromResult = true;
-			}
-			else if (EvaluationMode == EMotionMatchingInteractionEvaluationMode::ContinuousReselection)
-			{
-				const FBlendStackAnimPlayer& MainAnimPlayer = AnimPlayers[0];
-				const UAnimationAsset* PlayingAnimAsset = MainAnimPlayer.GetAnimationAsset();
-
-				if (RoledAnimAsset != PlayingAnimAsset ||
-					Result.bIsMirrored != MainAnimPlayer.GetMirror() ||
-					Result.BlendParameters != MainAnimPlayer.GetBlendParameters() ||
-					!Result.bIsContinuingPoseSearch)
-				{
-					bExecuteBlendTo = true;
-				}
-
-				bUpdatePropertiesFromResult = true;
-			}
-			else if (RoledAnimAsset == AnimPlayers[0].GetAnimationAsset() && Result.bIsContinuingPoseSearch)
-			{
-				// we don't update FullAlignedActorRootBoneTransform since we're not planning to blend into the newly selected animation here
-				bUpdatePropertiesFromResult = true;
-			}
-
-			if (bUpdatePropertiesFromResult)
+			if (AnimPlayers.IsEmpty() || !CurrentResult.bIsContinuingPoseSearch)
 			{
-				FullAlignedActorRootBoneTransform = Result.FullAlignedActorRootBoneTransform;
-				WantedPlayRate = Result.WantedPlayRate;
-				BlendParameters = Result.BlendParameters;
-			}
-
-			if (bExecuteBlendTo)
-			{
-				const FPoseSearchRoledSkeleton* RoledSkeleton = Result.SelectedDatabase->Schema->GetRoledSkeleton(Result.Role);
+				const FPoseSearchRoledSkeleton* RoledSkeleton = CurrentResult.SelectedDatabase->Schema->GetRoledSkeleton(CurrentResult.Role);
 				check(RoledSkeleton);
 
-				BlendTo(Context, RoledAnimAsset, Result.SelectedTime, Result.bLoop, Result.bIsMirrored, RoledSkeleton->MirrorDataTable.Get(),
-					BlendTime, BlendProfile, BlendOption, bUseInertialBlend, BlendParameters, WantedPlayRate);
+				BlendTo(Context, RoledAnimAsset, CurrentResult.SelectedTime, CurrentResult.bLoop, CurrentResult.bIsMirrored, RoledSkeleton->MirrorDataTable.Get(),
+					BlendTime, BlendProfile, BlendOption, bUseInertialBlend, CurrentResult.BlendParameters, CurrentResult.WantedPlayRate);
 
 				bBlendToExecuted = true;
 			}
 		}
-
-		if (bUseAnimRootMotionProvider)
-		{
-			const IPoseHistory& PoseHistory = PoseHistoryProvider->GetPoseHistory();
-			const USkeleton* Skeleton = Context.AnimInstanceProxy->GetSkeleton();
-			FTransform RootBoneTrasform;
-			PoseHistory.GetTransformAtTime(0.f, RootBoneTrasform, Skeleton, RootBoneIndexType, WorldSpaceIndexType);
-
-			const FTransform LerpedAlignedActorRootBoneTransform(
-				FQuat::Slerp(RootBoneTrasform.GetRotation(), FullAlignedActorRootBoneTransform.GetRotation(), RotationWarpLerp),
-				FMath::Lerp(RootBoneTrasform.GetTranslation(), FullAlignedActorRootBoneTransform.GetTranslation(), TranslationWarpLerp),
-				FVector::ZeroVector);
-
-#if ENABLE_ANIM_DEBUG
-			if (GVarAnimNodeMotionMatchingInteractionDebug)
-			{
-				Context.AnimInstanceProxy->AnimDrawDebugCoordinateSystem(RootBoneTrasform.GetLocation(), RootBoneTrasform.Rotator(), 25.f, false, 0.f, 0.f, SDPG_Foreground);
-				Context.AnimInstanceProxy->AnimDrawDebugCoordinateSystem(FullAlignedActorRootBoneTransform.GetLocation(), FullAlignedActorRootBoneTransform.Rotator(), 50.f, false, 0.f, 0.f, SDPG_Foreground);
-			}
-#endif // ENABLE_ANIM_DEBUG
-
-			WarpRootMotionTransform = LerpedAlignedActorRootBoneTransform.GetRelativeTransform(RootBoneTrasform);
-			check(WarpRootMotionTransform.IsRotationNormalized());
-
-#if ENABLE_VISUAL_LOG
-			if (FVisualLogger::IsRecording())
-			{
-				static const TCHAR* LogName = TEXT("MotionMatchingInteraction");
-
-				UObject* LogOwner = Context.AnimInstanceProxy->GetAnimInstanceObject();
-
-				UE_VLOG_ARROW(LogOwner, LogName, Display, FullAlignedActorRootBoneTransform.GetLocation(), FullAlignedActorRootBoneTransform.GetLocation() + FullAlignedActorRootBoneTransform.GetRotation().GetForwardVector() * 100, FColor::Red, TEXT(""));
-				UE_VLOG_ARROW(LogOwner, LogName, Display, LerpedAlignedActorRootBoneTransform.GetLocation(), LerpedAlignedActorRootBoneTransform.GetLocation() + LerpedAlignedActorRootBoneTransform.GetRotation().GetForwardVector() * 100, FColor::Blue, TEXT(""));
-				UE_VLOG_ARROW(LogOwner, LogName, Display, RootBoneTrasform.GetLocation(), RootBoneTrasform.GetLocation() + RootBoneTrasform.GetRotation().GetForwardVector() * 100, FColor::Green, TEXT(""));
-			}
-#endif
-		}
 	}
 	else
 	{
@@ -173,28 +113,8 @@ void FAnimNode_MotionMatchingInteraction::UpdateAssetPlayer(const FAnimationUpda
 	const bool bDidBlendToRequestAnInertialBlend = bBlendToExecuted && bUseInertialBlend;
 	UE::Anim::TOptionalScopedGraphMessage<UE::Anim::FAnimInertializationSyncScope> InertializationSync(bDidBlendToRequestAnInertialBlend, Context);
 	
-	UpdatePlayRate(WantedPlayRate);
-	UpdateBlendspaceParameters(BlendspaceUpdateMode, BlendParameters);
-
-	// calculating the translation and rotation warp lerps, used to warp the root transform towards the last computed FullAlignedActorRootBoneTransform
-	const float WarpSign = bEnableWarping && NewInteractingRolesNum > 1 ? 1.f : -1.f;
-	if (InitialTranslationWarpTime > UE_KINDA_SMALL_NUMBER)
-	{
-		TranslationWarpLerp = FMath::Clamp(TranslationWarpLerp + (WarpSign * DeltaTime / InitialTranslationWarpTime), 0.f, 1.f);
-	}
-	else
-	{
-		TranslationWarpLerp = 0.f;
-	}
-
-	if (InitialRotationWarpTime > UE_KINDA_SMALL_NUMBER)
-	{
-		RotationWarpLerp = FMath::Clamp(RotationWarpLerp + (WarpSign * DeltaTime / InitialRotationWarpTime), 0.f, 1.f);
-	}
-	else
-	{
-		RotationWarpLerp = 0.f;
-	}
+	UpdatePlayRate(CurrentResult.WantedPlayRate);
+	UpdateBlendspaceParameters(BlendspaceUpdateMode, CurrentResult.BlendParameters);
 
 	// bypassing FAnimNode_BlendStack::UpdateAssetPlayer, since we overridden its behaviour
 	FAnimNode_BlendStack_Standalone::UpdateAssetPlayer(Context);
@@ -202,63 +122,99 @@ void FAnimNode_MotionMatchingInteraction::UpdateAssetPlayer(const FAnimationUpda
 #if ENABLE_ANIM_DEBUG
 	if (UE_TRACE_CHANNELEXPR_IS_ENABLED(AnimationChannel))
 	{
-		TRACE_ANIM_NODE_VALUE(Context, *FString("InteractingRolesNum"), InteractingRolesNum);
-		TRACE_ANIM_NODE_VALUE(Context, *FString("NewInteractingRolesNum"), NewInteractingRolesNum);
 		TRACE_ANIM_NODE_VALUE(Context, *FString("BlendToExecuted"), bBlendToExecuted);
-		TRACE_ANIM_NODE_VALUE(Context, *FString("TranslationWarpLerp"), TranslationWarpLerp);
-		TRACE_ANIM_NODE_VALUE(Context, *FString("RotationWarpLerp"), RotationWarpLerp);
 	}
 #endif // ENABLE_ANIM_DEBUG
-
-	InteractingRolesNum = NewInteractingRolesNum;
 }
 
 void FAnimNode_MotionMatchingInteraction::Evaluate_AnyThread(FPoseContext& Output)
 {
+	using namespace UE::PoseSearch;
+
 	check(Output.AnimInstanceProxy);
 
 	Super::Evaluate_AnyThread(Output);
 
-	if (TranslationWarpLerp > UE_KINDA_SMALL_NUMBER || RotationWarpLerp > UE_KINDA_SMALL_NUMBER)
+	if (const UE::Anim::IAnimRootMotionProvider* RootMotionProvider = UE::Anim::IAnimRootMotionProvider::Get())
 	{
-		if (bUseAnimRootMotionProvider)
+		if (const UMultiAnimAsset* MultiAnimAsset = Cast<UMultiAnimAsset>(CurrentResult.SelectedAnimation))
 		{
-#if ENABLE_ANIM_DEBUG
-			if (GVarAnimNodeMotionMatchingInteractionDebug)
+			const FInteractionSearchContext InteractionSearchContext = UPoseSearchInteractionLibrary::FindSearchContext_AnyThread(Output.AnimInstanceProxy->GetAnimInstanceObject());
+			if (InteractionSearchContext.IsValid())
 			{
-				const FTransform ComponentTransform = Output.AnimInstanceProxy->GetComponentTransform();
-				Output.AnimInstanceProxy->AnimDrawDebugCoordinateSystem(FullAlignedActorRootBoneTransform.GetLocation(), FullAlignedActorRootBoneTransform.Rotator(), 15.f, false, 0.f, 0.f, SDPG_Foreground);
-				Output.AnimInstanceProxy->AnimDrawDebugCoordinateSystem(ComponentTransform.GetLocation(), ComponentTransform.Rotator(), 5.f, false, 0.f, 0.f, SDPG_Foreground);
-			}
+				// populating ActorTransforms containing all the SkelMeshComponent transforms of all the characters involved into this interaction,
+				// to be able to calculate the required warp transform for this actor via UMultiAnimAsset::CalculateWarpTransforms
+				check(InteractionSearchContext.CheckForConsistency());
+				const FRoleToIndex InteractionSearchContextRoleToIndex = MakeRoleToIndex(InteractionSearchContext.Roles);
+				TArray<FTransform, TInlineAllocator<PreallocatedRolesNum>> ActorTransforms;
+				ActorTransforms.SetNum(MultiAnimAsset->GetNumRoles());
+
+				int32 CurrentResultRoleIndex = INDEX_NONE;
+				bool bAllRolesResolved = true;
+				for (int32 MultiAnimAssetRoleIndex = 0; MultiAnimAssetRoleIndex < MultiAnimAsset->GetNumRoles(); ++MultiAnimAssetRoleIndex)
+				{
+					const FRole& MultiAnimAssetRole = MultiAnimAsset->GetRole(MultiAnimAssetRoleIndex);
+
+					if (CurrentResultRoleIndex == INDEX_NONE && MultiAnimAssetRole == CurrentResult.Role)
+					{
+						CurrentResultRoleIndex = MultiAnimAssetRoleIndex;
+					}
+					
+					if (const int32* InteractionSearchContextRoleIndex = InteractionSearchContextRoleToIndex.Find(MultiAnimAssetRole))
+					{
+						if (const UAnimInstance* AnimInstance = Cast<UAnimInstance>(InteractionSearchContext.AnimContexts[*InteractionSearchContextRoleIndex].Get()))
+						{
+							ActorTransforms[MultiAnimAssetRoleIndex] = AnimInstance->GetSkelMeshComponent()->GetComponentTransform();
+						}
+						else
+						{
+							bAllRolesResolved = false;
+							break;
+						}
+					}
+					else
+					{
+						bAllRolesResolved = false;
+						break;
+					}
+				}
 
-#endif // ENABLE_ANIM_DEBUG
-		
-			if (const UE::Anim::IAnimRootMotionProvider* RootMotionProvider = UE::Anim::IAnimRootMotionProvider::Get())
-			{
-				RootMotionProvider->OverrideRootMotion(WarpRootMotionTransform, Output.CustomAttributes);
-			}
-			else
-			{
-				UE_LOG(LogPoseSearch, Error, TEXT("FAnimNode_MotionMatchingInteraction::Evaluate_AnyThread couldn't find the IAnimRootMotionProvider"));
-			}
-		}
-		else
-		{
-			const FTransform ComponentTransform = Output.AnimInstanceProxy->GetComponentTransform();
-			const FTransform FullAlignedActorRootBoneLocalTransform = FullAlignedActorRootBoneTransform.GetRelativeTransform(ComponentTransform);
+				if (bAllRolesResolved && CurrentResultRoleIndex != INDEX_NONE)
+				{
+					TArray<FTransform, TInlineAllocator<PreallocatedRolesNum>> FullAlignedActorRootBoneTransforms;
+					FullAlignedActorRootBoneTransforms.SetNum(MultiAnimAsset->GetNumRoles());
 
-#if ENABLE_ANIM_DEBUG
-			if (GVarAnimNodeMotionMatchingInteractionDebug)
-			{
-				Output.AnimInstanceProxy->AnimDrawDebugCoordinateSystem(FullAlignedActorRootBoneTransform.GetLocation(), FullAlignedActorRootBoneTransform.Rotator(), 15.f, false, 0.f, 0.f, SDPG_Foreground);
-				Output.AnimInstanceProxy->AnimDrawDebugCoordinateSystem(ComponentTransform.GetLocation(), ComponentTransform.Rotator(), 5.f, false, 0.f, 0.f, SDPG_Foreground);
-			}
-#endif // ENABLE_ANIM_DEBUG
+					MultiAnimAsset->CalculateWarpTransforms(GetAccumulatedTime(), ActorTransforms, FullAlignedActorRootBoneTransforms);
+
+					const FTransform& FullAlignedActorRootBoneTransform = FullAlignedActorRootBoneTransforms[CurrentResultRoleIndex];
+
+					const FTransform RootBoneTransform = Output.AnimInstanceProxy->GetSkelMeshComponent()->GetBoneTransform(RootBoneIndexType);
+					FTransform WarpRootMotionTransform = FullAlignedActorRootBoneTransform.GetRelativeTransform(RootBoneTransform);
+
+					check(WarpRootMotionTransform.IsRotationNormalized());
 
-			const FCompactPoseBoneIndex RootBoneIndex(0);
-			FTransform& RootBoneTransform = Output.Pose[RootBoneIndex];
-			RootBoneTransform.SetTranslation(FMath::Lerp(RootBoneTransform.GetTranslation(), FullAlignedActorRootBoneLocalTransform.GetTranslation(), TranslationWarpLerp));
-			RootBoneTransform.SetRotation(FQuat::Slerp(RootBoneTransform.GetRotation(), FullAlignedActorRootBoneLocalTransform.GetRotation(), RotationWarpLerp));
+					const float DampenWarpingTranslation = WarpingTranslationHalflife > 0.f ? DamperImplicit(WarpingTranslationHalflife, CachedDeltaTime) : 1.f;
+					const float DampenWarpingRotation = WarpingRotationHalfLife > 0.f ? DamperImplicit(WarpingRotationHalfLife, CachedDeltaTime) : 1.f;
+
+					WarpRootMotionTransform.SetTranslation(WarpRootMotionTransform.GetTranslation() * DampenWarpingTranslation);
+					WarpRootMotionTransform.SetRotation(FQuat::Slerp(FQuat::Identity, WarpRootMotionTransform.GetRotation(), DampenWarpingRotation));
+
+#if ENABLE_VISUAL_LOG
+					if (FVisualLogger::IsRecording())
+					{
+						static const TCHAR* LogName = TEXT("MotionMatchingInteraction");
+
+						UObject* LogOwner = Output.AnimInstanceProxy->GetAnimInstanceObject();
+						UE_VLOG_ARROW(LogOwner, LogName, Display, FullAlignedActorRootBoneTransform.GetLocation(), FullAlignedActorRootBoneTransform.GetLocation() + FullAlignedActorRootBoneTransform.GetRotation().GetForwardVector() * 100, FColor::Purple, TEXT(""));
+					}
+#endif
+					RootMotionProvider->OverrideRootMotion(WarpRootMotionTransform, Output.CustomAttributes);
+				}
+			}
 		}
 	}
+	else
+	{
+		UE_LOG(LogPoseSearch, Error, TEXT("FAnimNode_MotionMatchingInteraction::Evaluate_AnyThread couldn't find the IAnimRootMotionProvider"));
+	}
 }
diff --git a/Engine/Plugins/Animation/PoseSearch/Source/Runtime/Private/MotionMatchingInteractionAnimNodeLibrary.cpp b/Engine/Plugins/Animation/PoseSearch/Source/Runtime/Private/MotionMatchingInteractionAnimNodeLibrary.cpp
--- a/Engine/Plugins/Animation/PoseSearch/Source/Runtime/Private/MotionMatchingInteractionAnimNodeLibrary.cpp
+++ b/Engine/Plugins/Animation/PoseSearch/Source/Runtime/Private/MotionMatchingInteractionAnimNodeLibrary.cpp
@@ -20,24 +20,13 @@ void UMotionMatchingInteractionAnimNodeLibrary::SetAvailabilities(const FMotionM
 	}
 }
 
-float UMotionMatchingInteractionAnimNodeLibrary::GetTranslationWarpLerp(const FMotionMatchingInteractionAnimNodeReference& MotionMatchingInteractionNode)
+bool UMotionMatchingInteractionAnimNodeLibrary::IsInteracting(const FMotionMatchingInteractionAnimNodeReference& MotionMatchingInteractionNode)
 {
 	if (const FAnimNode_MotionMatchingInteraction* MotionMatchingInteractionNodePtr = MotionMatchingInteractionNode.GetAnimNodePtr<FAnimNode_MotionMatchingInteraction>())
 	{
-		return MotionMatchingInteractionNodePtr->GetTranslationWarpLerp();
+		return MotionMatchingInteractionNodePtr->IsInteracting();
 	}
 
-	UE_LOG(LogPoseSearch, Warning, TEXT("UMotionMatchingInteractionAnimNodeLibrary::GetTranslationWarpLerp called on an invalid context or with an invalid type"));
-	return 0.f;
+	UE_LOG(LogPoseSearch, Warning, TEXT("UMotionMatchingInteractionAnimNodeLibrary::IsInteracting called on an invalid context or with an invalid type"));
+	return false;
 }
-
-float UMotionMatchingInteractionAnimNodeLibrary::GetRotationWarpLerp(const FMotionMatchingInteractionAnimNodeReference& MotionMatchingInteractionNode)
-{
-	if (FAnimNode_MotionMatchingInteraction* MotionMatchingInteractionNodePtr = MotionMatchingInteractionNode.GetAnimNodePtr<FAnimNode_MotionMatchingInteraction>())
-	{
-		return MotionMatchingInteractionNodePtr->GetRotationWarpLerp();
-	}
-	
-	UE_LOG(LogPoseSearch, Warning, TEXT("UMotionMatchingInteractionAnimNodeLibrary::GetRotationWarpLerp called on an invalid context or with an invalid type"));
-	return 0.f;
-}
\ No newline at end of file
diff --git a/Engine/Plugins/Animation/PoseSearch/Source/Runtime/Private/PoseSearchInteractionAsset.cpp b/Engine/Plugins/Animation/PoseSearch/Source/Runtime/Private/PoseSearchInteractionAsset.cpp
--- a/Engine/Plugins/Animation/PoseSearch/Source/Runtime/Private/PoseSearchInteractionAsset.cpp
+++ b/Engine/Plugins/Animation/PoseSearch/Source/Runtime/Private/PoseSearchInteractionAsset.cpp
@@ -260,6 +260,8 @@ FTransform UPoseSearchInteractionAsset::GetDebugWarpOrigin(const UE::PoseSearch:
 
 void UPoseSearchInteractionAsset::CalculateWarpTransforms(float Time, const TArrayView<const FTransform> ActorRootBoneTransforms, TArrayView<FTransform> FullAlignedActorRootBoneTransforms) const
 {
+	using namespace UE::PoseSearch;
+
 	check(ActorRootBoneTransforms.Num() == GetNumRoles());
 	check(FullAlignedActorRootBoneTransforms.Num() == GetNumRoles());
 
@@ -273,7 +275,7 @@ void UPoseSearchInteractionAsset::CalculateWarpTransforms(float Time, const TArr
 	}
 	else
 	{
-		TArray<FTransform, TInlineAllocator<UE::PoseSearch::PreallocatedRolesNum>> AssetRootBoneTransforms;
+		TArray<FTransform, TInlineAllocator<PreallocatedRolesNum>> AssetRootBoneTransforms;
 		AssetRootBoneTransforms.SetNum(ItemsNum);
 
 		// ItemIndex is the RoleIndex and Role = Item.Role
@@ -282,7 +284,7 @@ void UPoseSearchInteractionAsset::CalculateWarpTransforms(float Time, const TArr
 			const FPoseSearchInteractionAssetItem& Item = Items[ItemIndex];
 
 			// sampling the AnimationAsset to extract the current time transform and the initial (time of 0) transform
-			const UE::PoseSearch::FAnimationAssetSampler Sampler(Item.Animation, Item.Origin);
+			const FAnimationAssetSampler Sampler(Item.Animation, Item.Origin);
 			AssetRootBoneTransforms[ItemIndex] = Sampler.ExtractRootTransform(Time);
 
 #if ENABLE_ANIM_DEBUG && WITH_EDITOR
@@ -312,9 +314,9 @@ void UPoseSearchInteractionAsset::CalculateWarpTransforms(float Time, const TArr
 #endif // ENABLE_ANIM_DEBUG && WITH_EDITOR
 		}
 
-		TArray<int32, TInlineAllocator<UE::PoseSearch::PreallocatedRolesNum>> SortedByWarpingWeightRotationItemIndex;
-		TArray<float, TInlineAllocator<UE::PoseSearch::PreallocatedRolesNum>> NormalizedWarpingWeightRotation;
-		TArray<float, TInlineAllocator<UE::PoseSearch::PreallocatedRolesNum>> NormalizedWarpingWeightTranslation;
+		TArray<int32, TInlineAllocator<PreallocatedRolesNum>> SortedByWarpingWeightRotationItemIndex;
+		TArray<float, TInlineAllocator<PreallocatedRolesNum>> NormalizedWarpingWeightRotation;
+		TArray<float, TInlineAllocator<PreallocatedRolesNum>> NormalizedWarpingWeightTranslation;
 		SortedByWarpingWeightRotationItemIndex.SetNum(ItemsNum);
 		NormalizedWarpingWeightRotation.SetNum(ItemsNum);
 		NormalizedWarpingWeightTranslation.SetNum(ItemsNum);
diff --git a/Engine/Plugins/Animation/PoseSearch/Source/Runtime/Private/PoseSearchInteractionIsland.cpp b/Engine/Plugins/Animation/PoseSearch/Source/Runtime/Private/PoseSearchInteractionIsland.cpp
--- a/Engine/Plugins/Animation/PoseSearch/Source/Runtime/Private/PoseSearchInteractionIsland.cpp
+++ b/Engine/Plugins/Animation/PoseSearch/Source/Runtime/Private/PoseSearchInteractionIsland.cpp
@@ -28,88 +28,6 @@ static FInteractionSearchResult InitSearchResult(const FSearchResult& SearchResu
 	FInteractionSearchResult InteractionSearchResult;
 	static_cast<FSearchResult&>(InteractionSearchResult) = SearchResult;
 	InteractionSearchResult.SearchIndex = SearchIndex;
-
-	if (const UE::PoseSearch::FSearchIndexAsset* SearchIndexAsset = SearchResult.GetSearchIndexAsset())
-	{
-		if (const FPoseSearchDatabaseAnimationAssetBase* DatabaseAnimationAssetBase = SearchResult.Database->GetDatabaseAnimationAsset<FPoseSearchDatabaseAnimationAssetBase>(*SearchIndexAsset))
-		{
-			if (const UMultiAnimAsset* MultiAnimAsset = Cast<UMultiAnimAsset>(DatabaseAnimationAssetBase->GetAnimationAsset()))
-			{
-				check(MultiAnimAsset->GetNumRoles() == SearchContext.AnimContexts.Num());
-				check(MultiAnimAsset->GetNumRoles() == SearchContext.HistoryCollectors.Num());
-				check(MultiAnimAsset->GetNumRoles() == SearchContext.Roles.Num());
-
-				UE::PoseSearch::FRoleToIndex SearchContextRoleToIndex;
-				SearchContextRoleToIndex.Reserve(SearchContext.Roles.Num());
-				for (int32 SearchContextRoleIndex = 0; SearchContextRoleIndex < SearchContext.Roles.Num(); ++SearchContextRoleIndex)
-				{
-					SearchContextRoleToIndex.Add(SearchContext.Roles[SearchContextRoleIndex]) = SearchContextRoleIndex;
-				}
-
-				// mapping MultiAnimAsset Roles to SearchContext.* indexes
-				TArray<FTransform, TInlineAllocator<UE::PoseSearch::PreallocatedRolesNum>> ActorRootBoneTransforms;
-				for (int32 MultiAnimAssetRoleIndex = 0; MultiAnimAssetRoleIndex < MultiAnimAsset->GetNumRoles(); ++MultiAnimAssetRoleIndex)
-				{
-					const UE::PoseSearch::FRole& InteractionAssetRole = MultiAnimAsset->GetRole(MultiAnimAssetRoleIndex);
-					const int32 SearchContextIndex = SearchContextRoleToIndex[InteractionAssetRole];
-
-					if (const UObject* AnimContext = SearchContext.AnimContexts[SearchContextIndex].Get())
-					{
-						if (const UAnimInstance* AnimInstance = Cast<UAnimInstance>(AnimContext))
-						{
-							const FTransform RootBoneTransform = AnimInstance->GetSkelMeshComponent()->GetBoneTransform(0);
-							ActorRootBoneTransforms.Add(RootBoneTransform);
-						}
-						else
-						{
-							unimplemented();
-						}
-					}
-				}
-
-				if (ActorRootBoneTransforms.Num() == MultiAnimAsset->GetNumRoles())
-				{
-					// FullAlignedActorRootBoneTransforms is mapped to the MultiAnimAsset roles:
-					// FullAlignedActorRootBoneTransforms[0] is for MultiAnimAsset->GetRole(0)
-					TArray<FTransform, TInlineAllocator<UE::PoseSearch::PreallocatedRolesNum>> FullAlignedActorRootBoneTransforms;
-					FullAlignedActorRootBoneTransforms.SetNum(MultiAnimAsset->GetNumRoles());
-
-					// @todo: should it be SearchResult.AssetTime + DeltaTime?
-					MultiAnimAsset->CalculateWarpTransforms(SearchResult.AssetTime, ActorRootBoneTransforms, FullAlignedActorRootBoneTransforms);
-					InteractionSearchResult.FullAlignedActorRootBoneTransforms.SetNum(MultiAnimAsset->GetNumRoles(), EAllowShrinking::No);
-
-					for (int32 InteractionAssetRoleIndex = 0; InteractionAssetRoleIndex < MultiAnimAsset->GetNumRoles(); ++InteractionAssetRoleIndex)
-					{
-						const UE::PoseSearch::FRole& InteractionAssetRole = MultiAnimAsset->GetRole(InteractionAssetRoleIndex);
-						const int32 SearchContextIndex = SearchContextRoleToIndex[InteractionAssetRole];
-
-						InteractionSearchResult.FullAlignedActorRootBoneTransforms[SearchContextIndex] = FullAlignedActorRootBoneTransforms[InteractionAssetRoleIndex];
-					}
-				}
-			}
-			else
-			{
-				// support for non UMultiAnimAsset to be backward compatible with regular motion matching searches
-				check(SearchContext.AnimContexts.Num() == 1);
-
-				if (const UObject* AnimContext = SearchContext.AnimContexts[0].Get())
-				{
-					if (const UAnimInstance* AnimInstance = Cast<UAnimInstance>(AnimContext))
-					{
-						const FTransform RootBoneTransform = AnimInstance->GetSkelMeshComponent()->GetBoneTransform(0);
-						InteractionSearchResult.FullAlignedActorRootBoneTransforms.SetNum(1, EAllowShrinking::No);
-						// @todo: should it be RootBoneTransform + root motion transform?
-						InteractionSearchResult.FullAlignedActorRootBoneTransforms[0] = RootBoneTransform;
-					}
-					else
-					{
-						unimplemented();
-					}
-				}
-			}
-		}
-	}
-
 	return InteractionSearchResult;
 }
 
@@ -147,7 +65,8 @@ static void InitSearchResults(TArray<FInteractionSearchResult>& SearchResults, T
 
 // FInteractionSearchContext
 ///////////////////////////////////////////////////////////
-bool FInteractionSearchContext::IsValid() const
+#if DO_CHECK
+bool FInteractionSearchContext::CheckForConsistency() const
 {
 	if (Database == nullptr)
 	{
@@ -215,6 +134,7 @@ bool FInteractionSearchContext::IsValid() const
 
 	return true;
 }
+#endif // DO_CHECK
 
 bool FInteractionSearchContext::IsEquivalent(const FInteractionSearchContext& Other) const
 {
@@ -256,9 +176,8 @@ bool FInteractionSearchContext::IsEquivalent(const FInteractionSearchContext& Ot
 ///////////////////////////////////////////////////////////
 bool FInteractionSearchResult::operator==(const FInteractionSearchResult& Other) const
 {
-	return static_cast<const FSearchResult&>(*this) == static_cast<const FSearchResult&>(Other) &&
-				FullAlignedActorRootBoneTransforms.Num() == Other.FullAlignedActorRootBoneTransforms.Num() &&
-				0 == FMemory::Memcmp(FullAlignedActorRootBoneTransforms.GetData(), Other.FullAlignedActorRootBoneTransforms.GetData(), FullAlignedActorRootBoneTransforms.Num() * sizeof(FTransform));
+	// not checking for SearchIndex equality
+	return static_cast<const FSearchResult&>(*this) == static_cast<const FSearchResult&>(Other);
 }
 
 // FIslandPreTickFunction
@@ -442,7 +361,7 @@ void FInteractionIsland::InjectToActor(const UObject* AnimContext)
 void FInteractionIsland::AddSearchContext(const FInteractionSearchContext& SearchContext)
 {
 #if DO_CHECK
-	check(SearchContext.IsValid());
+	check(SearchContext.CheckForConsistency());
 	check(IsInGameThread());
 
 	for (const FInteractionSearchContext& ContainedSearchContext : SearchContexts)
@@ -523,9 +442,9 @@ bool FInteractionIsland::DoSearch_AnyThread(UObject* AnimContext, const FPoseSea
 	{
 		FMemMark Mark(FMemStack::Get());
 
-		TArray<const UObject*, TInlineAllocator<UE::PoseSearch::PreallocatedRolesNum, TMemStackAllocator<>>> AnimContexts;
-		TArray<const UE::PoseSearch::IPoseHistory*, TInlineAllocator<UE::PoseSearch::PreallocatedRolesNum, TMemStackAllocator<>>> PoseHistories;
-		TArray<UE::PoseSearch::FSearchResult, TMemStackAllocator<>> PoseSearchResults;
+		TArray<const UObject*, TInlineAllocator<PreallocatedRolesNum, TMemStackAllocator<>>> AnimContexts;
+		TArray<const IPoseHistory*, TInlineAllocator<PreallocatedRolesNum, TMemStackAllocator<>>> PoseHistories;
+		TArray<FSearchResult, TMemStackAllocator<>> PoseSearchResults;
 
 		// SearchContexts are modified only by UPoseSearchInteractionSubsystem::Tick and constant otherwise, so it's safe to access them in a threaded enviroment without locks
 		PoseSearchResults.SetNum(SearchContexts.Num());
@@ -601,7 +520,7 @@ bool FInteractionIsland::DoSearch_AnyThread(UObject* AnimContext, const FPoseSea
 			//}
 
 			// @todo: we could perform multiple UPoseSearchLibrary::MotionMatch in parallel!
-			const UE::PoseSearch::FSearchResult PoseSearchResult = UPoseSearchLibrary::MotionMatch(AnimContexts, SearchContext.Roles,
+			const FSearchResult PoseSearchResult = UPoseSearchLibrary::MotionMatch(AnimContexts, SearchContext.Roles,
 				PoseHistories, AssetsToSearch, ContinuingPropertiesToUse, PoseSearchFutureProperties);
 
 			if (PoseSearchResult.PoseCost < SearchContext.MaxCost)
@@ -644,6 +563,28 @@ bool FInteractionIsland::DoSearch_AnyThread(UObject* AnimContext, const FPoseSea
 	return GetResult_AnyThread(AnimContext, Result);
 }
 
+FInteractionSearchContext FInteractionIsland::FindSearchContext_AnyThread(UObject* AnimContext)
+{
+	FScopeLock Lock(&SearchResultsMutex);
+
+	if (bSearchPerfomed)
+	{
+		for (const FInteractionSearchResult& SearchResult : SearchResults)
+		{
+			const FInteractionSearchContext& SearchContext = SearchContexts[SearchResult.SearchIndex];
+			for (int32 AnimContextIndex = 0; AnimContextIndex < SearchContext.AnimContexts.Num(); ++AnimContextIndex)
+			{
+				if (SearchContext.AnimContexts[AnimContextIndex].Get() == AnimContext)
+				{
+					return SearchContext;
+				}
+			}
+		}
+	}
+
+	return FInteractionSearchContext();
+}
+
 bool FInteractionIsland::GetResult_AnyThread(UObject* AnimContext, FPoseSearchInteractionBlueprintResult& Result)
 {
 	// locking to read SearchResults
@@ -660,13 +601,13 @@ bool FInteractionIsland::GetResult_AnyThread(UObject* AnimContext, FPoseSearchIn
 				const UPoseSearchDatabase* Database = SearchResult.Database.Get();
 				check(Database);
 
-				const UE::PoseSearch::FSearchIndexAsset* SearchIndexAsset = SearchResult.GetSearchIndexAsset();
+				const FSearchIndexAsset* SearchIndexAsset = SearchResult.GetSearchIndexAsset();
 				check(SearchIndexAsset);
 
 				const FPoseSearchDatabaseAnimationAssetBase* DatabaseAnimationAssetBase = Database->GetDatabaseAnimationAsset<FPoseSearchDatabaseAnimationAssetBase>(*SearchIndexAsset);
 				check(DatabaseAnimationAssetBase);
 
-				const UE::PoseSearch::FRole Role = SearchContext.Roles[AnimContextIndex];
+				const FRole Role = SearchContext.Roles[AnimContextIndex];
 
 				Result.SelectedAnimation = DatabaseAnimationAssetBase->GetAnimationAsset();
 				Result.SelectedTime = SearchResult.AssetTime;
@@ -678,11 +619,6 @@ bool FInteractionIsland::GetResult_AnyThread(UObject* AnimContext, FPoseSearchIn
 				Result.SearchCost = SearchResult.PoseCost;
 				Result.Role = Role;
 
-				if (SearchResult.FullAlignedActorRootBoneTransforms.IsValidIndex(AnimContextIndex))
-				{
-					Result.FullAlignedActorRootBoneTransform = SearchResult.FullAlignedActorRootBoneTransforms[AnimContextIndex];
-				}
-
 				//Result.bIsFromContinuingPlaying = SearchResult.bIsContinuingPoseSearch;
 
 				// figuring out the WantedPlayRate
diff --git a/Engine/Plugins/Animation/PoseSearch/Source/Runtime/Private/PoseSearchInteractionLibrary.cpp b/Engine/Plugins/Animation/PoseSearch/Source/Runtime/Private/PoseSearchInteractionLibrary.cpp
--- a/Engine/Plugins/Animation/PoseSearch/Source/Runtime/Private/PoseSearchInteractionLibrary.cpp
+++ b/Engine/Plugins/Animation/PoseSearch/Source/Runtime/Private/PoseSearchInteractionLibrary.cpp
@@ -5,31 +5,40 @@
 #include "PoseSearch/PoseSearchInteractionSubsystem.h"
 #include "Engine/World.h"
 
-FPoseSearchInteractionBlueprintResult UPoseSearchInteractionLibrary::MotionMatchInteraction_Pure(TArray<FPoseSearchInteractionAvailability> Availabilities, UObject* AnimInstance, FPoseSearchContinuingProperties ContinuingProperties, FName PoseHistoryName, bool bValidateResultAgainstAvailabilities)
+FPoseSearchInteractionBlueprintResult UPoseSearchInteractionLibrary::MotionMatchInteraction_Pure(TArray<FPoseSearchInteractionAvailability> Availabilities, UObject* AnimContext, FPoseSearchContinuingProperties ContinuingProperties, FName PoseHistoryName, bool bValidateResultAgainstAvailabilities)
 {
 	FPoseSearchInteractionBlueprintResult Result;
-	if (UPoseSearchInteractionSubsystem* InteractionSubsystem = UPoseSearchInteractionSubsystem::GetSubsystem_AnyThread(AnimInstance))
+	if (UPoseSearchInteractionSubsystem* InteractionSubsystem = UPoseSearchInteractionSubsystem::GetSubsystem_AnyThread(AnimContext))
 	{
-		InteractionSubsystem->Query_AnyThread(Availabilities, AnimInstance, ContinuingProperties, Result, PoseHistoryName, nullptr, bValidateResultAgainstAvailabilities);
+		InteractionSubsystem->Query_AnyThread(Availabilities, AnimContext, ContinuingProperties, Result, PoseHistoryName, nullptr, bValidateResultAgainstAvailabilities);
 	}
 	return Result;
 }
 
-FPoseSearchInteractionBlueprintResult UPoseSearchInteractionLibrary::MotionMatchInteraction(TArray<FPoseSearchInteractionAvailability> Availabilities, UObject* AnimInstance, FPoseSearchContinuingProperties ContinuingProperties, FName PoseHistoryName, bool bValidateResultAgainstAvailabilities)
+FPoseSearchInteractionBlueprintResult UPoseSearchInteractionLibrary::MotionMatchInteraction(TArray<FPoseSearchInteractionAvailability> Availabilities, UObject* AnimContext, FPoseSearchContinuingProperties ContinuingProperties, FName PoseHistoryName, bool bValidateResultAgainstAvailabilities)
 {
-	return MotionMatchInteraction_Pure(Availabilities, AnimInstance, ContinuingProperties, PoseHistoryName, bValidateResultAgainstAvailabilities);
+	return MotionMatchInteraction_Pure(Availabilities, AnimContext, ContinuingProperties, PoseHistoryName, bValidateResultAgainstAvailabilities);
 }
 
-FPoseSearchInteractionBlueprintResult UPoseSearchInteractionLibrary::MotionMatchInteraction(const TArrayView<const FPoseSearchInteractionAvailability> Availabilities, UObject* AnimInstance, const FPoseSearchContinuingProperties& ContinuingProperties, const FAnimNode_PoseSearchHistoryCollector_Base* HistoryCollector, bool bValidateResultAgainstAvailabilities)
+FPoseSearchInteractionBlueprintResult UPoseSearchInteractionLibrary::MotionMatchInteraction(const TArrayView<const FPoseSearchInteractionAvailability> Availabilities, UObject* AnimContext, const FPoseSearchContinuingProperties& ContinuingProperties, const FAnimNode_PoseSearchHistoryCollector_Base* HistoryCollector, bool bValidateResultAgainstAvailabilities)
 {
 	FPoseSearchInteractionBlueprintResult Result;
-	if (UPoseSearchInteractionSubsystem* InteractionSubsystem = UPoseSearchInteractionSubsystem::GetSubsystem_AnyThread(AnimInstance))
+	if (UPoseSearchInteractionSubsystem* InteractionSubsystem = UPoseSearchInteractionSubsystem::GetSubsystem_AnyThread(AnimContext))
 	{
-		InteractionSubsystem->Query_AnyThread(Availabilities, AnimInstance, ContinuingProperties, Result, FName(), HistoryCollector, bValidateResultAgainstAvailabilities);
+		InteractionSubsystem->Query_AnyThread(Availabilities, AnimContext, ContinuingProperties, Result, FName(), HistoryCollector, bValidateResultAgainstAvailabilities);
 	}
 	return Result;
 }
 
+UE::PoseSearch::FInteractionSearchContext UPoseSearchInteractionLibrary::FindSearchContext_AnyThread(UObject* AnimContext)
+{
+	if (UPoseSearchInteractionSubsystem* InteractionSubsystem = UPoseSearchInteractionSubsystem::GetSubsystem_AnyThread(AnimContext))
+	{
+		return InteractionSubsystem->FindSearchContext_AnyThread(AnimContext);
+	}
+	return UE::PoseSearch::FInteractionSearchContext();
+}
+
 FPoseSearchContinuingProperties UPoseSearchInteractionLibrary::GetMontageContinuingProperties(UAnimInstance* AnimInstance)
 {
 	FPoseSearchContinuingProperties ContinuingProperties;
diff --git a/Engine/Plugins/Animation/PoseSearch/Source/Runtime/Private/PoseSearchInteractionSubsystem.cpp b/Engine/Plugins/Animation/PoseSearch/Source/Runtime/Private/PoseSearchInteractionSubsystem.cpp
--- a/Engine/Plugins/Animation/PoseSearch/Source/Runtime/Private/PoseSearchInteractionSubsystem.cpp
+++ b/Engine/Plugins/Animation/PoseSearch/Source/Runtime/Private/PoseSearchInteractionSubsystem.cpp
@@ -75,14 +75,14 @@ namespace UE::PoseSearch
 			}
 		}
 
-		TSet<const UObject*, DefaultKeyFuncs<const UObject*>, TInlineSetAllocator<UE::PoseSearch::PreallocatedRolesNum>> AnimContexts;
+		TSet<const UObject*, DefaultKeyFuncs<const UObject*>, TInlineSetAllocator<PreallocatedRolesNum>> AnimContexts;
 		TArray<int32, TInlineAllocator<PreallocatedSearchesNum>> SearchContextsIndices;
 	};
 
 	struct FRoledAnimContext
 	{
 		const UObject* AnimContext = nullptr;
-		UE::PoseSearch::FRole Role = UE::PoseSearch::DefaultRole;
+		FRole Role = DefaultRole;
 		const FAnimNode_PoseSearchHistoryCollector_Base* HistoryCollector = nullptr;
 		float BroadPhaseRadius = 0.f;
 		float MaxCost = MAX_flt;
@@ -94,7 +94,7 @@ namespace UE::PoseSearch
 	// that ultimately represent the same search, but endign up with different SearchIDs
 	struct FRoledAnimContexts
 	{
-		void AddRoledAnimContext(const UObject* AnimContext, UE::PoseSearch::FRole Role, const FAnimNode_PoseSearchHistoryCollector_Base* HistoryCollector, float BroadPhaseRadius, float MaxCost)
+		void AddRoledAnimContext(const UObject* AnimContext, FRole Role, const FAnimNode_PoseSearchHistoryCollector_Base* HistoryCollector, float BroadPhaseRadius, float MaxCost)
 		{
 			check(AnimContext && HistoryCollector);
 
@@ -140,10 +140,10 @@ namespace UE::PoseSearch
 		}
 
 	private:
-		TArray<FRoledAnimContext, TInlineAllocator<UE::PoseSearch::PreallocatedRolesNum, TMemStackAllocator<>>> RoledAnimContexts;
+		TArray<FRoledAnimContext, TInlineAllocator<PreallocatedRolesNum, TMemStackAllocator<>>> RoledAnimContexts;
 	};
 
-	typedef TMap<const UPoseSearchDatabase*, FRoledAnimContexts, TInlineSetAllocator<PreallocatedSearchesNum, UE::PoseSearch::TMemStackSetAllocator<>>> FDatabaseToRoledAnimContexts;
+	typedef TMap<const UPoseSearchDatabase*, FRoledAnimContexts, TInlineSetAllocator<PreallocatedSearchesNum, TMemStackSetAllocator<>>> FDatabaseToRoledAnimContexts;
 	typedef TPair<const UPoseSearchDatabase*, FRoledAnimContexts> FDatabaseToRoledAnimContextsPair;
 
 	static void ProcessAvailabilityRequestsMap(const FAvailabilityRequestsMap& AvailabilityRequestsMap, UWorld* SubsystemWorld, FDatabaseToRoledAnimContexts& DatabaseToRoledAnimContexts)
@@ -202,7 +202,7 @@ namespace UE::PoseSearch
 				}
 				else
 				{
-					for (const UE::PoseSearch::FRole& Role : AvailabilityRequest.RolesFilter)
+					for (const FRole& Role : AvailabilityRequest.RolesFilter)
 					{
 						if (AvailabilityRequest.Database->Schema->GetRoledSkeleton(Role))
 						{
@@ -236,7 +236,7 @@ namespace UE::PoseSearch
 	template <typename DataType, typename EvaluateCombinationType>
 	static void GeneratePermutations(const TArrayView<DataType> Data, int32 CombinationCardinality, EvaluateCombinationType EvaluateCombination)
 	{
-		TArray<int32, TInlineAllocator<UE::PoseSearch::PreallocatedRolesNum>> Combination;
+		TArray<int32, TInlineAllocator<PreallocatedRolesNum>> Combination;
 		Combination.SetNum(CombinationCardinality);
 		GeneratePermutationsRecursive(Data, 0, Combination, 0, EvaluateCombination);
 	}
@@ -466,8 +466,8 @@ void UPoseSearchInteractionSubsystem::Tick(float DeltaSeconds)
 	//
 	// CharB/RoleA - CharB/RoleB -> invalid because of same CharB
 
-	TSet<const UObject*, DefaultKeyFuncs<const UObject*>, TInlineSetAllocator<UE::PoseSearch::PreallocatedRolesNum>> AnimContexts;
-	TSet<UE::PoseSearch::FRole, DefaultKeyFuncs<UE::PoseSearch::FRole>, TInlineSetAllocator<UE::PoseSearch::PreallocatedRolesNum>> CoveredRoles;
+	TSet<const UObject*, DefaultKeyFuncs<const UObject*>, TInlineSetAllocator<PreallocatedRolesNum>> AnimContexts;
+	TSet<FRole, DefaultKeyFuncs<FRole>, TInlineSetAllocator<PreallocatedRolesNum>> CoveredRoles;
 	TArray<FInteractionSearchContext, TInlineAllocator<PreallocatedSearchesNum, TMemStackAllocator<>>> SearchContexts;
 	for (const FDatabaseToRoledAnimContextsPair& DatabaseToRoledAnimContextsPair : DatabaseToRoledAnimContexts)
 	{
@@ -607,7 +607,7 @@ void UPoseSearchInteractionSubsystem::Tick(float DeltaSeconds)
 
 	for (const FInteractionSearchContextGroup& SearchContextGroup : SearchContextGroups)
 	{
-		UE::PoseSearch::FInteractionIsland& Island = GetAvailableIsland();
+		FInteractionIsland& Island = GetAvailableIsland();
 
 		for (const UObject* AnimContext : SearchContextGroup.AnimContexts)
 		{
@@ -638,52 +638,58 @@ void UPoseSearchInteractionSubsystem::Query_AnyThread(const TArrayView<const FPo
 
 	Result = FPoseSearchInteractionBlueprintResult();
 
-	if (AnimContext)
+	// if we find AnimContext in an island, we perform ALL the Island motion matching searches.
+	// thread safety is ensured by the lock on Island->SearchResultsMutex in DoSearch_AnyThread
+	if (FInteractionIsland* Island = FindIsland(AnimContext))
 	{
-		// if we find AnimContext in an island, we perform ALL the Island motion matching searches.
-		// thread safety is ensured by the lock on Island->SearchResultsMutex in DoSearch_AnyThread
-		if (FInteractionIsland* Island = FindIsland(AnimContext))
+		Island->DoSearch_AnyThread(AnimContext, ContinuingProperties, Result);
+
+		if (bValidateResultAgainstAvailabilities && Result.SelectedAnimation)
 		{
-			Island->DoSearch_AnyThread(AnimContext, ContinuingProperties, Result);
+			bool bResultValidated = false;
 
-			if (bValidateResultAgainstAvailabilities && Result.SelectedAnimation)
+			for (const FPoseSearchInteractionAvailability& Availability : Availabilities)
 			{
-				bool bResultValidated = false;
-
-				for (const FPoseSearchInteractionAvailability& Availability : Availabilities)
+				if (Availability.Database == Result.SelectedDatabase &&
+					(Availability.RolesFilter.IsEmpty() || Availability.RolesFilter.Contains(Result.Role)))
 				{
-					if (Availability.Database == Result.SelectedDatabase &&
-						(Availability.RolesFilter.IsEmpty() || Availability.RolesFilter.Contains(Result.Role)))
-					{
-						bResultValidated = true;
-						break;
-					}
+					bResultValidated = true;
+					break;
 				}
+			}
 
-				if (!bResultValidated)
-				{
-					Result = FPoseSearchInteractionBlueprintResult();
-				}
+			if (!bResultValidated)
+			{
+				Result = FPoseSearchInteractionBlueprintResult();
 			}
 		}
+	}
 
-		// queuing the availabilities for the next frame Query_AnyThread
-		if (!Availabilities.IsEmpty())
-		{
-			FScopeLock Lock(&AvailabilityRequestsMapMutex);
+	// queuing the availabilities for the next frame Query_AnyThread
+	if (!Availabilities.IsEmpty())
+	{
+		FScopeLock Lock(&AvailabilityRequestsMapMutex);
 
-			TArray<FInteractionAvailabilityEx>& AvailabilityRequests = AvailabilityRequestsMap.FindOrAdd(AnimContext);
-			for (const FPoseSearchInteractionAvailability& Availability : Availabilities)
+		TArray<FInteractionAvailabilityEx>& AvailabilityRequests = AvailabilityRequestsMap.FindOrAdd(AnimContext);
+		for (const FPoseSearchInteractionAvailability& Availability : Availabilities)
+		{
+			if (Availability.Database)
 			{
-				if (Availability.Database)
-				{
-					AvailabilityRequests.Emplace(Availability, PoseHistoryName, HistoryCollector);
-				}
+				AvailabilityRequests.Emplace(Availability, PoseHistoryName, HistoryCollector);
 			}
 		}
 	}
 }
 
+UE::PoseSearch::FInteractionSearchContext UPoseSearchInteractionSubsystem::FindSearchContext_AnyThread(UObject* AnimContext)
+{
+	if (UE::PoseSearch::FInteractionIsland* Island = FindIsland(AnimContext))
+	{
+		return Island->FindSearchContext_AnyThread(AnimContext);
+	}
+	return UE::PoseSearch::FInteractionSearchContext();
+}
+
 void UPoseSearchInteractionSubsystem::DebugDraw() const
 {
 	using namespace UE::PoseSearch;
@@ -701,116 +707,5 @@ void UPoseSearchInteractionSubsystem::DebugDraw() const
 			CurrentColorIndex = (CurrentColorIndex + 1) % NumColors;
 		}
 	}
-
 #endif //ENABLE_DRAW_DEBUG
-
-#if ENABLE_VISUAL_LOG
-	TStringBuilder<512> StringBuilder;
-
-	for (FAvailabilityRequestsMap::TConstIterator Iter = AvailabilityRequestsMap.CreateConstIterator(); Iter; ++Iter)
-	{
-		// making sure the weak pointer is still valid
-		if (UObject* AnimContext = Iter->Key.Get())
-		{
-			// looking for valid results for this AnimContext
-			const UPoseSearchDatabase* ResultDatabase = nullptr;
-			UE::PoseSearch::FRole ResultRole;
-			for (const FInteractionIsland* Island : Islands)
-			{
-				if (!Island->IsUninjected())
-				{
-					for (const FInteractionSearchResult& SearchResult : Island->GetSearchResults())
-					{
-						const FInteractionSearchContext& SearchContext = Island->GetSearchContexts()[SearchResult.SearchIndex];
-						for (int32 AnimInstanceIndex = 0; AnimInstanceIndex < SearchContext.AnimContexts.Num(); ++AnimInstanceIndex)
-						{
-							if (SearchContext.AnimContexts[AnimInstanceIndex].Get() == AnimContext)
-							{
-								ResultDatabase = SearchResult.Database.Get();
-								check(ResultDatabase);
-								ResultRole = SearchContext.Roles[AnimInstanceIndex];
-								break;
-							}
-						}
-
-						if (ResultDatabase)
-						{
-							break;
-						}
-					}
-				}
-
-				if (ResultDatabase)
-				{
-					break;
-				}
-			}
-
-			StringBuilder.Reset();
-			for (const FInteractionAvailabilityEx& AvailabilityRequest : Iter->Value)
-			{
-				StringBuilder.Append(GetNameSafe(AvailabilityRequest.Database));
-				StringBuilder.Append(" / ");
-				StringBuilder.Append(AvailabilityRequest.GetPoseHistoryName());
-				StringBuilder.Append(" [");
-				bool bAddComma = false;
-
-				if (AvailabilityRequest.RolesFilter.IsEmpty())
-				{
-					// adding ALL the possible roles from the database:
-					if (AvailabilityRequest.Database && AvailabilityRequest.Database->Schema)
-					{
-						for (const FPoseSearchRoledSkeleton& RoledSkeleton : AvailabilityRequest.Database->Schema->GetRoledSkeletons())
-						{
-							if (bAddComma)
-							{
-								StringBuilder.Append(",");
-							}
-							else
-							{
-								bAddComma = true;
-							}
-
-							if (AvailabilityRequest.Database == ResultDatabase && RoledSkeleton.Role == ResultRole)
-							{
-								StringBuilder.Append("(*)");
-							}
-
-							StringBuilder.Append(RoledSkeleton.Role.ToString());
-						}
-					}
-				}
-				else
-				{
-					for (const FName& Role : AvailabilityRequest.RolesFilter)
-					{
-						if (bAddComma)
-						{
-							StringBuilder.Append(",");
-						}
-						else
-						{
-							bAddComma = true;
-						}
-
-						if (AvailabilityRequest.Database == ResultDatabase && Role == ResultRole)
-						{
-							StringBuilder.Append("(*)");
-						}
-
-						StringBuilder.Append(Role.ToString());
-					}
-				}
-
-				StringBuilder.Append("]");
-				StringBuilder.Append("\n");
-			}
-			
-			// @todo: waiting for visual logger fix to display strings :/
-			const FVector StringOffset(0.f, 0.f, 0.001f);
-			const FVector ActorLocation = GetContextLocation(AnimContext);
-			UE_VLOG_SEGMENT(AnimContext, "PoseSearchInteraction", Display, ActorLocation, ActorLocation + StringOffset, FColor::Transparent, TEXT("%s"), StringBuilder.ToString());
-		}
-	}
-#endif // ENABLE_VISUAL_LOG
 }
\ No newline at end of file
diff --git a/Engine/Plugins/Animation/PoseSearch/Source/Runtime/Public/PoseSearch/AnimNode_MotionMatchingInteraction.h b/Engine/Plugins/Animation/PoseSearch/Source/Runtime/Public/PoseSearch/AnimNode_MotionMatchingInteraction.h
--- a/Engine/Plugins/Animation/PoseSearch/Source/Runtime/Public/PoseSearch/AnimNode_MotionMatchingInteraction.h
+++ b/Engine/Plugins/Animation/PoseSearch/Source/Runtime/Public/PoseSearch/AnimNode_MotionMatchingInteraction.h
@@ -6,52 +6,50 @@
 #include "PoseSearch/PoseSearchInteractionLibrary.h"
 #include "AnimNode_MotionMatchingInteraction.generated.h"
 
-struct FAnimationInitializeContext;
 struct FNodeDebugData;
 
-UENUM(Experimental)
-enum class EMotionMatchingInteractionEvaluationMode
+USTRUCT(Experimental, BlueprintInternalUseOnly)
+struct POSESEARCH_API FAnimNode_MotionMatchingInteraction : public FAnimNode_BlendStack_Standalone
 {
-	// Node will continuously provide its availabilities and eventually blend to newly selected animations
-	ContinuousReselection,
+	GENERATED_BODY();
 
-	// Node will continuously provide its availabilities to keep the interaction alive, but will play only the first selected animation
-	// the idea is to let the animation play until the end and allow the evenutual state machine playing this node to be able to perform an automatic transition
-	SingleSelection,
+	UPROPERTY(EditAnywhere, Category = Search, meta = (PinHiddenByDefault))
+	TArray<FPoseSearchInteractionAvailability> Availabilities;
 
-	// @todo: is this needed?
-	// Node will stop providing its availabilities and consequently kill the interaction, when the first selected animation stop playing as continuing pose
-	// the idea is to let the animation play until its valid as continuing pose and allow the evenutual state machine playing this node to be able to perform an automatic transition
-	// UntilContinuingPoseIsValid,
-};
+	UPROPERTY(EditAnywhere, Category = Search, meta = (PinHiddenByDefault))
+	bool bValidateResultAgainstAvailabilities = true;
 
-USTRUCT(Experimental, BlueprintInternalUseOnly)
-struct POSESEARCH_API FAnimNode_MotionMatchingInteraction : public FAnimNode_BlendStack
-{
-	GENERATED_BODY();
+	// Controls how fast the warp translation offset is blended in
+	// Values closer to 0 make it faster. if 0 the warp translation is applied in full
+	UPROPERTY(EditAnywhere, Category = Warping, meta = (PinHiddenByDefault))
+	float WarpingTranslationHalflife = 0.05f;
 
+	// Controls how fast the warp rotation offset is blended in
+	// Values closer to 0 make it faster. if 0 the warp rotation is applied in full
+	UPROPERTY(EditAnywhere, Category = Warping, meta = (PinHiddenByDefault))
+	float WarpingRotationHalfLife = 0.05f;
+
+	// Reset the blend stack if it has become relevant to the graph after not being updated on previous frames.
 	UPROPERTY(EditAnywhere, Category = Settings)
-	EMotionMatchingInteractionEvaluationMode EvaluationMode = EMotionMatchingInteractionEvaluationMode::ContinuousReselection;
+	bool bResetOnBecomingRelevant = true;
 
-	UPROPERTY(EditAnywhere, Category = Settings, meta = (PinHiddenByDefault))
-	TArray<FPoseSearchInteractionAvailability> Availabilities;
+	// tunable animation transition blend time 
+	UPROPERTY(EditAnywhere, BlueprintReadWrite, Category = Settings, meta = (PinHiddenByDefault, ClampMin = "0"))
+	float BlendTime = 0.2f;
 
-	UPROPERTY(EditAnywhere, Category = Settings, meta = (PinHiddenByDefault))
-	bool bValidateResultAgainstAvailabilities = true;
+	UPROPERTY(EditAnywhere, BlueprintReadWrite, Category = Settings, meta = (PinHiddenByDefault, UseAsBlendProfile = true))
+	TObjectPtr<UBlendProfile> BlendProfile;
 
-	UPROPERTY(EditAnywhere, Category = Settings, meta = (PinHiddenByDefault))
-	bool bEnableWarping = true;
+	UPROPERTY(EditAnywhere, BlueprintReadWrite, Category = Settings, meta = (PinHiddenByDefault))
+	EAlphaBlendOption BlendOption = EAlphaBlendOption::Linear;
 
-	UPROPERTY(EditAnywhere, Category = Settings, meta = (PinHiddenByDefault))
-	bool bUseAnimRootMotionProvider = true;
-	
-	// time from the beginning of the interaction to warp to full translation alignment
-	UPROPERTY(EditAnywhere, BlueprintReadWrite, Category = Warping, meta = (PinHiddenByDefault, ClampMin = "0"))
-	float InitialTranslationWarpTime = 0.2f;
+	// How we should update individual blend space parameters. See dropdown options tooltips.
+	UPROPERTY(EditAnywhere, Category = Settings)
+	EBlendStack_BlendspaceUpdateMode BlendspaceUpdateMode = EBlendStack_BlendspaceUpdateMode::InitialOnly;
 
-	// time from the beginning of the interaction to warp to full rotation alignment
-	UPROPERTY(EditAnywhere, BlueprintReadWrite, Category = Warping, meta = (PinHiddenByDefault, ClampMin = "0"))
-	float InitialRotationWarpTime = 0.2f;
+	// tunable animation transition blend time 
+	UPROPERTY(EditAnywhere, Category = Settings)
+	bool bUseInertialBlend = false;
 
 	// FAnimNode_Base interface
 	virtual void GatherDebugData(FNodeDebugData& DebugData) override;
@@ -62,21 +60,14 @@ struct POSESEARCH_API FAnimNode_MotionMatchingInteraction : public FAnimNode_Ble
 
 	virtual void Reset() override;
 
-	float GetTranslationWarpLerp() const { return TranslationWarpLerp; }
-	float GetRotationWarpLerp() const { return RotationWarpLerp; }
+	bool IsInteracting() const { return CurrentResult.SelectedAnimation != nullptr; }
 
 protected:
-	float TranslationWarpLerp = 0.f;
-	float RotationWarpLerp = 0.f;
-	// if a search is successful InteractingRolesNum > 0. 
-	// if InteractingRolesNum == 1 it means that the search is a regular single character motion matching search
-	// if InteractingRolesNum > 1 it means this node is interacting with other actors via the Interaction Subsystem
-	int32 InteractingRolesNum = 0;
+	bool NeedsReset(const FAnimationUpdateContext& Context) const;
+
+	FPoseSearchInteractionBlueprintResult CurrentResult;
 
 	// Update Counter for detecting being relevant
 	FGraphTraversalCounter UpdateCounter;
-
-	// wanted world transform for full aligment interaction
-	FTransform FullAlignedActorRootBoneTransform = FTransform::Identity;
-	FTransform WarpRootMotionTransform = FTransform::Identity;
+	float CachedDeltaTime = 0.f;
 };
diff --git a/Engine/Plugins/Animation/PoseSearch/Source/Runtime/Public/PoseSearch/MotionMatchingInteractionAnimNodeLibrary.h b/Engine/Plugins/Animation/PoseSearch/Source/Runtime/Public/PoseSearch/MotionMatchingInteractionAnimNodeLibrary.h
--- a/Engine/Plugins/Animation/PoseSearch/Source/Runtime/Public/PoseSearch/MotionMatchingInteractionAnimNodeLibrary.h
+++ b/Engine/Plugins/Animation/PoseSearch/Source/Runtime/Public/PoseSearch/MotionMatchingInteractionAnimNodeLibrary.h
@@ -38,8 +38,5 @@ public:
 	static void SetAvailabilities(const FMotionMatchingInteractionAnimNodeReference& MotionMatchingInteractionNode, const TArray<FPoseSearchInteractionAvailability>& Availabilities);
 
 	UFUNCTION(BlueprintPure, Category = "Animation|Pose Search", meta = (BlueprintThreadSafe))
-	static float GetTranslationWarpLerp(const FMotionMatchingInteractionAnimNodeReference& MotionMatchingInteractionNode);
-
-	UFUNCTION(BlueprintPure, Category = "Animation|Pose Search", meta = (BlueprintThreadSafe))
-	static float GetRotationWarpLerp(const FMotionMatchingInteractionAnimNodeReference& MotionMatchingInteractionNode);
+	static bool IsInteracting(const FMotionMatchingInteractionAnimNodeReference& MotionMatchingInteractionNode);
 };
\ No newline at end of file
diff --git a/Engine/Plugins/Animation/PoseSearch/Source/Runtime/Public/PoseSearch/PoseSearchInteractionIsland.h b/Engine/Plugins/Animation/PoseSearch/Source/Runtime/Public/PoseSearch/PoseSearchInteractionIsland.h
--- a/Engine/Plugins/Animation/PoseSearch/Source/Runtime/Public/PoseSearch/PoseSearchInteractionIsland.h
+++ b/Engine/Plugins/Animation/PoseSearch/Source/Runtime/Public/PoseSearch/PoseSearchInteractionIsland.h
@@ -17,10 +17,10 @@ struct FInteractionSearchContext
 {
 	// @todo: since AnimContexts, HistoryCollectors, Roles, and BroadPhaseRadiuses have the same cardinality,
 	// add an array of structs instead of many arrays
-	TArray<TWeakObjectPtr<const UObject>> AnimContexts;
+	TArray<TWeakObjectPtr<const UObject>, TInlineAllocator<PreallocatedRolesNum>> AnimContexts;
 	// @todo: perhaps use the animnode of the pose history collector instead
-	TArray<const FAnimNode_PoseSearchHistoryCollector_Base*> HistoryCollectors;
-	TArray<UE::PoseSearch::FRole> Roles;
+	TArray<const FAnimNode_PoseSearchHistoryCollector_Base*, TInlineAllocator<PreallocatedRolesNum>> HistoryCollectors;
+	TArray<UE::PoseSearch::FRole, TInlineAllocator<PreallocatedRolesNum>> Roles;
 
 #if ENABLE_DRAW_DEBUG
 	TArray<float> BroadPhaseRadiuses;
@@ -31,15 +31,18 @@ struct FInteractionSearchContext
 	TWeakObjectPtr<const UPoseSearchDatabase> Database;
 	FPoseSearchContinuingProperties ContinuingProperties;
 
-	bool IsValid() const;
+	bool IsValid() const { return !AnimContexts.IsEmpty(); }
 	bool IsEquivalent(const FInteractionSearchContext& Other) const;
+
+#if DO_CHECK
+	bool CheckForConsistency() const;
+#endif // DO_CHECK
 };
 
 // Experimental, this feature might be removed without warning, not for production use
 struct FInteractionSearchResult : public FSearchResult
 {
 	int32 SearchIndex = INDEX_NONE;
-	TArray<FTransform, TInlineAllocator<PreallocatedRolesNum>> FullAlignedActorRootBoneTransforms;
 
 	bool operator==(const FInteractionSearchResult& Other) const;
 };
@@ -57,6 +60,7 @@ struct FInteractionIsland
 	~FInteractionIsland();
 
 	bool DoSearch_AnyThread(UObject* AnimContext, const FPoseSearchContinuingProperties& ContinuingProperties, FPoseSearchInteractionBlueprintResult& Result);
+	UE::PoseSearch::FInteractionSearchContext FindSearchContext_AnyThread(UObject* AnimContext);
 
 	const TArray<TWeakObjectPtr<UActorComponent>>& GetPreTickComponents() const { return PreTickComponents; }
 	const TArray<TWeakObjectPtr<UActorComponent>>& GetPostTickComponents() const { return PostTickComponents; }
diff --git a/Engine/Plugins/Animation/PoseSearch/Source/Runtime/Public/PoseSearch/PoseSearchInteractionLibrary.h b/Engine/Plugins/Animation/PoseSearch/Source/Runtime/Public/PoseSearch/PoseSearchInteractionLibrary.h
--- a/Engine/Plugins/Animation/PoseSearch/Source/Runtime/Public/PoseSearch/PoseSearchInteractionLibrary.h
+++ b/Engine/Plugins/Animation/PoseSearch/Source/Runtime/Public/PoseSearch/PoseSearchInteractionLibrary.h
@@ -3,20 +3,15 @@
 #pragma once
 
 #include "Kismet/BlueprintFunctionLibrary.h"
+#include "PoseSearch/PoseSearchInteractionIsland.h"
 #include "PoseSearch/PoseSearchLibrary.h"
 #include "PoseSearchInteractionLibrary.generated.h"
 
-// subsystem in charge of coordinating Characters availability for motion matched interactions,
-// scheduling motion matching searches and synchronizing animations playback
-
-
-// input for MotionMatchInteraction_Pure: it declares that the associated character (AnimInstance) is willing to partecipate in an interction 
-// described by a UMultiAnimAsset (derived by UPoseSearchInteractionAsset) contained in the UPoseSearchDatabase Database
-// with one of the roles in RolesFilter (if empty ANY of the Database roles can be taken)
-// the MotionMatchInteraction_Pure will ultimately setup a motion matching query using looking for the pose history "PoseHistoryName" 
-// to gather bone and trajectory positions for this character
-// for an interaction to be valid, the query needs to find all the other interacting characters within BroadPhaseRadius,
-// and reach a maximum cost of MaxCost
+// input for MotionMatchInteraction_Pure: it declares that the associated character ("AnimContext" that could be an AnimInstance or an AnimNextCharacterComponent)
+// is willing to partecipate in an interaction described by a UMultiAnimAsset (derived by UPoseSearchInteractionAsset) contained in the UPoseSearchDatabase Database
+// with one of the roles in RolesFilter (if empty ANY of the Database roles can be taken) the MotionMatchInteraction_Pure will ultimately setup a motion matching query
+// using looking for the pose history "PoseHistoryName" to gather bone and trajectory positions for this character for an interaction to be valid,
+// the query needs to find all the other interacting characters within BroadPhaseRadius, and reach a maximum cost of MaxCost
 USTRUCT(Experimental, BlueprintType, Category = "Animation|Pose Search")
 struct POSESEARCH_API FPoseSearchInteractionAvailability
 {
@@ -84,10 +79,6 @@ public:
 	// assigned role to this character (AnimInstance)
 	UPROPERTY(Transient, VisibleAnywhere, BlueprintReadOnly, Category=State)
 	FName Role;
-
-	// root bone transform for the character at full aligment
-	UPROPERTY(Transient, VisibleAnywhere, BlueprintReadOnly, Category=State)
-	FTransform FullAlignedActorRootBoneTransform = FTransform::Identity;	
 };
 
 UCLASS(Experimental)
@@ -108,14 +99,18 @@ public:
 	// PoseHistoryName is the name of the pose history node used for the associated motion matching search
 	// if bValidateResultAgainstAvailabilities is true, the result will be invalidated if doesn't respect the new availabilities
 	UFUNCTION(Experimental, BlueprintPure, Category = "Animation|Pose Search", meta = (BlueprintThreadSafe))
-	static FPoseSearchInteractionBlueprintResult MotionMatchInteraction_Pure(TArray<FPoseSearchInteractionAvailability> Availabilities, UObject* AnimInstance, FPoseSearchContinuingProperties ContinuingProperties, FName PoseHistoryName, bool bValidateResultAgainstAvailabilities);
+	static FPoseSearchInteractionBlueprintResult MotionMatchInteraction_Pure(TArray<FPoseSearchInteractionAvailability> Availabilities, UObject* AnimInstance, FPoseSearchContinuingProperties ContinuingProperties, FName PoseHistoryName, bool bValidateResultAgainstAvailabilities = true);
 
 	// BlueprintCallable version of MotionMatchInteraction_Pure
 	UFUNCTION(Experimental, BlueprintCallable, Category = "Animation|Pose Search", meta = (BlueprintThreadSafe))
-	static FPoseSearchInteractionBlueprintResult MotionMatchInteraction(TArray<FPoseSearchInteractionAvailability> Availabilities, UObject* AnimInstance, FPoseSearchContinuingProperties ContinuingProperties, FName PoseHistoryName, bool bValidateResultAgainstAvailabilities);
+	static FPoseSearchInteractionBlueprintResult MotionMatchInteraction(TArray<FPoseSearchInteractionAvailability> Availabilities, UObject* AnimInstance, FPoseSearchContinuingProperties ContinuingProperties, FName PoseHistoryName, bool bValidateResultAgainstAvailabilities = true);
 
 	// version of MotionMatchInteraction_Pure referencing directly the HistoryCollector rather than looking for it by name
-	static FPoseSearchInteractionBlueprintResult MotionMatchInteraction(const TArrayView<const FPoseSearchInteractionAvailability> Availabilities, UObject* AnimInstance, const FPoseSearchContinuingProperties& ContinuingProperties, const FAnimNode_PoseSearchHistoryCollector_Base* HistoryCollector, bool bValidateResultAgainstAvailabilities);
+	static FPoseSearchInteractionBlueprintResult MotionMatchInteraction(const TArrayView<const FPoseSearchInteractionAvailability> Availabilities, UObject* AnimInstance, const FPoseSearchContinuingProperties& ContinuingProperties, const FAnimNode_PoseSearchHistoryCollector_Base* HistoryCollector, bool bValidateResultAgainstAvailabilities = true);
+
+	// returns the associated FInteractionSearchContext to AnimContext. We do a copy to be thread safe.
+	// check FInteractionSearchContext::IsValid for its validity
+	static UE::PoseSearch::FInteractionSearchContext FindSearchContext_AnyThread(UObject* AnimContext);
 
 	UFUNCTION(Experimental, BlueprintPure, Category = "Animation|Pose Search", meta = (BlueprintThreadSafe))
 	static FPoseSearchContinuingProperties GetMontageContinuingProperties(UAnimInstance* AnimInstance);
diff --git a/Engine/Plugins/Animation/PoseSearch/Source/Runtime/Public/PoseSearch/PoseSearchInteractionSubsystem.h b/Engine/Plugins/Animation/PoseSearch/Source/Runtime/Public/PoseSearch/PoseSearchInteractionSubsystem.h
--- a/Engine/Plugins/Animation/PoseSearch/Source/Runtime/Public/PoseSearch/PoseSearchInteractionSubsystem.h
+++ b/Engine/Plugins/Animation/PoseSearch/Source/Runtime/Public/PoseSearch/PoseSearchInteractionSubsystem.h
@@ -13,9 +13,6 @@ class UPoseSearchDatabase;
 namespace UE::PoseSearch
 {
 
-struct FInteractionIsland;
-struct FInteractionSearchContext;
-
 // Experimental, this feature might be removed without warning, not for production use
 struct FInteractionAvailabilityEx : public FPoseSearchInteractionAvailability
 {
@@ -101,6 +98,11 @@ public:
 	void Query_AnyThread(const TArrayView<const FPoseSearchInteractionAvailability> Availabilities, UObject* AnimContext, 
 		const FPoseSearchContinuingProperties& ContinuingProperties, FPoseSearchInteractionBlueprintResult& Result,
 		FName PoseHistoryName, const FAnimNode_PoseSearchHistoryCollector_Base* HistoryCollector, bool bValidateResultAgainstAvailabilities);
+
+	// returns the associated FInteractionSearchContext to AnimContext. We do a copy to be thread safe.
+	// check FInteractionSearchContext::IsValid for its validity
+	UE::PoseSearch::FInteractionSearchContext FindSearchContext_AnyThread(UObject* AnimContext);
+
 private:
 	UE::PoseSearch::FInteractionIsland& CreateIsland();
 	UE::PoseSearch::FInteractionIsland& GetAvailableIsland();
diff --git a/Engine/Plugins/Animation/PoseSearch/Source/Runtime/Public/PoseSearch/PoseSearchRole.h b/Engine/Plugins/Animation/PoseSearch/Source/Runtime/Public/PoseSearch/PoseSearchRole.h
--- a/Engine/Plugins/Animation/PoseSearch/Source/Runtime/Public/PoseSearch/PoseSearchRole.h
+++ b/Engine/Plugins/Animation/PoseSearch/Source/Runtime/Public/PoseSearch/PoseSearchRole.h
@@ -33,6 +33,17 @@ namespace UE::PoseSearch
 		}
 		return true;
 	}
+
+	static FRoleToIndex MakeRoleToIndex(const TConstArrayView<FRole> Roles)
+	{
+		FRoleToIndex RoleToIndex;
+		RoleToIndex.Reserve(Roles.Num());
+		for (int32 RoleIndex = 0; RoleIndex < Roles.Num(); ++RoleIndex)
+		{
+			RoleToIndex.Add(Roles[RoleIndex]) = RoleIndex;
+		}
+		return RoleToIndex;
+	}
 }
 
 
