PoseSearch - adding debug drawing support for UPoseSearchFeatureChannel_Position::MaxPositionDistanceSquared

diff --git a/Engine/Plugins/Animation/PoseSearch/Source/Runtime/Private/PoseSearchContext.cpp b/Engine/Plugins/Animation/PoseSearch/Source/Runtime/Private/PoseSearchContext.cpp
--- a/Engine/Plugins/Animation/PoseSearch/Source/Runtime/Private/PoseSearchContext.cpp
+++ b/Engine/Plugins/Animation/PoseSearch/Source/Runtime/Private/PoseSearchContext.cpp
@@ -389,6 +389,22 @@ void FDebugDrawParams::DrawCircle(const FMatrix& TransformMatrix, float Radius,
 	}
 }
 
+void FDebugDrawParams::DrawSphere(const FVector& Center, float Radius, int32 Segments, const FColor& Color, float Thickness) const
+{
+	if (Color.A > 0)
+	{
+		if (FAnimInstanceProxy* AnimInstanceProxy = GetAnimInstanceProxy(AnimContexts))
+		{
+			AnimInstanceProxy->AnimDrawDebugSphere(Center, Radius, Segments, Color, false, 0.f, Thickness, SDPG_Foreground);
+		}
+		else if (const UWorld* World = GetWorld(AnimContexts))
+		{
+			// any Mesh is fine to draw
+			DrawDebugSphere(World, Center, Radius, Segments, Color, false, 0.f, SDPG_Foreground, Thickness);
+		}
+	}
+}
+
 void FDebugDrawParams::DrawCentripetalCatmullRomSpline(TConstArrayView<FVector> Points, TConstArrayView<FColor> Colors, float Alpha, int32 NumSamplesPerSegment, float Thickness) const
 {
 	const int32 NumPoints = Points.Num();
diff --git a/Engine/Plugins/Animation/PoseSearch/Source/Runtime/Private/PoseSearchFeatureChannel_Position.cpp b/Engine/Plugins/Animation/PoseSearch/Source/Runtime/Private/PoseSearchFeatureChannel_Position.cpp
--- a/Engine/Plugins/Animation/PoseSearch/Source/Runtime/Private/PoseSearchFeatureChannel_Position.cpp
+++ b/Engine/Plugins/Animation/PoseSearch/Source/Runtime/Private/PoseSearchFeatureChannel_Position.cpp
@@ -209,6 +209,12 @@ void UPoseSearchFeatureChannel_Position::DebugDraw(const UE::PoseSearch::FDebugD
 		const FVector BonePos = OriginBonePos + DeltaPos;
 		DrawParams.DrawPoint(BonePos, Color);
 
+		if (MaxPositionDistanceSquared > 0.f)
+		{
+			static int32 Segments = 8;
+			DrawParams.DrawSphere(BonePos, FMath::Sqrt(MaxPositionDistanceSquared), Segments, Color);
+		}
+
 		const bool bDrawOrigin = !DeltaPos.IsNearlyZero() && (SchemaOriginBoneIdx != RootSchemaBoneIdx || !FMath::IsNearlyZero(OriginTimeOffset) ||
 			SampleRole != OriginRole || PermutationTimeType != EPermutationTimeType::UseSampleTime || bUseBlueprintQueryOverride);
 		if (bDrawOrigin)
diff --git a/Engine/Plugins/Animation/PoseSearch/Source/Runtime/Public/PoseSearch/PoseSearchContext.h b/Engine/Plugins/Animation/PoseSearch/Source/Runtime/Public/PoseSearch/PoseSearchContext.h
--- a/Engine/Plugins/Animation/PoseSearch/Source/Runtime/Public/PoseSearch/PoseSearchContext.h
+++ b/Engine/Plugins/Animation/PoseSearch/Source/Runtime/Public/PoseSearch/PoseSearchContext.h
@@ -73,7 +73,8 @@ struct POSESEARCH_API FDebugDrawParams
 
 	void DrawLine(const FVector& LineStart, const FVector& LineEnd, const FColor& Color, float Thickness = 0.f) const;
 	void DrawPoint(const FVector& Position, const FColor& Color, float Thickness = 6.f) const;
-	void DrawCircle(const FMatrix& TransformMatrix, float Radius, int32 Segments, const FColor& Color, float Thickness = 1.f) const;
+	void DrawCircle(const FMatrix& TransformMatrix, float Radius, int32 Segments, const FColor& Color, float Thickness = 0.f) const;
+	void DrawSphere(const FVector& Center, float Radius, int32 Segments, const FColor& Color, float Thickness = 0.f) const;
 	void DrawCentripetalCatmullRomSpline(TConstArrayView<FVector> Points, TConstArrayView<FColor> Colors, float Alpha, int32 NumSamplesPerSegment, float Thickness = 1.f) const;
 	
 	void DrawFeatureVector(TConstArrayView<float> PoseVector);
