PoseSearch - fix for searching empty databases when !WITH_EDITOR

diff --git a/Engine/Plugins/Animation/PoseSearch/Source/Runtime/Private/PoseSearchDatabase.cpp b/Engine/Plugins/Animation/PoseSearch/Source/Runtime/Private/PoseSearchDatabase.cpp
--- a/Engine/Plugins/Animation/PoseSearch/Source/Runtime/Private/PoseSearchDatabase.cpp
+++ b/Engine/Plugins/Animation/PoseSearch/Source/Runtime/Private/PoseSearchDatabase.cpp
@@ -1101,6 +1101,11 @@ UE::PoseSearch::FSearchResult UPoseSearchDatabase::Search(UE::PoseSearch::FSearc
 		SearchContext.SetAsyncBuildIndexInProgress();
 		return Result;
 	}
+#else
+	if (SearchIndexPrivate.IsEmpty())
+	{
+		return Result;
+	}
 #endif // WITH_EDITOR
 
 	switch (PoseSearchMode)
