PoseSearch - support for channels in trajectory space

diff --git a/Engine/Plugins/Animation/PoseSearch/Source/Runtime/Private/AnimNode_MotionMatchingInteraction.cpp b/Engine/Plugins/Animation/PoseSearch/Source/Runtime/Private/AnimNode_MotionMatchingInteraction.cpp
--- a/Engine/Plugins/Animation/PoseSearch/Source/Runtime/Private/AnimNode_MotionMatchingInteraction.cpp
+++ b/Engine/Plugins/Animation/PoseSearch/Source/Runtime/Private/AnimNode_MotionMatchingInteraction.cpp
@@ -205,7 +205,9 @@ void FAnimNode_MotionMatchingInteraction::Evaluate_AnyThread(FPoseContext& Outpu
 						static const TCHAR* LogName = TEXT("MotionMatchingInteraction");
 
 						UObject* LogOwner = Output.AnimInstanceProxy->GetAnimInstanceObject();
-						UE_VLOG_ARROW(LogOwner, LogName, Display, FullAlignedActorRootBoneTransform.GetLocation(), FullAlignedActorRootBoneTransform.GetLocation() + FullAlignedActorRootBoneTransform.GetRotation().GetForwardVector() * 100, FColor::Purple, TEXT(""));
+						const FTransform WorldWarpRootMotionTransform = WarpRootMotionTransform * RootBoneTransform;
+						UE_VLOG_ARROW(LogOwner, LogName, Display, WorldWarpRootMotionTransform.GetLocation(), WorldWarpRootMotionTransform.GetLocation() + WorldWarpRootMotionTransform.GetRotation().GetForwardVector() * 100, FColor::Green, TEXT(""));
+						UE_VLOG_ARROW(LogOwner, LogName, Display, FullAlignedActorRootBoneTransform.GetLocation(), FullAlignedActorRootBoneTransform.GetLocation() + FullAlignedActorRootBoneTransform.GetRotation().GetForwardVector() * 100, FColor::Blue, TEXT(""));
 					}
 #endif
 					RootMotionProvider->OverrideRootMotion(WarpRootMotionTransform, Output.CustomAttributes);
diff --git a/Engine/Plugins/Animation/PoseSearch/Source/Runtime/Private/PoseSearchAssetIndexer.cpp b/Engine/Plugins/Animation/PoseSearch/Source/Runtime/Private/PoseSearchAssetIndexer.cpp
--- a/Engine/Plugins/Animation/PoseSearch/Source/Runtime/Private/PoseSearchAssetIndexer.cpp
+++ b/Engine/Plugins/Animation/PoseSearch/Source/Runtime/Private/PoseSearchAssetIndexer.cpp
@@ -509,6 +509,11 @@ FTransform FAssetIndexer::GetComponentSpaceTransform(float SampleTime, const FRo
 {
 	using namespace UE::PoseSearch;
 
+	if (SchemaBoneIdx == TrajectorySchemaBoneIdx)
+	{
+		return FTransform::Identity;
+	}
+
 	FCachedEntry& Entry = GetEntry(SampleTime);
 	bClamped = Entry.bClamped;
 
@@ -526,6 +531,11 @@ FTransform FAssetIndexer::GetTransform(float SampleTime, const FRole& Role, bool
 	bClamped = Entry.bClamped;
 
 	const int32 RoleIndex = RoleToIndex[Role];
+	if (SchemaBoneIdx == TrajectorySchemaBoneIdx)
+	{
+		return Entry.RootTransform[RoleIndex];
+	}
+	
 	const FBoneReference& BoneReference = Schema.GetBoneReferences(Role)[SchemaBoneIdx];
 	return CalculateComponentSpaceTransform(Entry, BoneReference, RoleIndex) * Entry.RootTransform[RoleIndex];
 }
@@ -561,16 +571,6 @@ bool FAssetIndexer::GetSampleRotation(FQuat& OutSampleRotation, float SampleTime
 	const float SampleTime = Time + SampleTimeOffset + PermutationSampleTimeOffset;
 	const float OriginTime = Time + OriginTimeOffset + PermutationOriginTimeOffset;
 
-	// @todo: add support for SchemaSampleBoneIdx
-	if (SchemaOriginBoneIdx != RootSchemaBoneIdx)
-	{
-		UE_LOG(LogPoseSearch,
-			Error,
-			TEXT("FAssetIndexer::GetSampleRotation: support for non root origin bones not implemented (bone: '%s', schema: '%s'"), 
-			*Schema.GetBoneReferences(OriginRole)[SchemaOriginBoneIdx].BoneName.ToString(),
-			*GetNameSafe(&Schema));
-	}
-
 	if (SamplingAttributeId >= 0)
 	{
 		const FPoseSearchTimedNotifies<UAnimNotifyState_PoseSearchSamplingAttribute> TimedNotifies(SamplingAttributeId, *this);
diff --git a/Engine/Plugins/Animation/PoseSearch/Source/Runtime/Private/PoseSearchContext.cpp b/Engine/Plugins/Animation/PoseSearch/Source/Runtime/Private/PoseSearchContext.cpp
--- a/Engine/Plugins/Animation/PoseSearch/Source/Runtime/Private/PoseSearchContext.cpp
+++ b/Engine/Plugins/Animation/PoseSearch/Source/Runtime/Private/PoseSearchContext.cpp
@@ -169,7 +169,7 @@ FVector FDebugDrawParams::ExtractPosition(TConstArrayView<float> PoseVector, flo
 			{
 				if (const USkeleton* Skeleton = Schema->GetSkeleton(Role))
 				{
-					const FBoneIndexType BoneIndexType = Schema->GetBoneReferences(Role)[SchemaBoneIdx].BoneIndex;
+					const FBoneIndexType BoneIndexType = SchemaBoneIdx == TrajectorySchemaBoneIdx ? ComponentSpaceIndexType : Schema->GetBoneReferences(Role)[SchemaBoneIdx].BoneIndex;
 
 					FTransform WorldBoneTransform;
 					if (PoseHistory->GetTransformAtTime(SampleTimeOffset + PermutationSampleTimeOffset, WorldBoneTransform, Skeleton, BoneIndexType, WorldSpaceIndexType))
@@ -308,7 +308,7 @@ FQuat FDebugDrawParams::ExtractRotation(TConstArrayView<float> PoseVector, float
 			{
 				if (const USkeleton* Skeleton = Schema->GetSkeleton(Role))
 				{
-					const FBoneIndexType BoneIndexType = Schema->GetBoneReferences(Role)[SchemaBoneIdx].BoneIndex;
+					const FBoneIndexType BoneIndexType = SchemaBoneIdx == TrajectorySchemaBoneIdx ? ComponentSpaceIndexType : Schema->GetBoneReferences(Role)[SchemaBoneIdx].BoneIndex;
 
 					FTransform WorldBoneTransform;
 					if (PoseHistory->GetTransformAtTime(SampleTimeOffset + PermutationSampleTimeOffset, WorldBoneTransform, Skeleton, BoneIndexType, WorldSpaceIndexType))
@@ -384,7 +384,7 @@ void FDebugDrawParams::DrawCircle(const FMatrix& TransformMatrix, float Radius,
 		else if (const UWorld* World = GetWorld(AnimContexts))
 		{
 			// any Mesh is fine to draw
-			DrawDebugCircle(World, TransformMatrix, Radius, Segments, Color, false, 0.f, SDPG_Foreground, Thickness);
+			DrawDebugCircle(World, TransformMatrix, Radius, Segments, Color, false, 0.f, SDPG_Foreground, Thickness, false);
 		}
 	}
 }
@@ -729,9 +729,7 @@ FTransform FSearchContext::GetWorldBoneTransformAtTime(float SampleTime, const F
 	const UPoseSearchSchema* Schema = CachedQueries.Last().GetSchema();
 	check(Schema);
 
-	TConstArrayView<FBoneReference> BoneReferences = Schema->GetBoneReferences(SampleRole);
-	check(BoneReferences[SchemaBoneIdx].HasValidSetup());
-	const FBoneIndexType BoneIndexType = BoneReferences[SchemaBoneIdx].BoneIndex;
+	const FBoneIndexType BoneIndexType = SchemaBoneIdx == TrajectorySchemaBoneIdx ? ComponentSpaceIndexType : Schema->GetBoneReferences(SampleRole)[SchemaBoneIdx].BoneIndex;
 
 	const uint32 SampleTimeHash = GetTypeHash(SampleTime);
 	const uint32 SampleRoleHash = GetTypeHash(SampleRole);
diff --git a/Engine/Plugins/Animation/PoseSearch/Source/Runtime/Private/PoseSearchDerivedDataKey.cpp b/Engine/Plugins/Animation/PoseSearch/Source/Runtime/Private/PoseSearchDerivedDataKey.cpp
--- a/Engine/Plugins/Animation/PoseSearch/Source/Runtime/Private/PoseSearchDerivedDataKey.cpp
+++ b/Engine/Plugins/Animation/PoseSearch/Source/Runtime/Private/PoseSearchDerivedDataKey.cpp
@@ -57,7 +57,7 @@ FKeyBuilder::FKeyBuilder(const UObject* Object, bool bUseDataVer, bool bPerformC
 	if (bUseDataVer)
 	{
 		// used to invalidate the key without having to change POSESEARCHDB_DERIVEDDATA_VER all the times
-		int32 POSESEARCHDB_DERIVEDDATA_VER_SMALL = 2657;
+		int32 POSESEARCHDB_DERIVEDDATA_VER_SMALL = 2659;
 		FGuid VersionGuid = FDevSystemGuids::GetSystemGuid(FDevSystemGuids::Get().POSESEARCHDB_DERIVEDDATA_VER);
 		FString AnimationCompressionVersionString = UE::Anim::Compression::AnimationCompressionVersionString;
 
diff --git a/Engine/Plugins/Animation/PoseSearch/Source/Runtime/Private/PoseSearchFeatureChannel_Curve.h b/Engine/Plugins/Animation/PoseSearch/Source/Runtime/Private/PoseSearchFeatureChannel_Curve.h
--- a/Engine/Plugins/Animation/PoseSearch/Source/Runtime/Private/PoseSearchFeatureChannel_Curve.h
+++ b/Engine/Plugins/Animation/PoseSearch/Source/Runtime/Private/PoseSearchFeatureChannel_Curve.h
@@ -18,8 +18,8 @@ public:
 	UPROPERTY(EditAnywhere, Category = "Settings")
 	FName SampleRole = UE::PoseSearch::DefaultRole;
 
-	UPROPERTY()
-	int8 CurveIdx;
+	UPROPERTY(Transient)
+	int8 CurveIdx = UE::PoseSearch::InvalidSchemaCurveIdx;
 
 #if WITH_EDITORONLY_DATA
 	UPROPERTY(EditAnywhere, Category = "Settings")
diff --git a/Engine/Plugins/Animation/PoseSearch/Source/Runtime/Private/PoseSearchFeatureChannel_FilterCrashingLegs.cpp b/Engine/Plugins/Animation/PoseSearch/Source/Runtime/Private/PoseSearchFeatureChannel_FilterCrashingLegs.cpp
--- a/Engine/Plugins/Animation/PoseSearch/Source/Runtime/Private/PoseSearchFeatureChannel_FilterCrashingLegs.cpp
+++ b/Engine/Plugins/Animation/PoseSearch/Source/Runtime/Private/PoseSearchFeatureChannel_FilterCrashingLegs.cpp
@@ -29,6 +29,8 @@ static float ComputeCrashingLegsValue(const FVector& RightThighPos, const FVecto
 
 bool UPoseSearchFeatureChannel_FilterCrashingLegs::Finalize(UPoseSearchSchema* Schema)
 {
+	using namespace UE::PoseSearch;
+
 	ChannelDataOffset = Schema->SchemaCardinality;
 	ChannelCardinality = 1;
 	Schema->SchemaCardinality += ChannelCardinality;
@@ -38,7 +40,7 @@ bool UPoseSearchFeatureChannel_FilterCrashingLegs::Finalize(UPoseSearchSchema* S
 	LeftFootIdx = Schema->AddBoneReference(LeftFoot, SampleRole);
 	RightFootIdx = Schema->AddBoneReference(RightFoot, SampleRole);
 
-	return LeftThighIdx >= 0 && RightThighIdx >= 0 && LeftFootIdx >= 0 && RightFootIdx >= 0;
+	return LeftThighIdx != InvalidSchemaBoneIdx && RightThighIdx != InvalidSchemaBoneIdx && LeftFootIdx != InvalidSchemaBoneIdx && RightFootIdx != InvalidSchemaBoneIdx;
 }
 
 void UPoseSearchFeatureChannel_FilterCrashingLegs::AddDependentChannels(UPoseSearchSchema* Schema) const
diff --git a/Engine/Plugins/Animation/PoseSearch/Source/Runtime/Private/PoseSearchFeatureChannel_Heading.cpp b/Engine/Plugins/Animation/PoseSearch/Source/Runtime/Private/PoseSearchFeatureChannel_Heading.cpp
--- a/Engine/Plugins/Animation/PoseSearch/Source/Runtime/Private/PoseSearchFeatureChannel_Heading.cpp
+++ b/Engine/Plugins/Animation/PoseSearch/Source/Runtime/Private/PoseSearchFeatureChannel_Heading.cpp
@@ -24,6 +24,7 @@ void UPoseSearchFeatureChannel_Heading::FindOrAddToSchema(UPoseSearchSchema* Sch
 		{
 			if (const UPoseSearchFeatureChannel_Heading* Heading = Cast<UPoseSearchFeatureChannel_Heading>(Channel))
 			{
+				// @todo: channels are already finalized, so we can use SchemaBoneIdx and SchemaOriginBoneIdx instead of Bone.BoneName and OriginBone.BoneName
 				if (Heading->Bone.BoneName == BoneName &&
 					Heading->OriginBone.BoneName == NAME_None && 
 					Heading->SampleTimeOffset == SampleTimeOffset &&
@@ -31,7 +32,8 @@ void UPoseSearchFeatureChannel_Heading::FindOrAddToSchema(UPoseSearchSchema* Sch
 					Heading->HeadingAxis == HeadingAxis &&
 					Heading->PermutationTimeType == PermutationTimeType &&
 					Heading->SampleRole == Role &&
-					Heading->OriginRole == Role)
+					Heading->OriginRole == Role &&
+					Heading->bDefaultWithRootBone)
 				{
 					return Heading;
 				}
@@ -56,14 +58,16 @@ void UPoseSearchFeatureChannel_Heading::FindOrAddToSchema(UPoseSearchSchema* Sch
 
 bool UPoseSearchFeatureChannel_Heading::Finalize(UPoseSearchSchema* Schema)
 {
+	using namespace UE::PoseSearch;
+
 	ChannelDataOffset = Schema->SchemaCardinality;
 	ChannelCardinality = UE::PoseSearch::FFeatureVectorHelper::GetVectorCardinality(ComponentStripping);
 	Schema->SchemaCardinality += ChannelCardinality;
 
-	SchemaBoneIdx = Schema->AddBoneReference(Bone, SampleRole);
-	SchemaOriginBoneIdx = Schema->AddBoneReference(OriginBone, OriginRole);
+	SchemaBoneIdx = Schema->AddBoneReference(Bone, SampleRole, bDefaultWithRootBone);
+	SchemaOriginBoneIdx = Schema->AddBoneReference(OriginBone, OriginRole, bDefaultWithRootBone);
 
-	return SchemaBoneIdx >= 0 && SchemaOriginBoneIdx >= 0;
+	return SchemaBoneIdx != InvalidSchemaBoneIdx && SchemaOriginBoneIdx != InvalidSchemaBoneIdx;
 }
 
 void UPoseSearchFeatureChannel_Heading::AddDependentChannels(UPoseSearchSchema* Schema) const
@@ -286,14 +290,9 @@ UE::PoseSearch::TLabelBuilder& UPoseSearchFeatureChannel_Heading::GetLabel(UE::P
 		LabelBuilder.Append(TEXT("_"));
 		LabelBuilder.Append(Schema->GetBoneReferences(SampleRole)[SchemaBoneIdx].BoneName.ToString());
 	}
-
-	if (PermutationTimeType == EPermutationTimeType::UsePermutationTime)
-	{
-		LabelBuilder.Append(TEXT("_PT"));
-	}
-	else if (PermutationTimeType == EPermutationTimeType::UseSampleToPermutationTime)
+	else if (SchemaBoneIdx == TrajectorySchemaBoneIdx)
 	{
-		LabelBuilder.Append(TEXT("_SPT"));
+		LabelBuilder.Append(TEXT("_Trj"));
 	}
 
 	if (SampleRole != DefaultRole)
@@ -308,6 +307,10 @@ UE::PoseSearch::TLabelBuilder& UPoseSearchFeatureChannel_Heading::GetLabel(UE::P
 		LabelBuilder.Append(TEXT("_"));
 		LabelBuilder.Append(Schema->GetBoneReferences(OriginRole)[SchemaOriginBoneIdx].BoneName.ToString());
 	}
+	else if (SchemaOriginBoneIdx == TrajectorySchemaBoneIdx)
+	{
+		LabelBuilder.Append(TEXT("_Trj"));
+	}
 
 	if (OriginRole != DefaultRole)
 	{
@@ -316,6 +319,15 @@ UE::PoseSearch::TLabelBuilder& UPoseSearchFeatureChannel_Heading::GetLabel(UE::P
 		LabelBuilder.Append(TEXT("]"));
 	}
 
+	if (PermutationTimeType == EPermutationTimeType::UsePermutationTime)
+	{
+		LabelBuilder.Append(TEXT("_PT"));
+	}
+	else if (PermutationTimeType == EPermutationTimeType::UseSampleToPermutationTime)
+	{
+		LabelBuilder.Append(TEXT("_SPT"));
+	}
+
 	AppendLabelSeparator(LabelBuilder, LabelFormat, true);
 
 	LabelBuilder.Appendf(TEXT("%.2f"), SampleTimeOffset);
diff --git a/Engine/Plugins/Animation/PoseSearch/Source/Runtime/Private/PoseSearchFeatureChannel_Heading.h b/Engine/Plugins/Animation/PoseSearch/Source/Runtime/Private/PoseSearchFeatureChannel_Heading.h
--- a/Engine/Plugins/Animation/PoseSearch/Source/Runtime/Private/PoseSearchFeatureChannel_Heading.h
+++ b/Engine/Plugins/Animation/PoseSearch/Source/Runtime/Private/PoseSearchFeatureChannel_Heading.h
@@ -22,6 +22,12 @@ class POSESEARCH_API UPoseSearchFeatureChannel_Heading : public UPoseSearchFeatu
 	GENERATED_BODY()
 
 public:
+	// Experimental, this feature might be removed without warning, not for production use
+	// if bDefaultWithRootBone is true Bone.BoneName and OriginBone.BoneName get initialized to the associated
+	// skeleton root bone if equals to NAME_None, otherwise transforms from the trajectory will be used
+	UPROPERTY(EditAnywhere, Category = "Experimental")
+	bool bDefaultWithRootBone = true;
+
 	UPROPERTY(EditAnywhere, Category = "Settings")
 	FBoneReference Bone;
 
diff --git a/Engine/Plugins/Animation/PoseSearch/Source/Runtime/Private/PoseSearchFeatureChannel_Phase.cpp b/Engine/Plugins/Animation/PoseSearch/Source/Runtime/Private/PoseSearchFeatureChannel_Phase.cpp
--- a/Engine/Plugins/Animation/PoseSearch/Source/Runtime/Private/PoseSearchFeatureChannel_Phase.cpp
+++ b/Engine/Plugins/Animation/PoseSearch/Source/Runtime/Private/PoseSearchFeatureChannel_Phase.cpp
@@ -286,13 +286,15 @@ namespace UE::PoseSearch
 
 bool UPoseSearchFeatureChannel_Phase::Finalize(UPoseSearchSchema* Schema)
 {
+	using namespace UE::PoseSearch;
+
 	ChannelDataOffset = Schema->SchemaCardinality;
 	ChannelCardinality = 2;
 	Schema->SchemaCardinality += ChannelCardinality;
 
 	SchemaBoneIdx = Schema->AddBoneReference(Bone, SampleRole);
 
-	return SchemaBoneIdx >= 0;
+	return SchemaBoneIdx != InvalidSchemaBoneIdx;
 }
 
 void UPoseSearchFeatureChannel_Phase::AddDependentChannels(UPoseSearchSchema* Schema) const
diff --git a/Engine/Plugins/Animation/PoseSearch/Source/Runtime/Private/PoseSearchFeatureChannel_Position.cpp b/Engine/Plugins/Animation/PoseSearch/Source/Runtime/Private/PoseSearchFeatureChannel_Position.cpp
--- a/Engine/Plugins/Animation/PoseSearch/Source/Runtime/Private/PoseSearchFeatureChannel_Position.cpp
+++ b/Engine/Plugins/Animation/PoseSearch/Source/Runtime/Private/PoseSearchFeatureChannel_Position.cpp
@@ -23,13 +23,15 @@ void UPoseSearchFeatureChannel_Position::FindOrAddToSchema(UPoseSearchSchema* Sc
 		{
 			if (const UPoseSearchFeatureChannel_Position* Position = Cast<UPoseSearchFeatureChannel_Position>(Channel))
 			{
+				// @todo: channels are already finalized, so we can use SchemaBoneIdx and SchemaOriginBoneIdx instead of Bone.BoneName and OriginBone.BoneName
 				if (Position->Bone.BoneName == BoneName &&
 					Position->OriginBone.BoneName == NAME_None &&
 					Position->SampleTimeOffset == SampleTimeOffset &&
 					Position->OriginTimeOffset == 0.f &&
 					Position->PermutationTimeType == PermutationTimeType &&
 					Position->SampleRole == Role &&
-					Position->OriginRole == Role)
+					Position->OriginRole == Role &&
+					Position->bDefaultWithRootBone)
 				{
 					return Position;
 				}
@@ -53,14 +55,16 @@ void UPoseSearchFeatureChannel_Position::FindOrAddToSchema(UPoseSearchSchema* Sc
 
 bool UPoseSearchFeatureChannel_Position::Finalize(UPoseSearchSchema* Schema)
 {
+	using namespace UE::PoseSearch;
+
 	ChannelDataOffset = Schema->SchemaCardinality;
 	ChannelCardinality = UE::PoseSearch::FFeatureVectorHelper::GetVectorCardinality(ComponentStripping);
 	Schema->SchemaCardinality += ChannelCardinality;
 
-	SchemaBoneIdx = Schema->AddBoneReference(Bone, SampleRole);
-	SchemaOriginBoneIdx = Schema->AddBoneReference(OriginBone, OriginRole);
+	SchemaBoneIdx = Schema->AddBoneReference(Bone, SampleRole, bDefaultWithRootBone);
+	SchemaOriginBoneIdx = Schema->AddBoneReference(OriginBone, OriginRole, bDefaultWithRootBone);
 	
-	return SchemaBoneIdx >= 0 && SchemaOriginBoneIdx >= 0;
+	return SchemaBoneIdx != InvalidSchemaBoneIdx && SchemaOriginBoneIdx != InvalidSchemaBoneIdx;
 }
 
 void UPoseSearchFeatureChannel_Position::AddDependentChannels(UPoseSearchSchema* Schema) const
@@ -71,8 +75,15 @@ void UPoseSearchFeatureChannel_Position::AddDependentChannels(UPoseSearchSchema*
 	{
 		if (SchemaOriginBoneIdx != RootSchemaBoneIdx || PermutationTimeType == EPermutationTimeType::UsePermutationTime)
 		{
-			const EPermutationTimeType DependentChannelsPermutationTimeType = PermutationTimeType == EPermutationTimeType::UsePermutationTime ? EPermutationTimeType::UseSampleToPermutationTime : EPermutationTimeType::UseSampleTime;
-			UPoseSearchFeatureChannel_Position::FindOrAddToSchema(Schema, 0.f, OriginBone.BoneName, OriginRole, DependentChannelsPermutationTimeType);
+			if (bDefaultWithRootBone)
+			{
+				const EPermutationTimeType DependentChannelsPermutationTimeType = PermutationTimeType == EPermutationTimeType::UsePermutationTime ? EPermutationTimeType::UseSampleToPermutationTime : EPermutationTimeType::UseSampleTime;
+				UPoseSearchFeatureChannel_Position::FindOrAddToSchema(Schema, 0.f, OriginBone.BoneName, OriginRole, DependentChannelsPermutationTimeType);
+			}
+			else
+			{
+				// @todo: add bInjectAdditionalDebugChannels support for !bDefaultWithRootBone UPoseSearchFeatureChannel_Position
+			}
 		}
 	}
 }
@@ -211,8 +222,17 @@ void UPoseSearchFeatureChannel_Position::DebugDraw(const UE::PoseSearch::FDebugD
 
 		if (MaxPositionDistanceSquared > 0.f)
 		{
-			static int32 Segments = 8;
-			DrawParams.DrawSphere(BonePos, FMath::Sqrt(MaxPositionDistanceSquared), Segments, Color);
+			static constexpr int32 Segments = 32;
+			const float Radius = FMath::Sqrt(MaxPositionDistanceSquared);
+			if (ComponentStripping == EComponentStrippingVector::StripZ)
+			{
+				const FMatrix CircleTransform(FVector::ZAxisVector, FVector::XAxisVector, FVector::YAxisVector, BonePos);
+				DrawParams.DrawCircle(CircleTransform, Radius, Segments, Color);
+			}
+			else
+			{
+				DrawParams.DrawSphere(BonePos, Radius, Segments, Color);
+			}
 		}
 
 		const bool bDrawOrigin = !DeltaPos.IsNearlyZero() && (SchemaOriginBoneIdx != RootSchemaBoneIdx || !FMath::IsNearlyZero(OriginTimeOffset) ||
@@ -279,7 +299,11 @@ UE::PoseSearch::TLabelBuilder& UPoseSearchFeatureChannel_Position::GetLabel(UE::
 		LabelBuilder.Append(TEXT("_"));
 		LabelBuilder.Append(Schema->GetBoneReferences(SampleRole)[SchemaBoneIdx].BoneName.ToString());
 	}
-
+	else if (SchemaBoneIdx == TrajectorySchemaBoneIdx)
+	{
+		LabelBuilder.Append(TEXT("_Trj"));
+	}
+	
 	if (SampleRole != DefaultRole)
 	{
 		LabelBuilder.Append(TEXT("["));
@@ -292,16 +316,11 @@ UE::PoseSearch::TLabelBuilder& UPoseSearchFeatureChannel_Position::GetLabel(UE::
 		LabelBuilder.Append(TEXT("_"));
 		LabelBuilder.Append(Schema->GetBoneReferences(OriginRole)[SchemaOriginBoneIdx].BoneName.ToString());
 	}
-
-	if (PermutationTimeType == EPermutationTimeType::UsePermutationTime)
-	{
-		LabelBuilder.Append(TEXT("_PT"));
-	}
-	else if (PermutationTimeType == EPermutationTimeType::UseSampleToPermutationTime)
+	else if (SchemaOriginBoneIdx == TrajectorySchemaBoneIdx)
 	{
-		LabelBuilder.Append(TEXT("_SPT"));
+		LabelBuilder.Append(TEXT("_Trj"));
 	}
-
+	
 	if (OriginRole != DefaultRole)
 	{
 		LabelBuilder.Append(TEXT("["));
@@ -309,6 +328,15 @@ UE::PoseSearch::TLabelBuilder& UPoseSearchFeatureChannel_Position::GetLabel(UE::
 		LabelBuilder.Append(TEXT("]"));
 	}
 
+	if (PermutationTimeType == EPermutationTimeType::UsePermutationTime)
+	{
+		LabelBuilder.Append(TEXT("_PT"));
+	}
+	else if (PermutationTimeType == EPermutationTimeType::UseSampleToPermutationTime)
+	{
+		LabelBuilder.Append(TEXT("_SPT"));
+	}
+	
 	AppendLabelSeparator(LabelBuilder, LabelFormat, true);
 
 	LabelBuilder.Appendf(TEXT("%.2f"), SampleTimeOffset);
diff --git a/Engine/Plugins/Animation/PoseSearch/Source/Runtime/Private/PoseSearchFeatureChannel_Position.h b/Engine/Plugins/Animation/PoseSearch/Source/Runtime/Private/PoseSearchFeatureChannel_Position.h
--- a/Engine/Plugins/Animation/PoseSearch/Source/Runtime/Private/PoseSearchFeatureChannel_Position.h
+++ b/Engine/Plugins/Animation/PoseSearch/Source/Runtime/Private/PoseSearchFeatureChannel_Position.h
@@ -12,6 +12,12 @@ class POSESEARCH_API UPoseSearchFeatureChannel_Position : public UPoseSearchFeat
 	GENERATED_BODY()
 
 public:
+	// Experimental, this feature might be removed without warning, not for production use
+	// if bDefaultWithRootBone is true Bone.BoneName and OriginBone.BoneName get initialized to the associated
+	// skeleton root bone if equals to NAME_None, otherwise transforms from the trajectory will be used
+	UPROPERTY(EditAnywhere, Category = "Experimental")
+	bool bDefaultWithRootBone = true;
+
 	UPROPERTY(EditAnywhere, Category = "Settings")
 	FBoneReference Bone;
 
@@ -71,7 +77,7 @@ public:
 	// During pose selection if the squared distance between query versus candidate poses for this 
 	// position channel is greater than MaxPositionDistanceSquared the candidate will be discarded.
 	// the filtering will be enabled only for MaxPositionDistanceSquared > 0
-	UPROPERTY(EditAnywhere, Category = Experimental)
+	UPROPERTY(EditAnywhere, Category = "Experimental|Filter")
 	float MaxPositionDistanceSquared = 0.f;
 
 #if WITH_EDITORONLY_DATA
diff --git a/Engine/Plugins/Animation/PoseSearch/Source/Runtime/Private/PoseSearchFeatureChannel_Velocity.cpp b/Engine/Plugins/Animation/PoseSearch/Source/Runtime/Private/PoseSearchFeatureChannel_Velocity.cpp
--- a/Engine/Plugins/Animation/PoseSearch/Source/Runtime/Private/PoseSearchFeatureChannel_Velocity.cpp
+++ b/Engine/Plugins/Animation/PoseSearch/Source/Runtime/Private/PoseSearchFeatureChannel_Velocity.cpp
@@ -21,14 +21,16 @@ UPoseSearchFeatureChannel_Velocity::UPoseSearchFeatureChannel_Velocity()
 
 bool UPoseSearchFeatureChannel_Velocity::Finalize(UPoseSearchSchema* Schema)
 {
+	using namespace UE::PoseSearch;
+
 	ChannelDataOffset = Schema->SchemaCardinality;
 	ChannelCardinality = UE::PoseSearch::FFeatureVectorHelper::GetVectorCardinality(ComponentStripping);
 	Schema->SchemaCardinality += ChannelCardinality;
 
-	SchemaBoneIdx = Schema->AddBoneReference(Bone, SampleRole);
-	SchemaOriginBoneIdx = Schema->AddBoneReference(OriginBone, OriginRole);
+	SchemaBoneIdx = Schema->AddBoneReference(Bone, SampleRole, bDefaultWithRootBone);
+	SchemaOriginBoneIdx = Schema->AddBoneReference(OriginBone, OriginRole, bDefaultWithRootBone);
 
-	return SchemaBoneIdx >= 0 && SchemaOriginBoneIdx >= 0;
+	return SchemaBoneIdx != InvalidSchemaBoneIdx && SchemaOriginBoneIdx != InvalidSchemaBoneIdx;
 }
 
 void UPoseSearchFeatureChannel_Velocity::AddDependentChannels(UPoseSearchSchema* Schema) const
@@ -229,16 +231,11 @@ UE::PoseSearch::TLabelBuilder& UPoseSearchFeatureChannel_Velocity::GetLabel(UE::
 		LabelBuilder.Append(TEXT("_"));
 		LabelBuilder.Append(Schema->GetBoneReferences(SampleRole)[SchemaBoneIdx].BoneName.ToString());
 	}
-
-	if (PermutationTimeType == EPermutationTimeType::UsePermutationTime)
-	{
-		LabelBuilder.Append(TEXT("_PT"));
-	}
-	else if (PermutationTimeType == EPermutationTimeType::UseSampleToPermutationTime)
+	else if (SchemaBoneIdx == TrajectorySchemaBoneIdx)
 	{
-		LabelBuilder.Append(TEXT("_SPT"));
+		LabelBuilder.Append(TEXT("_Trj"));
 	}
-
+	
 	if (SampleRole != DefaultRole)
 	{
 		LabelBuilder.Append(TEXT("["));
@@ -251,7 +248,11 @@ UE::PoseSearch::TLabelBuilder& UPoseSearchFeatureChannel_Velocity::GetLabel(UE::
 		LabelBuilder.Append(TEXT("_"));
 		LabelBuilder.Append(Schema->GetBoneReferences(OriginRole)[SchemaOriginBoneIdx].BoneName.ToString());
 	}
-
+	else if (SchemaOriginBoneIdx == TrajectorySchemaBoneIdx)
+	{
+		LabelBuilder.Append(TEXT("_Trj"));
+	}
+	
 	if (OriginRole != DefaultRole)
 	{
 		LabelBuilder.Append(TEXT("["));
@@ -259,6 +260,15 @@ UE::PoseSearch::TLabelBuilder& UPoseSearchFeatureChannel_Velocity::GetLabel(UE::
 		LabelBuilder.Append(TEXT("]"));
 	}
 
+	if (PermutationTimeType == EPermutationTimeType::UsePermutationTime)
+	{
+		LabelBuilder.Append(TEXT("_PT"));
+	}
+	else if (PermutationTimeType == EPermutationTimeType::UseSampleToPermutationTime)
+	{
+		LabelBuilder.Append(TEXT("_SPT"));
+	}
+
 	AppendLabelSeparator(LabelBuilder, LabelFormat, true);
 
 	LabelBuilder.Appendf(TEXT("%.2f"), SampleTimeOffset);
diff --git a/Engine/Plugins/Animation/PoseSearch/Source/Runtime/Private/PoseSearchFeatureChannel_Velocity.h b/Engine/Plugins/Animation/PoseSearch/Source/Runtime/Private/PoseSearchFeatureChannel_Velocity.h
--- a/Engine/Plugins/Animation/PoseSearch/Source/Runtime/Private/PoseSearchFeatureChannel_Velocity.h
+++ b/Engine/Plugins/Animation/PoseSearch/Source/Runtime/Private/PoseSearchFeatureChannel_Velocity.h
@@ -12,6 +12,12 @@ class POSESEARCH_API UPoseSearchFeatureChannel_Velocity : public UPoseSearchFeat
 	GENERATED_BODY()
 
 public:
+	// Experimental, this feature might be removed without warning, not for production use
+	// if bDefaultWithRootBone is true Bone.BoneName and OriginBone.BoneName get initialized to the associated
+	// skeleton root bone if equals to NAME_None, otherwise transforms from the trajectory will be used
+	UPROPERTY(EditAnywhere, Category = "Experimental")
+	bool bDefaultWithRootBone = true;
+
 	UPROPERTY(EditAnywhere, Category = "Settings")
 	FBoneReference Bone;
 
diff --git a/Engine/Plugins/Animation/PoseSearch/Source/Runtime/Private/PoseSearchHistory.cpp b/Engine/Plugins/Animation/PoseSearch/Source/Runtime/Private/PoseSearchHistory.cpp
--- a/Engine/Plugins/Animation/PoseSearch/Source/Runtime/Private/PoseSearchHistory.cpp
+++ b/Engine/Plugins/Animation/PoseSearch/Source/Runtime/Private/PoseSearchHistory.cpp
@@ -162,8 +162,12 @@ static bool LerpEntries(float Time, bool bExtrapolate, const FPoseHistoryEntry&
 static bool LerpEntries(float Time, bool bExtrapolate, const FPoseHistoryEntry& PrevEntry, const FPoseHistoryEntry& NextEntry, const USkeleton* BoneIndexSkeleton, const USkeleton* LastUpdateSkeleton,
 	const FBoneToTransformMap& BoneToTransformMap, FBoneIndexType BoneIndexType, FBoneIndexType ReferenceBoneIndexType, FTransform& OutBoneTransform)
 {
-	bool bSuccess = true;
+	check(BoneIndexType != ReferenceBoneIndexType);
 
+	bool bSuccess = true;
+	const FComponentSpaceTransformIndex BoneTransformIndex = GetRemappedComponentSpaceTransformIndex(BoneIndexSkeleton, LastUpdateSkeleton, BoneToTransformMap, BoneIndexType, bSuccess);
+	check(BoneTransformIndex != ComponentSpaceIndexType);
+	
 	const float Denominator = NextEntry.AccumulatedSeconds - PrevEntry.AccumulatedSeconds;
 	float LerpValue = 0.f;
 	if (!FMath::IsNearlyZero(Denominator))
@@ -172,53 +176,42 @@ static bool LerpEntries(float Time, bool bExtrapolate, const FPoseHistoryEntry&
 		LerpValue = bExtrapolate ? Numerator / Denominator : FMath::Clamp(Numerator / Denominator, 0.f, 1.f);
 	}
 
-	const FComponentSpaceTransformIndex BoneTransformIndex = GetRemappedComponentSpaceTransformIndex(BoneIndexSkeleton, LastUpdateSkeleton, BoneToTransformMap, BoneIndexType, bSuccess);
 	const FComponentSpaceTransformIndex ReferenceBoneTransformIndex = GetRemappedComponentSpaceTransformIndex(BoneIndexSkeleton, LastUpdateSkeleton, BoneToTransformMap, ReferenceBoneIndexType, bSuccess);
-
-	if (BoneTransformIndex != ComponentSpaceIndexType)
+	if (ReferenceBoneTransformIndex == ComponentSpaceIndexType)
 	{
-		if (ReferenceBoneTransformIndex == ComponentSpaceIndexType)
+		if (FMath::IsNearlyZero(LerpValue, ZERO_ANIMWEIGHT_THRESH))
 		{
-			if (FMath::IsNearlyZero(LerpValue, ZERO_ANIMWEIGHT_THRESH))
-			{
-				OutBoneTransform = PrevEntry.GetComponentSpaceTransform(BoneTransformIndex);
-			}
-			else if (FMath::IsNearlyZero(LerpValue - 1.f, ZERO_ANIMWEIGHT_THRESH))
-			{
-				OutBoneTransform = NextEntry.GetComponentSpaceTransform(BoneTransformIndex);
-			}
-			else
-			{
-				OutBoneTransform.Blend(
-					PrevEntry.GetComponentSpaceTransform(BoneTransformIndex),
-					NextEntry.GetComponentSpaceTransform(BoneTransformIndex),
-					LerpValue);
-			}
+			OutBoneTransform = PrevEntry.GetComponentSpaceTransform(BoneTransformIndex);
+		}
+		else if (FMath::IsNearlyZero(LerpValue - 1.f, ZERO_ANIMWEIGHT_THRESH))
+		{
+			OutBoneTransform = NextEntry.GetComponentSpaceTransform(BoneTransformIndex);
 		}
 		else
 		{
-			if (FMath::IsNearlyZero(LerpValue, ZERO_ANIMWEIGHT_THRESH))
-			{
-				OutBoneTransform = PrevEntry.GetComponentSpaceTransform(BoneTransformIndex) * PrevEntry.GetComponentSpaceTransform(ReferenceBoneTransformIndex).Inverse();
-			}
-			else if (FMath::IsNearlyZero(LerpValue - 1.f, ZERO_ANIMWEIGHT_THRESH))
-			{
-				OutBoneTransform = NextEntry.GetComponentSpaceTransform(BoneTransformIndex) * NextEntry.GetComponentSpaceTransform(ReferenceBoneTransformIndex).Inverse();
-			}
-			else
-			{
-				OutBoneTransform.Blend(
-					PrevEntry.GetComponentSpaceTransform(BoneTransformIndex) * PrevEntry.GetComponentSpaceTransform(ReferenceBoneTransformIndex).Inverse(),
-					NextEntry.GetComponentSpaceTransform(BoneTransformIndex) * NextEntry.GetComponentSpaceTransform(ReferenceBoneTransformIndex).Inverse(),
-					LerpValue);
-			}
+			OutBoneTransform.Blend(
+				PrevEntry.GetComponentSpaceTransform(BoneTransformIndex),
+				NextEntry.GetComponentSpaceTransform(BoneTransformIndex),
+				LerpValue);
 		}
 	}
 	else
 	{
-		OutBoneTransform = FTransform::Identity;
-		bSuccess = false;
-		unimplemented();
+		if (FMath::IsNearlyZero(LerpValue, ZERO_ANIMWEIGHT_THRESH))
+		{
+			OutBoneTransform = PrevEntry.GetComponentSpaceTransform(BoneTransformIndex) * PrevEntry.GetComponentSpaceTransform(ReferenceBoneTransformIndex).Inverse();
+		}
+		else if (FMath::IsNearlyZero(LerpValue - 1.f, ZERO_ANIMWEIGHT_THRESH))
+		{
+			OutBoneTransform = NextEntry.GetComponentSpaceTransform(BoneTransformIndex) * NextEntry.GetComponentSpaceTransform(ReferenceBoneTransformIndex).Inverse();
+		}
+		else
+		{
+			OutBoneTransform.Blend(
+				PrevEntry.GetComponentSpaceTransform(BoneTransformIndex) * PrevEntry.GetComponentSpaceTransform(ReferenceBoneTransformIndex).Inverse(),
+				NextEntry.GetComponentSpaceTransform(BoneTransformIndex) * NextEntry.GetComponentSpaceTransform(ReferenceBoneTransformIndex).Inverse(),
+				LerpValue);
+		}
 	}
 
 	return bSuccess;
@@ -476,18 +469,36 @@ void FArchivedPoseHistory::InitFrom(const IPoseHistory* PoseHistory)
 bool FArchivedPoseHistory::GetTransformAtTime(float Time, FTransform& OutBoneTransform, const USkeleton* BoneIndexSkeleton, FBoneIndexType BoneIndexType, FBoneIndexType ReferenceBoneIndexType, bool bExtrapolate) const
 {
 	static_assert(RootBoneIndexType == 0 && ComponentSpaceIndexType == FBoneIndexType(-1) && WorldSpaceIndexType == FBoneIndexType(-2)); // some assumptions
-	check(BoneIndexType != ComponentSpaceIndexType && BoneIndexType != WorldSpaceIndexType);
 	
-	bool bSuccess = false;
-	
-	const bool bApplyComponentToWorld = ReferenceBoneIndexType == WorldSpaceIndexType;
-	FTransform ComponentToWorld = FTransform::Identity;
-	if (bApplyComponentToWorld)
+	if (BoneIndexType == ReferenceBoneIndexType)
+	{
+		OutBoneTransform = FTransform::Identity;
+		return true;
+	}
+
+	if (ReferenceBoneIndexType == WorldSpaceIndexType)
 	{
-		ComponentToWorld = Trajectory.GetSampleAtTime(Time, bExtrapolate).GetTransform();
-		ReferenceBoneIndexType = ComponentSpaceIndexType;
+		if (BoneIndexType == ComponentSpaceIndexType)
+		{
+			OutBoneTransform = Trajectory.GetSampleAtTime(Time, bExtrapolate).GetTransform();
+			return true;
+		}
+
+		// getting BoneIndexType in ComponentSpaceIndexType and then multiplying it with the component to world transform from the trajectory
+		const bool bSuccess = GetTransformAtTime(Time, OutBoneTransform, BoneIndexSkeleton, BoneIndexType, ComponentSpaceIndexType, bExtrapolate);
+		OutBoneTransform *= Trajectory.GetSampleAtTime(Time, bExtrapolate).GetTransform();
+		return bSuccess;
 	}
 
+	if (BoneIndexType == WorldSpaceIndexType || BoneIndexType == ComponentSpaceIndexType)
+	{
+		const bool bSuccess = GetTransformAtTime(Time, OutBoneTransform, BoneIndexSkeleton, ReferenceBoneIndexType, BoneIndexType, bExtrapolate);
+		OutBoneTransform = OutBoneTransform.Inverse();
+		return bSuccess;
+	}
+
+	check(BoneIndexType != ComponentSpaceIndexType && BoneIndexType != WorldSpaceIndexType && ReferenceBoneIndexType != WorldSpaceIndexType);
+
 	const int32 NumEntries = Entries.Num();
 	if (NumEntries > 0)
 	{
@@ -504,18 +515,11 @@ bool FArchivedPoseHistory::GetTransformAtTime(float Time, FTransform& OutBoneTra
 		const FPoseHistoryEntry& PrevEntry = Entries[PrevIdx];
 		const FPoseHistoryEntry& NextEntry = Entries[NextIdx];
 
-		bSuccess = LerpEntries(Time, bExtrapolate, PrevEntry, NextEntry, BoneIndexSkeleton, nullptr, BoneToTransformMap, BoneIndexType, ReferenceBoneIndexType, OutBoneTransform);
-		if (bApplyComponentToWorld)
-		{
-			OutBoneTransform *= ComponentToWorld;
-		}
+		return LerpEntries(Time, bExtrapolate, PrevEntry, NextEntry, BoneIndexSkeleton, nullptr, BoneToTransformMap, BoneIndexType, ReferenceBoneIndexType, OutBoneTransform);
 	}
-	else
-	{
-		OutBoneTransform = ComponentToWorld;
-	}
-	
-	return bSuccess;
+
+	OutBoneTransform = FTransform::Identity;
+	return false;
 }
 
 bool FArchivedPoseHistory::GetCurveValueAtTime(float Time, const FName& CurveName, float& OutCurveValue, bool bExtrapolate /*= false*/) const
@@ -688,18 +692,36 @@ bool FPoseHistory::GetTransformAtTime(float Time, FTransform& OutBoneTransform,
 	CheckThreadSafetyRead(ReadPoseDataThreadSafeCounter);
 
 	static_assert(RootBoneIndexType == 0 && ComponentSpaceIndexType == FBoneIndexType(-1) && WorldSpaceIndexType == FBoneIndexType(-2)); // some assumptions
-	check(BoneIndexType != ComponentSpaceIndexType && BoneIndexType != WorldSpaceIndexType);
-	
-	bool bSuccess = false;
 	
-	const bool bApplyComponentToWorld = ReferenceBoneIndexType == WorldSpaceIndexType;
-	FTransform ComponentToWorld = FTransform::Identity;
-	if (bApplyComponentToWorld)
+	if (BoneIndexType == ReferenceBoneIndexType)
+	{
+		OutBoneTransform = FTransform::Identity;
+		return true;
+	}
+
+	if (ReferenceBoneIndexType == WorldSpaceIndexType)
+	{
+		if (BoneIndexType == ComponentSpaceIndexType)
+		{
+			OutBoneTransform = Trajectory.GetSampleAtTime(Time, bExtrapolate).GetTransform();
+			return true;
+		}
+
+		// getting BoneIndexType in ComponentSpaceIndexType and then multiplying it with the component to world transform from the trajectory
+		const bool bSuccess = GetTransformAtTime(Time, OutBoneTransform, BoneIndexSkeleton, BoneIndexType, ComponentSpaceIndexType, bExtrapolate);
+		OutBoneTransform *= Trajectory.GetSampleAtTime(Time, bExtrapolate).GetTransform();
+		return bSuccess;
+	}
+
+	if (BoneIndexType == WorldSpaceIndexType || BoneIndexType == ComponentSpaceIndexType)
 	{
-		ComponentToWorld = Trajectory.GetSampleAtTime(Time, bExtrapolate).GetTransform();
-		ReferenceBoneIndexType = ComponentSpaceIndexType;
+		const bool bSuccess = GetTransformAtTime(Time, OutBoneTransform, BoneIndexSkeleton, ReferenceBoneIndexType, BoneIndexType, bExtrapolate);
+		OutBoneTransform = OutBoneTransform.Inverse();
+		return bSuccess;
 	}
 
+	check(BoneIndexType != ComponentSpaceIndexType && BoneIndexType != WorldSpaceIndexType && ReferenceBoneIndexType != WorldSpaceIndexType);
+	
 	const int32 NumEntries = PoseData.Entries.Num();
 	if (NumEntries > 0)
 	{
@@ -716,25 +738,17 @@ bool FPoseHistory::GetTransformAtTime(float Time, FTransform& OutBoneTransform,
 		const FPoseHistoryEntry& PrevEntry = PoseData.Entries[PrevIdx];
 		const FPoseHistoryEntry& NextEntry = PoseData.Entries[NextIdx];
 
-		bSuccess = LerpEntries(Time, bExtrapolate, PrevEntry, NextEntry, BoneIndexSkeleton, PoseData.LastUpdateSkeleton.Get(), PoseData.BoneToTransformMap, BoneIndexType, ReferenceBoneIndexType, OutBoneTransform);
-		if (bApplyComponentToWorld)
-		{
-			OutBoneTransform *= ComponentToWorld;
-		}
-	}
-	else
-	{
-		OutBoneTransform = ComponentToWorld;
+		return LerpEntries(Time, bExtrapolate, PrevEntry, NextEntry, BoneIndexSkeleton, PoseData.LastUpdateSkeleton.Get(), PoseData.BoneToTransformMap, BoneIndexType, ReferenceBoneIndexType, OutBoneTransform);
 	}
 	
-	return bSuccess;
+	OutBoneTransform = FTransform::Identity;
+	return false;
 }
 
 bool FPoseHistory::GetCurveValueAtTime(float Time, const FName& CurveName, float& OutCurveValue, bool bExtrapolate /*= false*/) const
 {
 	CheckThreadSafetyRead(ReadPoseDataThreadSafeCounter);
 
-	bool bSuccess = false;
 	const int32 NumEntries = PoseData.Entries.Num();
 	if (NumEntries > 0)
 	{
@@ -751,14 +765,11 @@ bool FPoseHistory::GetCurveValueAtTime(float Time, const FName& CurveName, float
 		const FPoseHistoryEntry& PrevEntry = PoseData.Entries[PrevIdx];
 		const FPoseHistoryEntry& NextEntry = PoseData.Entries[NextIdx];
 
-		bSuccess = LerpEntries(Time, bExtrapolate, PrevEntry, NextEntry, CurveName, PoseData.CollectedCurves, OutCurveValue);
+		return LerpEntries(Time, bExtrapolate, PrevEntry, NextEntry, CurveName, PoseData.CollectedCurves, OutCurveValue);
 	}
-	else
-	{
-		OutCurveValue = 0.0f;
-	}
-
-	return bSuccess;
+	
+	OutCurveValue = 0.0f;
+	return false;
 }
 
 const FPoseSearchQueryTrajectory& FPoseHistory::GetTrajectory() const
@@ -1010,31 +1021,45 @@ void FMemStackPoseHistory::Init(const IPoseHistory* InPoseHistory)
 bool FMemStackPoseHistory::GetTransformAtTime(float Time, FTransform& OutBoneTransform, const USkeleton* BoneIndexSkeleton, FBoneIndexType BoneIndexType, FBoneIndexType ReferenceBoneIndexType, bool bExtrapolate) const
 {
 	check(PoseHistory);
-	if (Time > 0.f)
+
+	const int32 Num = FutureEntries.Num();
+	if (Time > 0.f && Num > 0)
 	{
-		const int32 Num = FutureEntries.Num();
-		if (Num > 0)
+		if (BoneIndexType == ReferenceBoneIndexType)
 		{
-			const bool bApplyComponentToWorld = ReferenceBoneIndexType == WorldSpaceIndexType;
-			FTransform ComponentToWorld = FTransform::Identity;
-			if (bApplyComponentToWorld)
-			{
-				ComponentToWorld = GetTrajectory().GetSampleAtTime(Time, bExtrapolate).GetTransform();
-				ReferenceBoneIndexType = ComponentSpaceIndexType;
-			}
+			OutBoneTransform = FTransform::Identity;
+			return true;
+		}
 
-			const int32 LowerBoundIdx = Algo::LowerBound(FutureEntries, Time, [](const FPoseHistoryEntry& Entry, float Value) { return Value > Entry.AccumulatedSeconds; });
-			const int32 NextIdx = FMath::Min(LowerBoundIdx, Num - 1);
-			const FPoseHistoryEntry& NextEntry = FutureEntries[NextIdx];
-			const FPoseHistoryEntry& PrevEntry = NextIdx > 0 ? FutureEntries[NextIdx - 1] : PoseHistory->GetNumEntries() > 0 ? PoseHistory->GetEntry(PoseHistory->GetNumEntries() - 1) : NextEntry;
-						
-			const bool bSuccess = LerpEntries(Time, bExtrapolate, PrevEntry, NextEntry, BoneIndexSkeleton, nullptr, GetBoneToTransformMap(), BoneIndexType, ReferenceBoneIndexType, OutBoneTransform);
-			if (bApplyComponentToWorld)
+		if (ReferenceBoneIndexType == WorldSpaceIndexType)
+		{
+			if (BoneIndexType == ComponentSpaceIndexType)
 			{
-				OutBoneTransform *= ComponentToWorld;
+				OutBoneTransform = PoseHistory->GetTrajectory().GetSampleAtTime(Time, bExtrapolate).GetTransform();
+				return true;
 			}
+
+			// getting BoneIndexType in ComponentSpaceIndexType and then multiplying it with the component to world transform from the trajectory
+			const bool bSuccess = GetTransformAtTime(Time, OutBoneTransform, BoneIndexSkeleton, BoneIndexType, ComponentSpaceIndexType, bExtrapolate);
+			OutBoneTransform *= PoseHistory->GetTrajectory().GetSampleAtTime(Time, bExtrapolate).GetTransform();
 			return bSuccess;
 		}
+
+		if (BoneIndexType == WorldSpaceIndexType || BoneIndexType == ComponentSpaceIndexType)
+		{
+			const bool bSuccess = GetTransformAtTime(Time, OutBoneTransform, BoneIndexSkeleton, ReferenceBoneIndexType, BoneIndexType, bExtrapolate);
+			OutBoneTransform = OutBoneTransform.Inverse();
+			return bSuccess;
+		}
+
+		check(BoneIndexType != ComponentSpaceIndexType && BoneIndexType != WorldSpaceIndexType && ReferenceBoneIndexType != WorldSpaceIndexType);
+
+		const int32 LowerBoundIdx = Algo::LowerBound(FutureEntries, Time, [](const FPoseHistoryEntry& Entry, float Value) { return Value > Entry.AccumulatedSeconds; });
+		const int32 NextIdx = FMath::Min(LowerBoundIdx, Num - 1);
+		const FPoseHistoryEntry& NextEntry = FutureEntries[NextIdx];
+		const FPoseHistoryEntry& PrevEntry = NextIdx > 0 ? FutureEntries[NextIdx - 1] : PoseHistory->GetNumEntries() > 0 ? PoseHistory->GetEntry(PoseHistory->GetNumEntries() - 1) : NextEntry;
+						
+		return LerpEntries(Time, bExtrapolate, PrevEntry, NextEntry, BoneIndexSkeleton, nullptr, GetBoneToTransformMap(), BoneIndexType, ReferenceBoneIndexType, OutBoneTransform);
 	}
 	
 	return PoseHistory->GetTransformAtTime(Time, OutBoneTransform, BoneIndexSkeleton, BoneIndexType, ReferenceBoneIndexType, bExtrapolate);
@@ -1043,19 +1068,16 @@ bool FMemStackPoseHistory::GetTransformAtTime(float Time, FTransform& OutBoneTra
 bool FMemStackPoseHistory::GetCurveValueAtTime(float Time, const FName& CurveName, float& OutCurveValue, bool bExtrapolate /*= false*/) const
 {
 	check(PoseHistory);
-	if (Time > 0.f)
+	
+	const int32 Num = FutureEntries.Num();
+	if (Time > 0.f && Num > 0)
 	{
-		const int32 Num = FutureEntries.Num();
-		if (Num > 0)
-		{
-			const int32 LowerBoundIdx = Algo::LowerBound(FutureEntries, Time, [](const FPoseHistoryEntry& Entry, float Value) { return Value > Entry.AccumulatedSeconds; });
-			const int32 NextIdx = FMath::Min(LowerBoundIdx, Num - 1);
-			const FPoseHistoryEntry& NextEntry = FutureEntries[NextIdx];
-			const FPoseHistoryEntry& PrevEntry = NextIdx > 0 ? FutureEntries[NextIdx - 1] : PoseHistory->GetNumEntries() > 0 ? PoseHistory->GetEntry(PoseHistory->GetNumEntries() - 1) : NextEntry;
+		const int32 LowerBoundIdx = Algo::LowerBound(FutureEntries, Time, [](const FPoseHistoryEntry& Entry, float Value) { return Value > Entry.AccumulatedSeconds; });
+		const int32 NextIdx = FMath::Min(LowerBoundIdx, Num - 1);
+		const FPoseHistoryEntry& NextEntry = FutureEntries[NextIdx];
+		const FPoseHistoryEntry& PrevEntry = NextIdx > 0 ? FutureEntries[NextIdx - 1] : PoseHistory->GetNumEntries() > 0 ? PoseHistory->GetEntry(PoseHistory->GetNumEntries() - 1) : NextEntry;
 
-			const bool bSuccess = LerpEntries(Time, bExtrapolate, PrevEntry, NextEntry, CurveName, GetCollectedCurves(), OutCurveValue);
-			return bSuccess;
-		}
+		return LerpEntries(Time, bExtrapolate, PrevEntry, NextEntry, CurveName, GetCollectedCurves(), OutCurveValue);
 	}
 
 	return PoseHistory->GetCurveValueAtTime(Time, CurveName, OutCurveValue, bExtrapolate);
@@ -1106,7 +1128,6 @@ const FPoseHistoryEntry& FMemStackPoseHistory::GetEntry(int32 EntryIndex) const
 	return FutureEntries[EntryIndex - PoseHistoryNumEntries];
 }
 
-
 #if ENABLE_DRAW_DEBUG && ENABLE_ANIM_DEBUG
 void FMemStackPoseHistory::DebugDraw(FAnimInstanceProxy& AnimInstanceProxy, FColor Color) const
 {
diff --git a/Engine/Plugins/Animation/PoseSearch/Source/Runtime/Private/PoseSearchSchema.cpp b/Engine/Plugins/Animation/PoseSearch/Source/Runtime/Private/PoseSearchSchema.cpp
--- a/Engine/Plugins/Animation/PoseSearch/Source/Runtime/Private/PoseSearchSchema.cpp
+++ b/Engine/Plugins/Animation/PoseSearch/Source/Runtime/Private/PoseSearchSchema.cpp
@@ -164,6 +164,11 @@ TConstArrayView<FBoneReference> UPoseSearchSchema::GetBoneReferences(const UE::P
 }
 
 int8 UPoseSearchSchema::AddBoneReference(const FBoneReference& BoneReference, const UE::PoseSearch::FRole& Role)
+{
+	return AddBoneReference(BoneReference, Role, true);
+}
+
+int8 UPoseSearchSchema::AddBoneReference(const FBoneReference& BoneReference, const UE::PoseSearch::FRole& Role, bool bDefaultWithRootBone)
 {
 	using namespace UE::PoseSearch;
 
@@ -171,7 +176,7 @@ int8 UPoseSearchSchema::AddBoneReference(const FBoneReference& BoneReference, co
 	if (!RoledSkeleton)
 	{
 		UE_LOG(LogPoseSearch, Error, TEXT("UPoseSearchSchema::AddBoneReference: couldn't find data for the requested Role '%s' in UPoseSearchSchema '%s'"), *Role.ToString(), *GetNameSafe(this));
-		return -1;
+		return InvalidSchemaBoneIdx;
 	}
 
 	int32 SchemaBoneIdx = 0;
@@ -179,7 +184,7 @@ int8 UPoseSearchSchema::AddBoneReference(const FBoneReference& BoneReference, co
 	if (!Skeleton)
 	{
 		UE_LOG(LogPoseSearch, Error, TEXT("UPoseSearchSchema::AddBoneReference: couldn't find Skeleton with Role '%s' in UPoseSearchSchema '%s'"), *Role.ToString(), *GetNameSafe(this));
-		return -1;
+		return InvalidSchemaBoneIdx;
 	}
 
 	bool bDefaultToRootBone = true;
@@ -191,15 +196,19 @@ int8 UPoseSearchSchema::AddBoneReference(const FBoneReference& BoneReference, co
 		{
 			UE_LOG(LogPoseSearch, Error, TEXT("UPoseSearchSchema::AddBoneReference: couldn't initialize FBoneReference '%s' with Skeleton '%s' with Role '%s' in UPoseSearchSchema '%s'"),
 				*TempBoneReference.BoneName.ToString(), *GetNameSafe(Skeleton), *Role.ToString(), *GetNameSafe(this));
-			return -1;
+			return InvalidSchemaBoneIdx;
 		}
 	}
-	else
+	else if (bDefaultWithRootBone)
 	{
 		TempBoneReference.BoneName = Skeleton->GetReferenceSkeleton().GetBoneName(int32(RootBoneIndexType));
 		TempBoneReference.Initialize(Skeleton);
 		check(TempBoneReference.HasValidSetup());
 	}
+	else
+	{
+		return TrajectorySchemaBoneIdx;
+	}
 
 	SchemaBoneIdx = RoledSkeleton->BoneReferences.AddUnique(TempBoneReference);
 	check(SchemaBoneIdx >= 0 && SchemaBoneIdx < 128);
@@ -214,15 +223,14 @@ int8 UPoseSearchSchema::AddCurveReference(const FName& CurveReference, const UE:
 	if (!RoledSkeleton)
 	{
 		UE_LOG(LogPoseSearch, Error, TEXT("UPoseSearchSchema::AddCurveReference: couldn't find data for the requested Role '%s' in UPoseSearchSchema '%s'"), *Role.ToString(), *GetNameSafe(this));
-		return -1;
+		return InvalidSchemaCurveIdx;
 	}
 
-	int32 SchemaBoneIdx = 0;
 	const USkeleton* Skeleton = RoledSkeleton->Skeleton;
 	if (!Skeleton)
 	{
 		UE_LOG(LogPoseSearch, Error, TEXT("UPoseSearchSchema::AddCurveReference: couldn't find Skeleton with Role '%s' in UPoseSearchSchema '%s'"), *Role.ToString(), *GetNameSafe(this));
-		return -1;
+		return InvalidSchemaCurveIdx;
 	}
 
 	// Curves are loosely bound, so there's no guarantee this curve will ever exist in any of the assets indexed by the database.
diff --git a/Engine/Plugins/Animation/PoseSearch/Source/Runtime/Public/PoseSearch/PoseSearchDefines.h b/Engine/Plugins/Animation/PoseSearch/Source/Runtime/Public/PoseSearch/PoseSearchDefines.h
--- a/Engine/Plugins/Animation/PoseSearch/Source/Runtime/Public/PoseSearch/PoseSearchDefines.h
+++ b/Engine/Plugins/Animation/PoseSearch/Source/Runtime/Public/PoseSearch/PoseSearchDefines.h
@@ -24,6 +24,11 @@ POSESEARCH_API DECLARE_LOG_CATEGORY_EXTERN(LogPoseSearch, Log, All);
 namespace UE::PoseSearch
 {
 static constexpr int8 RootSchemaBoneIdx = 0;
+static constexpr int8 TrajectorySchemaBoneIdx = -1;
+static constexpr int8 InvalidSchemaBoneIdx = -2;
+
+static constexpr int8 InvalidSchemaCurveIdx = -1;
+
 static constexpr FBoneIndexType RootBoneIndexType = 0;
 static constexpr FBoneIndexType ComponentSpaceIndexType = FBoneIndexType(-1);
 static constexpr FBoneIndexType WorldSpaceIndexType = FBoneIndexType(-2);
diff --git a/Engine/Plugins/Animation/PoseSearch/Source/Runtime/Public/PoseSearch/PoseSearchSchema.h b/Engine/Plugins/Animation/PoseSearch/Source/Runtime/Public/PoseSearch/PoseSearchSchema.h
--- a/Engine/Plugins/Animation/PoseSearch/Source/Runtime/Public/PoseSearch/PoseSearchSchema.h
+++ b/Engine/Plugins/Animation/PoseSearch/Source/Runtime/Public/PoseSearch/PoseSearchSchema.h
@@ -148,6 +148,8 @@ public:
 	// UObject
 	virtual void PostLoad() override;
 
+	// Experimental, this feature might be removed without warning, not for production use
+	int8 AddBoneReference(const FBoneReference& BoneReference, const UE::PoseSearch::FRole& Role, bool bDefaultWithRootBone);
 	int8 AddBoneReference(const FBoneReference& BoneReference, const UE::PoseSearch::FRole& Role);
 	int8 AddCurveReference(const FName& CurveReference, const UE::PoseSearch::FRole& Role);
 
