PoseSearch - debug lines and UPoseSearchDatabase::CalculateDynamicWeightsSqrt improvements

diff --git a/Engine/Plugins/Animation/PoseSearch/Source/Runtime/Private/PoseSearchContext.cpp b/Engine/Plugins/Animation/PoseSearch/Source/Runtime/Private/PoseSearchContext.cpp
--- a/Engine/Plugins/Animation/PoseSearch/Source/Runtime/Private/PoseSearchContext.cpp
+++ b/Engine/Plugins/Animation/PoseSearch/Source/Runtime/Private/PoseSearchContext.cpp
@@ -373,6 +373,27 @@ void FDebugDrawParams::DrawPoint(const FVector& Position, const FColor& Color, f
 	}
 }
 
+void FDebugDrawParams::DrawCircle(const FVector& Center, const FVector& UpVector, float Radius, int32 Segments, const FColor& Color, float Thickness) const
+{
+	FVector A = UpVector;
+	if (A.Normalize())
+	{
+		FMatrix TransformMatrix;
+		FVector B = A.Cross(FVector::ZAxisVector);
+		if (B.Normalize())
+		{
+			const FVector C = A.Cross(B);
+			TransformMatrix = FMatrix(A, B, C, Center);
+		}
+		else
+		{
+			TransformMatrix = FMatrix(FVector::ZAxisVector, FVector::XAxisVector, FVector::YAxisVector, Center);
+		}
+
+		DrawCircle(TransformMatrix, Radius, Segments, Color, Thickness);
+	}
+}
+
 void FDebugDrawParams::DrawCircle(const FMatrix& TransformMatrix, float Radius, int32 Segments, const FColor& Color, float Thickness) const
 {
 	if (Color.A > 0)
@@ -389,6 +410,33 @@ void FDebugDrawParams::DrawCircle(const FMatrix& TransformMatrix, float Radius,
 	}
 }
 
+void FDebugDrawParams::DrawWedge(const FVector& Origin, const FVector& Direction, float InnerRadius, float OuterRadius, float Width, int32 Segments, const FColor& Color, float Thickness) const
+{
+	FVector NormalizedDirection = Direction;
+	if (Color.A > 0 && Segments > 0 && NormalizedDirection.Normalize())
+	{
+		float AngleDeg = -Width * 0.5f;
+		const float AngleDegIncrement = Width / (Segments - 1);
+
+		FVector PrevDirection = NormalizedDirection.RotateAngleAxis(AngleDeg, FVector::ZAxisVector);
+		DrawLine(PrevDirection * OuterRadius + Origin, PrevDirection * InnerRadius + Origin, Color);
+
+		for (int32 Segment = 1; Segment < Segments; ++Segment)
+		{
+			AngleDeg += AngleDegIncrement;
+
+			const FVector CurrDirection = NormalizedDirection.RotateAngleAxis(AngleDeg, FVector::ZAxisVector);
+			
+			DrawLine(PrevDirection * InnerRadius + Origin, CurrDirection * InnerRadius + Origin, Color);
+			DrawLine(PrevDirection * OuterRadius + Origin, CurrDirection * OuterRadius + Origin, Color);
+
+			PrevDirection = CurrDirection;
+		}
+
+		DrawLine(PrevDirection * OuterRadius + Origin, PrevDirection * InnerRadius + Origin, Color);
+	}
+}
+
 void FDebugDrawParams::DrawSphere(const FVector& Center, float Radius, int32 Segments, const FColor& Color, float Thickness) const
 {
 	if (Color.A > 0)
@@ -930,5 +978,4 @@ bool FSearchContext::CanUseCurrentResult() const
 	return CurrentResult.IsValid() && CurrentResult.Database->Schema == CachedQueries.Last().GetSchema();
 }
 
-
 } // namespace UE::PoseSearch
\ No newline at end of file
diff --git a/Engine/Plugins/Animation/PoseSearch/Source/Runtime/Private/PoseSearchDatabase.cpp b/Engine/Plugins/Animation/PoseSearch/Source/Runtime/Private/PoseSearchDatabase.cpp
--- a/Engine/Plugins/Animation/PoseSearch/Source/Runtime/Private/PoseSearchDatabase.cpp
+++ b/Engine/Plugins/Animation/PoseSearch/Source/Runtime/Private/PoseSearchDatabase.cpp
@@ -655,7 +655,7 @@ TConstArrayView<float> UPoseSearchDatabase::CalculateDynamicWeightsSqrt(TArrayVi
 	check(DynamicWeightsSqrtBuffer.Num() == SearchIndex.WeightsSqrt.Num());
 
 	bool bInitialized = false;
-	Schema->FindChannel([&bInitialized, &SearchIndex, &DynamicWeightsSqrtBuffer](const UPoseSearchFeatureChannel* Channel)
+	Schema->IterateChannels([&bInitialized, &SearchIndex, &DynamicWeightsSqrtBuffer](const UPoseSearchFeatureChannel* Channel)
 		{
 			if (const UPoseSearchFeatureChannel_GroupBase* ChannelGroupBase = Cast<UPoseSearchFeatureChannel_GroupBase>(Channel))
 			{
@@ -672,7 +672,6 @@ TConstArrayView<float> UPoseSearchDatabase::CalculateDynamicWeightsSqrt(TArrayVi
 					FMemory::Memzero(DynamicWeightsSqrtBuffer.GetData() + ChannelGroupBase->GetChannelDataOffset(), ChannelGroupBase->GetChannelCardinality() * sizeof(float));
 				}
 			}
-			return nullptr;
 		});
 
 	if (bInitialized)
diff --git a/Engine/Plugins/Animation/PoseSearch/Source/Runtime/Private/PoseSearchFeatureChannel_FilterCrashingLegs.cpp b/Engine/Plugins/Animation/PoseSearch/Source/Runtime/Private/PoseSearchFeatureChannel_FilterCrashingLegs.cpp
--- a/Engine/Plugins/Animation/PoseSearch/Source/Runtime/Private/PoseSearchFeatureChannel_FilterCrashingLegs.cpp
+++ b/Engine/Plugins/Animation/PoseSearch/Source/Runtime/Private/PoseSearchFeatureChannel_FilterCrashingLegs.cpp
@@ -164,7 +164,7 @@ void UPoseSearchFeatureChannel_FilterCrashingLegs::DebugDraw(const UE::PoseSearc
 // IPoseSearchFilter interface
 bool UPoseSearchFeatureChannel_FilterCrashingLegs::IsFilterActive() const
 {
-	return AllowedTolerance > 0.f;
+	return AllowedTolerance != 0.f;
 }
 
 bool UPoseSearchFeatureChannel_FilterCrashingLegs::IsFilterValid(TConstArrayView<float> PoseValues, TConstArrayView<float> QueryValues, int32 PoseIdx, const UE::PoseSearch::FPoseMetadata& Metadata) const
diff --git a/Engine/Plugins/Animation/PoseSearch/Source/Runtime/Private/PoseSearchFeatureChannel_Phase.cpp b/Engine/Plugins/Animation/PoseSearch/Source/Runtime/Private/PoseSearchFeatureChannel_Phase.cpp
--- a/Engine/Plugins/Animation/PoseSearch/Source/Runtime/Private/PoseSearchFeatureChannel_Phase.cpp
+++ b/Engine/Plugins/Animation/PoseSearch/Source/Runtime/Private/PoseSearchFeatureChannel_Phase.cpp
@@ -383,10 +383,7 @@ void UPoseSearchFeatureChannel_Phase::DebugDraw(const UE::PoseSearch::FDebugDraw
 
 		const FVector PhaseVector = (TransformZAxisVector * Phase.X + TransformYAxisVector * Phase.Y) * ScaleFactor;
 		DrawParams.DrawLine(BonePos, BonePos + PhaseVector, Color);
-
-		FMatrix CircleTransform;
-		CircleTransform.SetAxes(&TransformXAxisVector, &TransformYAxisVector, &TransformZAxisVector, &BonePos);
-		DrawParams.DrawCircle(CircleTransform, PhaseVector.Length(), Segments, Color);
+		DrawParams.DrawCircle(BonePos, TransformXAxisVector, PhaseVector.Length(), Segments, Color);
 	}
 }
 #endif // ENABLE_DRAW_DEBUG
diff --git a/Engine/Plugins/Animation/PoseSearch/Source/Runtime/Public/PoseSearch/PoseSearchContext.h b/Engine/Plugins/Animation/PoseSearch/Source/Runtime/Public/PoseSearch/PoseSearchContext.h
--- a/Engine/Plugins/Animation/PoseSearch/Source/Runtime/Public/PoseSearch/PoseSearchContext.h
+++ b/Engine/Plugins/Animation/PoseSearch/Source/Runtime/Public/PoseSearch/PoseSearchContext.h
@@ -73,7 +73,9 @@ struct POSESEARCH_API FDebugDrawParams
 
 	void DrawLine(const FVector& LineStart, const FVector& LineEnd, const FColor& Color, float Thickness = 0.f) const;
 	void DrawPoint(const FVector& Position, const FColor& Color, float Thickness = 6.f) const;
+	void DrawCircle(const FVector& Center, const FVector& UpVector, float Radius, int32 Segments, const FColor& Color, float Thickness = 0.f) const;
 	void DrawCircle(const FMatrix& TransformMatrix, float Radius, int32 Segments, const FColor& Color, float Thickness = 0.f) const;
+	void DrawWedge(const FVector& Origin, const FVector& Direction, float InnerRadius, float OuterRadius, float Width, int32 Segments, const FColor& Color, float Thickness = 0.f) const;
 	void DrawSphere(const FVector& Center, float Radius, int32 Segments, const FColor& Color, float Thickness = 0.f) const;
 	void DrawCentripetalCatmullRomSpline(TConstArrayView<FVector> Points, TConstArrayView<FColor> Colors, float Alpha, int32 NumSamplesPerSegment, float Thickness = 1.f) const;
 	
diff --git a/Engine/Plugins/Animation/PoseSearch/Source/Runtime/Public/PoseSearch/PoseSearchSchema.h b/Engine/Plugins/Animation/PoseSearch/Source/Runtime/Public/PoseSearch/PoseSearchSchema.h
--- a/Engine/Plugins/Animation/PoseSearch/Source/Runtime/Public/PoseSearch/PoseSearchSchema.h
+++ b/Engine/Plugins/Animation/PoseSearch/Source/Runtime/Public/PoseSearch/PoseSearchSchema.h
@@ -145,6 +145,12 @@ public:
 		return static_cast<const ChannelType*>(FindChannel([this](const UPoseSearchFeatureChannel* Channel) -> const UPoseSearchFeatureChannel* { return Cast<ChannelType>(Channel); }));
 	}
 
+	template <typename IteratePredicateType>
+	void IterateChannels(IteratePredicateType IteratePredicate) const
+	{
+		IterateChannelsRecursive(GetChannels(), IteratePredicate);
+	}
+
 	// UObject
 	virtual void PostLoad() override;
 
@@ -199,6 +205,19 @@ private:
 		return nullptr;
 	}
 
+	template <typename IteratePredicateType>
+	static void IterateChannelsRecursive(TConstArrayView<TObjectPtr<UPoseSearchFeatureChannel>> Channels, IteratePredicateType IteratePredicate)
+	{
+		for (const TObjectPtr<UPoseSearchFeatureChannel>& ChannelPtr : Channels)
+		{
+			if (ChannelPtr)
+			{
+				IteratePredicate(ChannelPtr);
+				IterateChannelsRecursive(ChannelPtr->GetSubChannels(), IteratePredicate);
+			}
+		}
+	}
+
 	void Finalize();
 	void ResetFinalize();
 };
\ No newline at end of file
