PoseSearch - making the plugin easier to extend

diff --git a/Engine/Plugins/Animation/PoseSearch/Source/Runtime/Private/PoseSearchContext.cpp b/Engine/Plugins/Animation/PoseSearch/Source/Runtime/Private/PoseSearchContext.cpp
--- a/Engine/Plugins/Animation/PoseSearch/Source/Runtime/Private/PoseSearchContext.cpp
+++ b/Engine/Plugins/Animation/PoseSearch/Source/Runtime/Private/PoseSearchContext.cpp
@@ -8,9 +8,9 @@
 #include "PoseSearch/PoseSearchHistory.h"
 #include "PoseSearch/PoseSearchSchema.h"
 #include "PoseSearch/PoseSearchTrajectoryTypes.h"
-#include "PoseSearchFeatureChannel_Heading.h"
-#include "PoseSearchFeatureChannel_PermutationTime.h"
-#include "PoseSearchFeatureChannel_Position.h"
+#include "PoseSearch/PoseSearchFeatureChannel_Heading.h"
+#include "PoseSearch/PoseSearchFeatureChannel_PermutationTime.h"
+#include "PoseSearch/PoseSearchFeatureChannel_Position.h"
 
 namespace UE::PoseSearch
 {
diff --git a/Engine/Plugins/Animation/PoseSearch/Source/Runtime/Private/PoseSearchDatabase.cpp b/Engine/Plugins/Animation/PoseSearch/Source/Runtime/Private/PoseSearchDatabase.cpp
--- a/Engine/Plugins/Animation/PoseSearch/Source/Runtime/Private/PoseSearchDatabase.cpp
+++ b/Engine/Plugins/Animation/PoseSearch/Source/Runtime/Private/PoseSearchDatabase.cpp
@@ -11,7 +11,7 @@
 #include "PoseSearch/PoseSearchContext.h"
 #include "PoseSearch/PoseSearchDefines.h"
 #include "PoseSearch/PoseSearchDerivedData.h"
-#include "PoseSearchFeatureChannel_Group.h"
+#include "PoseSearch/PoseSearchFeatureChannel_Group.h"
 #include "PoseSearch/PoseSearchHistory.h"
 #include "PoseSearch/PoseSearchSchema.h"
 #include "Serialization/ArchiveCountMem.h"
diff --git a/Engine/Plugins/Animation/PoseSearch/Source/Runtime/Private/PoseSearchFeatureChannel_Curve.cpp b/Engine/Plugins/Animation/PoseSearch/Source/Runtime/Private/PoseSearchFeatureChannel_Curve.cpp
--- a/Engine/Plugins/Animation/PoseSearch/Source/Runtime/Private/PoseSearchFeatureChannel_Curve.cpp
+++ b/Engine/Plugins/Animation/PoseSearch/Source/Runtime/Private/PoseSearchFeatureChannel_Curve.cpp
@@ -1,6 +1,6 @@
 // Copyright Epic Games, Inc. All Rights Reserved.
 
-#include "PoseSearchFeatureChannel_Curve.h"
+#include "PoseSearch/PoseSearchFeatureChannel_Curve.h"
 #include "Engine/BlueprintGeneratedClass.h"
 #include "PoseSearch/PoseSearchAssetIndexer.h"
 #include "PoseSearch/PoseSearchAssetSampler.h"
diff --git a/Engine/Plugins/Animation/PoseSearch/Source/Runtime/Private/PoseSearchFeatureChannel_Curve.h b/Engine/Plugins/Animation/PoseSearch/Source/Runtime/Private/PoseSearchFeatureChannel_Curve.h
deleted file mode 100644
diff --git a/Engine/Plugins/Animation/PoseSearch/Source/Runtime/Private/PoseSearchFeatureChannel_FilterCrashingLegs.cpp b/Engine/Plugins/Animation/PoseSearch/Source/Runtime/Private/PoseSearchFeatureChannel_FilterCrashingLegs.cpp
--- a/Engine/Plugins/Animation/PoseSearch/Source/Runtime/Private/PoseSearchFeatureChannel_FilterCrashingLegs.cpp
+++ b/Engine/Plugins/Animation/PoseSearch/Source/Runtime/Private/PoseSearchFeatureChannel_FilterCrashingLegs.cpp
@@ -1,12 +1,12 @@
 // Copyright Epic Games, Inc. All Rights Reserved.
 
-#include "PoseSearchFeatureChannel_FilterCrashingLegs.h"
+#include "PoseSearch/PoseSearchFeatureChannel_FilterCrashingLegs.h"
 #include "PoseSearch/PoseSearchAssetIndexer.h"
 #include "PoseSearch/PoseSearchAssetSampler.h"
 #include "PoseSearch/PoseSearchContext.h"
 #include "PoseSearch/PoseSearchDatabase.h"
 #include "PoseSearch/PoseSearchSchema.h"
-#include "PoseSearchFeatureChannel_Position.h"
+#include "PoseSearch/PoseSearchFeatureChannel_Position.h"
 
 static void ComputeThighsSideAndForward(const FVector& RightThighPos, const FVector& LeftThighPos, FVector& ThighsSide, FVector& ThighsForward)
 {
diff --git a/Engine/Plugins/Animation/PoseSearch/Source/Runtime/Private/PoseSearchFeatureChannel_FilterCrashingLegs.h b/Engine/Plugins/Animation/PoseSearch/Source/Runtime/Private/PoseSearchFeatureChannel_FilterCrashingLegs.h
deleted file mode 100644
diff --git a/Engine/Plugins/Animation/PoseSearch/Source/Runtime/Private/PoseSearchFeatureChannel_Group.cpp b/Engine/Plugins/Animation/PoseSearch/Source/Runtime/Private/PoseSearchFeatureChannel_Group.cpp
--- a/Engine/Plugins/Animation/PoseSearch/Source/Runtime/Private/PoseSearchFeatureChannel_Group.cpp
+++ b/Engine/Plugins/Animation/PoseSearch/Source/Runtime/Private/PoseSearchFeatureChannel_Group.cpp
@@ -1,6 +1,6 @@
 // Copyright Epic Games, Inc. All Rights Reserved.
 
-#include "PoseSearchFeatureChannel_Group.h"
+#include "PoseSearch/PoseSearchFeatureChannel_Group.h"
 #include "PoseSearch/PoseSearchSchema.h"
 
 bool UPoseSearchFeatureChannel_GroupBase::Finalize(UPoseSearchSchema* Schema)
diff --git a/Engine/Plugins/Animation/PoseSearch/Source/Runtime/Private/PoseSearchFeatureChannel_Group.h b/Engine/Plugins/Animation/PoseSearch/Source/Runtime/Private/PoseSearchFeatureChannel_Group.h
deleted file mode 100644
diff --git a/Engine/Plugins/Animation/PoseSearch/Source/Runtime/Private/PoseSearchFeatureChannel_Heading.cpp b/Engine/Plugins/Animation/PoseSearch/Source/Runtime/Private/PoseSearchFeatureChannel_Heading.cpp
--- a/Engine/Plugins/Animation/PoseSearch/Source/Runtime/Private/PoseSearchFeatureChannel_Heading.cpp
+++ b/Engine/Plugins/Animation/PoseSearch/Source/Runtime/Private/PoseSearchFeatureChannel_Heading.cpp
@@ -1,13 +1,13 @@
 // Copyright Epic Games, Inc. All Rights Reserved.
 
-#include "PoseSearchFeatureChannel_Heading.h"
+#include "PoseSearch/PoseSearchFeatureChannel_Heading.h"
 #include "Engine/BlueprintGeneratedClass.h"
 #include "PoseSearch/PoseSearchAssetIndexer.h"
 #include "PoseSearch/PoseSearchAssetSampler.h"
 #include "PoseSearch/PoseSearchContext.h"
 #include "PoseSearch/PoseSearchDatabase.h"
 #include "PoseSearch/PoseSearchSchema.h"
-#include "PoseSearchFeatureChannel_Position.h"
+#include "PoseSearch/PoseSearchFeatureChannel_Position.h"
 
 #if WITH_EDITOR
 #include "PropertyHandle.h"
diff --git a/Engine/Plugins/Animation/PoseSearch/Source/Runtime/Private/PoseSearchFeatureChannel_Heading.h b/Engine/Plugins/Animation/PoseSearch/Source/Runtime/Private/PoseSearchFeatureChannel_Heading.h
deleted file mode 100644
diff --git a/Engine/Plugins/Animation/PoseSearch/Source/Runtime/Private/PoseSearchFeatureChannel_Padding.cpp b/Engine/Plugins/Animation/PoseSearch/Source/Runtime/Private/PoseSearchFeatureChannel_Padding.cpp
--- a/Engine/Plugins/Animation/PoseSearch/Source/Runtime/Private/PoseSearchFeatureChannel_Padding.cpp
+++ b/Engine/Plugins/Animation/PoseSearch/Source/Runtime/Private/PoseSearchFeatureChannel_Padding.cpp
@@ -1,6 +1,6 @@
 // Copyright Epic Games, Inc. All Rights Reserved.
 
-#include "PoseSearchFeatureChannel_Padding.h"
+#include "PoseSearch/PoseSearchFeatureChannel_Padding.h"
 #include "PoseSearch/PoseSearchAssetIndexer.h"
 #include "PoseSearch/PoseSearchContext.h"
 #include "PoseSearch/PoseSearchSchema.h"
diff --git a/Engine/Plugins/Animation/PoseSearch/Source/Runtime/Private/PoseSearchFeatureChannel_Padding.h b/Engine/Plugins/Animation/PoseSearch/Source/Runtime/Private/PoseSearchFeatureChannel_Padding.h
deleted file mode 100644
diff --git a/Engine/Plugins/Animation/PoseSearch/Source/Runtime/Private/PoseSearchFeatureChannel_PermutationTime.cpp b/Engine/Plugins/Animation/PoseSearch/Source/Runtime/Private/PoseSearchFeatureChannel_PermutationTime.cpp
--- a/Engine/Plugins/Animation/PoseSearch/Source/Runtime/Private/PoseSearchFeatureChannel_PermutationTime.cpp
+++ b/Engine/Plugins/Animation/PoseSearch/Source/Runtime/Private/PoseSearchFeatureChannel_PermutationTime.cpp
@@ -1,6 +1,6 @@
 // Copyright Epic Games, Inc. All Rights Reserved.
 
-#include "PoseSearchFeatureChannel_PermutationTime.h"
+#include "PoseSearch/PoseSearchFeatureChannel_PermutationTime.h"
 #include "PoseSearch/PoseSearchAssetIndexer.h"
 #include "PoseSearch/PoseSearchContext.h"
 #include "PoseSearch/PoseSearchSchema.h"
diff --git a/Engine/Plugins/Animation/PoseSearch/Source/Runtime/Private/PoseSearchFeatureChannel_PermutationTime.h b/Engine/Plugins/Animation/PoseSearch/Source/Runtime/Private/PoseSearchFeatureChannel_PermutationTime.h
deleted file mode 100644
diff --git a/Engine/Plugins/Animation/PoseSearch/Source/Runtime/Private/PoseSearchFeatureChannel_Phase.cpp b/Engine/Plugins/Animation/PoseSearch/Source/Runtime/Private/PoseSearchFeatureChannel_Phase.cpp
--- a/Engine/Plugins/Animation/PoseSearch/Source/Runtime/Private/PoseSearchFeatureChannel_Phase.cpp
+++ b/Engine/Plugins/Animation/PoseSearch/Source/Runtime/Private/PoseSearchFeatureChannel_Phase.cpp
@@ -1,12 +1,12 @@
 // Copyright Epic Games, Inc. All Rights Reserved.
 
-#include "PoseSearchFeatureChannel_Phase.h"
+#include "PoseSearch/PoseSearchFeatureChannel_Phase.h"
 #include "PoseSearch/PoseSearchAssetIndexer.h"
 #include "PoseSearch/PoseSearchAssetSampler.h"
 #include "PoseSearch/PoseSearchContext.h"
 #include "PoseSearch/PoseSearchDatabase.h"
 #include "PoseSearch/PoseSearchSchema.h"
-#include "PoseSearchFeatureChannel_Position.h"
+#include "PoseSearch/PoseSearchFeatureChannel_Position.h"
 
 #if WITH_EDITOR
 namespace UE::PoseSearch
diff --git a/Engine/Plugins/Animation/PoseSearch/Source/Runtime/Private/PoseSearchFeatureChannel_Phase.h b/Engine/Plugins/Animation/PoseSearch/Source/Runtime/Private/PoseSearchFeatureChannel_Phase.h
deleted file mode 100644
diff --git a/Engine/Plugins/Animation/PoseSearch/Source/Runtime/Private/PoseSearchFeatureChannel_Pose.cpp b/Engine/Plugins/Animation/PoseSearch/Source/Runtime/Private/PoseSearchFeatureChannel_Pose.cpp
--- a/Engine/Plugins/Animation/PoseSearch/Source/Runtime/Private/PoseSearchFeatureChannel_Pose.cpp
+++ b/Engine/Plugins/Animation/PoseSearch/Source/Runtime/Private/PoseSearchFeatureChannel_Pose.cpp
@@ -1,13 +1,13 @@
 // Copyright Epic Games, Inc. All Rights Reserved.
 
-#include "PoseSearchFeatureChannel_Pose.h"
+#include "PoseSearch/PoseSearchFeatureChannel_Pose.h"
 #include "Animation/Skeleton.h"
 #include "PoseSearch/PoseSearchSchema.h"
 #include "PoseSearch/PoseSearchDerivedDataKey.h"
-#include "PoseSearchFeatureChannel_Heading.h"
-#include "PoseSearchFeatureChannel_Phase.h"
-#include "PoseSearchFeatureChannel_Position.h"
-#include "PoseSearchFeatureChannel_Velocity.h"
+#include "PoseSearch/PoseSearchFeatureChannel_Heading.h"
+#include "PoseSearch/PoseSearchFeatureChannel_Phase.h"
+#include "PoseSearch/PoseSearchFeatureChannel_Position.h"
+#include "PoseSearch/PoseSearchFeatureChannel_Velocity.h"
 
 #define LOCTEXT_NAMESPACE "PoseSearchFeatureChannels"
 
diff --git a/Engine/Plugins/Animation/PoseSearch/Source/Runtime/Private/PoseSearchFeatureChannel_Pose.h b/Engine/Plugins/Animation/PoseSearch/Source/Runtime/Private/PoseSearchFeatureChannel_Pose.h
deleted file mode 100644
diff --git a/Engine/Plugins/Animation/PoseSearch/Source/Runtime/Private/PoseSearchFeatureChannel_Position.cpp b/Engine/Plugins/Animation/PoseSearch/Source/Runtime/Private/PoseSearchFeatureChannel_Position.cpp
--- a/Engine/Plugins/Animation/PoseSearch/Source/Runtime/Private/PoseSearchFeatureChannel_Position.cpp
+++ b/Engine/Plugins/Animation/PoseSearch/Source/Runtime/Private/PoseSearchFeatureChannel_Position.cpp
@@ -1,6 +1,6 @@
 // Copyright Epic Games, Inc. All Rights Reserved.
 
-#include "PoseSearchFeatureChannel_Position.h"
+#include "PoseSearch/PoseSearchFeatureChannel_Position.h"
 #include "Engine/BlueprintGeneratedClass.h"
 #include "PoseSearch/PoseSearchAssetIndexer.h"
 #include "PoseSearch/PoseSearchAssetSampler.h"
diff --git a/Engine/Plugins/Animation/PoseSearch/Source/Runtime/Private/PoseSearchFeatureChannel_Position.h b/Engine/Plugins/Animation/PoseSearch/Source/Runtime/Private/PoseSearchFeatureChannel_Position.h
deleted file mode 100644
diff --git a/Engine/Plugins/Animation/PoseSearch/Source/Runtime/Private/PoseSearchFeatureChannel_SamplingTime.cpp b/Engine/Plugins/Animation/PoseSearch/Source/Runtime/Private/PoseSearchFeatureChannel_SamplingTime.cpp
--- a/Engine/Plugins/Animation/PoseSearch/Source/Runtime/Private/PoseSearchFeatureChannel_SamplingTime.cpp
+++ b/Engine/Plugins/Animation/PoseSearch/Source/Runtime/Private/PoseSearchFeatureChannel_SamplingTime.cpp
@@ -1,6 +1,6 @@
 // Copyright Epic Games, Inc. All Rights Reserved.
 
-#include "PoseSearchFeatureChannel_SamplingTime.h"
+#include "PoseSearch/PoseSearchFeatureChannel_SamplingTime.h"
 #include "PoseSearch/PoseSearchAssetIndexer.h"
 #include "PoseSearch/PoseSearchContext.h"
 #include "PoseSearch/PoseSearchSchema.h"
diff --git a/Engine/Plugins/Animation/PoseSearch/Source/Runtime/Private/PoseSearchFeatureChannel_SamplingTime.h b/Engine/Plugins/Animation/PoseSearch/Source/Runtime/Private/PoseSearchFeatureChannel_SamplingTime.h
deleted file mode 100644
diff --git a/Engine/Plugins/Animation/PoseSearch/Source/Runtime/Private/PoseSearchFeatureChannel_TimeToEvent.cpp b/Engine/Plugins/Animation/PoseSearch/Source/Runtime/Private/PoseSearchFeatureChannel_TimeToEvent.cpp
--- a/Engine/Plugins/Animation/PoseSearch/Source/Runtime/Private/PoseSearchFeatureChannel_TimeToEvent.cpp
+++ b/Engine/Plugins/Animation/PoseSearch/Source/Runtime/Private/PoseSearchFeatureChannel_TimeToEvent.cpp
@@ -1,6 +1,6 @@
 // Copyright Epic Games, Inc. All Rights Reserved.
 
-#include "PoseSearchFeatureChannel_TimeToEvent.h"
+#include "PoseSearch/PoseSearchFeatureChannel_TimeToEvent.h"
 #include "Engine/BlueprintGeneratedClass.h"
 #include "PoseSearch/PoseSearchAnimNotifies.h"
 #include "PoseSearch/PoseSearchAssetIndexer.h"
diff --git a/Engine/Plugins/Animation/PoseSearch/Source/Runtime/Private/PoseSearchFeatureChannel_TimeToEvent.h b/Engine/Plugins/Animation/PoseSearch/Source/Runtime/Private/PoseSearchFeatureChannel_TimeToEvent.h
deleted file mode 100644
diff --git a/Engine/Plugins/Animation/PoseSearch/Source/Runtime/Private/PoseSearchFeatureChannel_Trajectory.cpp b/Engine/Plugins/Animation/PoseSearch/Source/Runtime/Private/PoseSearchFeatureChannel_Trajectory.cpp
--- a/Engine/Plugins/Animation/PoseSearch/Source/Runtime/Private/PoseSearchFeatureChannel_Trajectory.cpp
+++ b/Engine/Plugins/Animation/PoseSearch/Source/Runtime/Private/PoseSearchFeatureChannel_Trajectory.cpp
@@ -1,14 +1,14 @@
 // Copyright Epic Games, Inc. All Rights Reserved.
 
-#include "PoseSearchFeatureChannel_Trajectory.h"
+#include "PoseSearch/PoseSearchFeatureChannel_Trajectory.h"
 #include "Animation/Skeleton.h"
 #include "PoseSearch/PoseSearchContext.h"
 #include "PoseSearch/PoseSearchDerivedDataKey.h"
 #include "PoseSearch/PoseSearchSchema.h"
-#include "PoseSearchFeatureChannel_Heading.h"
-#include "PoseSearchFeatureChannel_Phase.h"
-#include "PoseSearchFeatureChannel_Position.h"
-#include "PoseSearchFeatureChannel_Velocity.h"
+#include "PoseSearch/PoseSearchFeatureChannel_Heading.h"
+#include "PoseSearch/PoseSearchFeatureChannel_Phase.h"
+#include "PoseSearch/PoseSearchFeatureChannel_Position.h"
+#include "PoseSearch/PoseSearchFeatureChannel_Velocity.h"
 
 #define LOCTEXT_NAMESPACE "PoseSearchFeatureChannels"
 
diff --git a/Engine/Plugins/Animation/PoseSearch/Source/Runtime/Private/PoseSearchFeatureChannel_Trajectory.h b/Engine/Plugins/Animation/PoseSearch/Source/Runtime/Private/PoseSearchFeatureChannel_Trajectory.h
deleted file mode 100644
diff --git a/Engine/Plugins/Animation/PoseSearch/Source/Runtime/Private/PoseSearchFeatureChannel_Velocity.cpp b/Engine/Plugins/Animation/PoseSearch/Source/Runtime/Private/PoseSearchFeatureChannel_Velocity.cpp
--- a/Engine/Plugins/Animation/PoseSearch/Source/Runtime/Private/PoseSearchFeatureChannel_Velocity.cpp
+++ b/Engine/Plugins/Animation/PoseSearch/Source/Runtime/Private/PoseSearchFeatureChannel_Velocity.cpp
@@ -1,6 +1,6 @@
 // Copyright Epic Games, Inc. All Rights Reserved.
 
-#include "PoseSearchFeatureChannel_Velocity.h"
+#include "PoseSearch/PoseSearchFeatureChannel_Velocity.h"
 #include "Engine/BlueprintGeneratedClass.h"
 #include "PoseSearch/PoseSearchAssetIndexer.h"
 #include "PoseSearch/PoseSearchAssetSampler.h"
@@ -8,7 +8,7 @@
 #include "PoseSearch/PoseSearchDatabase.h"
 #include "PoseSearch/PoseSearchHistory.h"
 #include "PoseSearch/PoseSearchSchema.h"
-#include "PoseSearchFeatureChannel_Position.h"
+#include "PoseSearch/PoseSearchFeatureChannel_Position.h"
 
 #if WITH_EDITOR
 #include "PropertyHandle.h"
diff --git a/Engine/Plugins/Animation/PoseSearch/Source/Runtime/Private/PoseSearchFeatureChannel_Velocity.h b/Engine/Plugins/Animation/PoseSearch/Source/Runtime/Private/PoseSearchFeatureChannel_Velocity.h
deleted file mode 100644
diff --git a/Engine/Plugins/Animation/PoseSearch/Source/Runtime/Private/PoseSearchLibrary.cpp b/Engine/Plugins/Animation/PoseSearch/Source/Runtime/Private/PoseSearchLibrary.cpp
--- a/Engine/Plugins/Animation/PoseSearch/Source/Runtime/Private/PoseSearchLibrary.cpp
+++ b/Engine/Plugins/Animation/PoseSearch/Source/Runtime/Private/PoseSearchLibrary.cpp
@@ -21,9 +21,9 @@
 #include "PoseSearch/PoseSearchDatabase.h"
 #include "PoseSearch/PoseSearchDerivedData.h"
 #include "PoseSearch/PoseSearchSchema.h"
-#include "PoseSearchFeatureChannel_Trajectory.h"
+#include "PoseSearch/PoseSearchFeatureChannel_Trajectory.h"
 #include "PoseSearch/Trace/PoseSearchTraceLogger.h"
-#include "PoseSearchFeatureChannel_PermutationTime.h"
+#include "PoseSearch/PoseSearchFeatureChannel_PermutationTime.h"
 
 #include UE_INLINE_GENERATED_CPP_BY_NAME(PoseSearchLibrary)
 
diff --git a/Engine/Plugins/Animation/PoseSearch/Source/Runtime/Private/PoseSearchModule.cpp b/Engine/Plugins/Animation/PoseSearch/Source/Runtime/Private/PoseSearchModule.cpp
--- a/Engine/Plugins/Animation/PoseSearch/Source/Runtime/Private/PoseSearchModule.cpp
+++ b/Engine/Plugins/Animation/PoseSearch/Source/Runtime/Private/PoseSearchModule.cpp
@@ -8,7 +8,7 @@
 #include "PoseSearch/PoseSearchDatabase.h"
 #include "PoseSearch/PoseSearchLibrary.h"
 #include "PoseSearch/PoseSearchSchema.h"
-#include "PoseSearchFeatureChannel_PermutationTime.h"
+#include "PoseSearch/PoseSearchFeatureChannel_PermutationTime.h"
 #include "RewindDebuggerRuntimeInterface/IRewindDebuggerRuntimeExtension.h"
 
 class FRewindDebuggerPoseSearchRuntime : public IRewindDebuggerRuntimeExtension 
diff --git a/Engine/Plugins/Animation/PoseSearch/Source/Runtime/Private/PoseSearchSchema.cpp b/Engine/Plugins/Animation/PoseSearch/Source/Runtime/Private/PoseSearchSchema.cpp
--- a/Engine/Plugins/Animation/PoseSearch/Source/Runtime/Private/PoseSearchSchema.cpp
+++ b/Engine/Plugins/Animation/PoseSearch/Source/Runtime/Private/PoseSearchSchema.cpp
@@ -6,10 +6,10 @@
 #include "PoseSearch/PoseSearchContext.h"
 #include "PoseSearch/PoseSearchDefines.h"
 #include "PoseSearch/PoseSearchResult.h"
-#include "PoseSearchFeatureChannel_Padding.h"
-#include "PoseSearchFeatureChannel_PermutationTime.h"
-#include "PoseSearchFeatureChannel_Pose.h"
-#include "PoseSearchFeatureChannel_Trajectory.h"
+#include "PoseSearch/PoseSearchFeatureChannel_Padding.h"
+#include "PoseSearch/PoseSearchFeatureChannel_PermutationTime.h"
+#include "PoseSearch/PoseSearchFeatureChannel_Pose.h"
+#include "PoseSearch/PoseSearchFeatureChannel_Trajectory.h"
 #include "UObject/ObjectSaveContext.h"
 
 void UPoseSearchSchema::AddChannel(UPoseSearchFeatureChannel* Channel)
diff --git a/Engine/Plugins/Animation/PoseSearch/Source/Runtime/Public/PoseSearch/PoseSearchFeatureChannel_Curve.h b/Engine/Plugins/Animation/PoseSearch/Source/Runtime/Public/PoseSearch/PoseSearchFeatureChannel_Curve.h
new file mode 100644
--- /dev/null
+++ b/Engine/Plugins/Animation/PoseSearch/Source/Runtime/Public/PoseSearch/PoseSearchFeatureChannel_Curve.h
@@ -0,0 +1,67 @@
+// Copyright Epic Games, Inc. All Rights Reserved.
+
+#pragma once
+
+#include "PoseSearch/PoseSearchFeatureChannel.h"
+#include "PoseSearchFeatureChannel_Curve.generated.h"
+
+UCLASS(Experimental, BlueprintType, EditInlineNew, Blueprintable, meta = (DisplayName = "Curve Channel"), CollapseCategories)
+class POSESEARCH_API UPoseSearchFeatureChannel_Curve : public UPoseSearchFeatureChannel
+{
+	GENERATED_BODY()
+
+public:
+
+	UPROPERTY(BlueprintReadOnly, EditAnywhere, Category = "Settings")
+	FName CurveName = NAME_None;
+
+	UPROPERTY(EditAnywhere, Category = "Settings")
+	FName SampleRole = UE::PoseSearch::DefaultRole;
+
+	UPROPERTY(Transient)
+	int8 CurveIdx = UE::PoseSearch::InvalidSchemaCurveIdx;
+
+#if WITH_EDITORONLY_DATA
+	UPROPERTY(EditAnywhere, Category = "Settings")
+	float Weight = 1.f;
+#endif // WITH_EDITORONLY_DATA
+	
+	// the data relative to the sampling time associated to this channel will be offset by SampleTimeOffset seconds.
+	// For example, if Curve is DistanceToWall, and SampleTimeOffset is 0.5, this channel will try to match the future DistanceToWall curve of the character 0.5 seconds ahead
+	UPROPERTY(EditAnywhere, Category = "Settings")
+	float SampleTimeOffset = 0.f;
+
+#if WITH_EDITORONLY_DATA
+	UPROPERTY(EditAnywhere, Category = "Settings", meta = (ExcludeFromHash, DisplayPriority = 0))
+	FLinearColor DebugColor = FLinearColor::Green;
+#endif // WITH_EDITORONLY_DATA
+
+	UPROPERTY(EditAnywhere, Category = "Settings")
+	EInputQueryPose InputQueryPose = EInputQueryPose::UseContinuingPose;
+
+#if WITH_EDITORONLY_DATA
+	// if set, all the channels of the same class with the same cardinality, and the same NormalizationGroup, will be normalized together.
+	// for example in a locomotion database of a character holding a weapon, containing non mirrorable animations, you'd still want to normalize together 
+	// left foot and right foot position and Curve
+	UPROPERTY(EditAnywhere, Category = "Settings")
+	FName NormalizationGroup;
+#endif //WITH_EDITORONLY_DATA
+
+	UFUNCTION(BlueprintPure, BlueprintImplementableEvent, meta=(BlueprintThreadSafe, DisplayName = "Get World Curve"), Category = "Settings")
+	float BP_GetCurveValue(const UAnimInstance* AnimInstance) const;
+
+	bool bUseBlueprintQueryOverride = false;
+
+	UPoseSearchFeatureChannel_Curve();
+
+	// UPoseSearchFeatureChannel interface
+	virtual bool Finalize(UPoseSearchSchema* Schema) override;
+	virtual void BuildQuery(UE::PoseSearch::FSearchContext& SearchContext) const override;
+
+#if WITH_EDITOR
+	virtual void FillWeights(TArrayView<float> Weights) const override;
+	virtual bool IndexAsset(UE::PoseSearch::FAssetIndexer& Indexer) const override;
+	virtual UE::PoseSearch::TLabelBuilder& GetLabel(UE::PoseSearch::TLabelBuilder& LabelBuilder, UE::PoseSearch::ELabelFormat LabelFormat = UE::PoseSearch::ELabelFormat::Full_Horizontal) const override;
+	virtual FName GetNormalizationGroup() const override { return NormalizationGroup; }
+#endif
+};
diff --git a/Engine/Plugins/Animation/PoseSearch/Source/Runtime/Public/PoseSearch/PoseSearchFeatureChannel_FilterCrashingLegs.h b/Engine/Plugins/Animation/PoseSearch/Source/Runtime/Public/PoseSearch/PoseSearchFeatureChannel_FilterCrashingLegs.h
new file mode 100644
--- /dev/null
+++ b/Engine/Plugins/Animation/PoseSearch/Source/Runtime/Public/PoseSearch/PoseSearchFeatureChannel_FilterCrashingLegs.h
@@ -0,0 +1,78 @@
+// Copyright Epic Games, Inc. All Rights Reserved.
+
+#pragma once
+
+#include "PoseSearch/PoseSearchFeatureChannel.h"
+#include "BoneContainer.h"
+#include "PoseSearchFeatureChannel_FilterCrashingLegs.generated.h"
+
+// the idea is to calculate the angle between the direction from LeftThigh position to RightThigh position and the direction from LeftFoot position to RightFoot position, and divide it by PI to have values in range [-1,1]
+// the number (called 'CrashingLegsValue' calculated in ComputeCrashingLegsValue) is gonna be
+// 0 if the feet are aligned with the thighs (for example in an idle standing position)
+// 0.5 if the right foot is exactly in front of the left foot (for example when a character is running  following a line)
+// -0.5 if the left foot is exactly in front of the right foot
+// close to 1 or -1 if the feet (and so the legs) are completely crossed
+// at runtime we'll match the CrashingLegsValue and also filter by discarding pose candidates that don't respect the 'AllowedTolerance' between query and database values (happening in IsFilterValid)
+UCLASS(Experimental, BlueprintType, EditInlineNew, meta = (DisplayName = "CrashingLegs Channel"), CollapseCategories)
+class POSESEARCH_API UPoseSearchFeatureChannel_FilterCrashingLegs : public UPoseSearchFeatureChannel
+{
+	GENERATED_BODY()
+
+public:
+	UPROPERTY(EditAnywhere, Category = "Settings")
+	FBoneReference LeftThigh;
+	
+	UPROPERTY(EditAnywhere, Category = "Settings")
+	FBoneReference RightThigh;
+
+	UPROPERTY(EditAnywhere, Category = "Settings")
+	FBoneReference LeftFoot;
+	
+	UPROPERTY(EditAnywhere, Category = "Settings")
+	FBoneReference RightFoot;
+
+	UPROPERTY(EditAnywhere, Category = "Settings")
+	FName SampleRole = UE::PoseSearch::DefaultRole;
+
+	UPROPERTY(EditAnywhere, Category = "Settings")
+	float Weight = 0.2f;
+
+	UPROPERTY()
+	int8 LeftThighIdx;
+
+	UPROPERTY()
+	int8 RightThighIdx;
+
+	UPROPERTY()
+	int8 LeftFootIdx;
+
+	UPROPERTY()
+	int8 RightFootIdx;
+
+	UPROPERTY(EditAnywhere, Category = "Settings")
+	EInputQueryPose InputQueryPose = EInputQueryPose::UseContinuingPose;
+
+	// if AllowedTolerance is zero the filter is disabled
+	UPROPERTY(EditAnywhere, Category = "Settings", meta = (ClampMin = "0", UIMin = "0"))
+	float AllowedTolerance = 0.3f;
+
+	// UPoseSearchFeatureChannel interface
+	virtual bool Finalize(UPoseSearchSchema* Schema) override;
+	virtual void BuildQuery(UE::PoseSearch::FSearchContext& SearchContext) const override;
+
+	virtual void AddDependentChannels(UPoseSearchSchema* Schema) const override;
+
+#if ENABLE_DRAW_DEBUG
+	virtual void DebugDraw(const UE::PoseSearch::FDebugDrawParams& DrawParams, TConstArrayView<float> PoseVector) const override;
+#endif // ENABLE_DRAW_DEBUG
+
+#if WITH_EDITOR
+	virtual void FillWeights(TArrayView<float> Weights) const override;
+	virtual bool IndexAsset(UE::PoseSearch::FAssetIndexer& Indexer) const override;
+	virtual const UE::PoseSearch::FRole GetDefaultRole() const override { return SampleRole; }
+#endif
+
+	// IPoseSearchFilter interface
+	virtual bool IsFilterActive() const override;
+	virtual bool IsFilterValid(TConstArrayView<float> PoseValues, TConstArrayView<float> QueryValues, int32 PoseIdx, const UE::PoseSearch::FPoseMetadata& Metadata) const override;
+};
diff --git a/Engine/Plugins/Animation/PoseSearch/Source/Runtime/Public/PoseSearch/PoseSearchFeatureChannel_Group.h b/Engine/Plugins/Animation/PoseSearch/Source/Runtime/Public/PoseSearch/PoseSearchFeatureChannel_Group.h
new file mode 100644
--- /dev/null
+++ b/Engine/Plugins/Animation/PoseSearch/Source/Runtime/Public/PoseSearch/PoseSearchFeatureChannel_Group.h
@@ -0,0 +1,60 @@
+// Copyright Epic Games, Inc. All Rights Reserved.
+
+#pragma once
+
+#include "PoseSearch/PoseSearchFeatureChannel.h"
+#include "BoneContainer.h"
+#include "PoseSearchFeatureChannel_Group.generated.h"
+
+// Feature channels interface
+UCLASS(Abstract, BlueprintType, EditInlineNew)
+class POSESEARCH_API UPoseSearchFeatureChannel_GroupBase : public UPoseSearchFeatureChannel
+{
+	GENERATED_BODY()
+
+public:
+	UPROPERTY(EditAnywhere, Category = "Settings")
+	FName SampleRole = UE::PoseSearch::DefaultRole;
+
+	// Experimental, this feature might be removed without warning, not for production use
+	// identifier / metadata to categorize this group base channel weights to be able to zero them via 
+	// UPoseSearchDatabase::CalculateDynamicWeightsSqrt
+	UPROPERTY(EditAnywhere, Category = "Settings")
+	int32 DebugWeightGroupID = INDEX_NONE;
+
+	// UPoseSearchFeatureChannel interface
+	virtual bool Finalize(UPoseSearchSchema* Schema) override;
+	virtual void BuildQuery(UE::PoseSearch::FSearchContext& SearchContext) const override;
+	virtual void AddDependentChannels(UPoseSearchSchema* Schema) const override; 
+
+#if ENABLE_DRAW_DEBUG
+	virtual void DebugDraw(const UE::PoseSearch::FDebugDrawParams& DrawParams, TConstArrayView<float> PoseVector) const override;
+#endif // ENABLE_DRAW_DEBUG
+
+#if WITH_EDITOR
+	virtual void FillWeights(TArrayView<float> Weights) const override;
+	virtual bool IndexAsset(UE::PoseSearch::FAssetIndexer& Indexer) const override;
+	virtual const UE::PoseSearch::FRole GetDefaultRole() const override { return SampleRole; }
+#endif //WITH_EDITOR
+
+	// IPoseSearchFilter interface
+	virtual bool IsFilterActive() const override;
+	virtual bool IsFilterValid(TConstArrayView<float> PoseValues, TConstArrayView<float> QueryValues, int32 PoseIdx, const UE::PoseSearch::FPoseMetadata& Metadata) const override;
+};
+
+UCLASS(BlueprintType, EditInlineNew, meta = (DisplayName = "Group Channel"), CollapseCategories)
+class POSESEARCH_API UPoseSearchFeatureChannel_Group : public UPoseSearchFeatureChannel_GroupBase
+{
+	GENERATED_BODY()
+
+	virtual TArrayView<TObjectPtr<UPoseSearchFeatureChannel>> GetSubChannels() override { return SubChannels; }
+	virtual TConstArrayView<TObjectPtr<UPoseSearchFeatureChannel>> GetSubChannels() const override { return SubChannels; }
+
+#if WITH_EDITOR
+	virtual UE::PoseSearch::TLabelBuilder& GetLabel(UE::PoseSearch::TLabelBuilder& LabelBuilder, UE::PoseSearch::ELabelFormat LabelFormat = UE::PoseSearch::ELabelFormat::Full_Horizontal) const override;
+#endif
+
+public:
+	UPROPERTY(EditAnywhere, Instanced, BlueprintReadWrite, Category = "SubChannels")
+	TArray<TObjectPtr<UPoseSearchFeatureChannel>> SubChannels;
+};
\ No newline at end of file
diff --git a/Engine/Plugins/Animation/PoseSearch/Source/Runtime/Public/PoseSearch/PoseSearchFeatureChannel_Heading.h b/Engine/Plugins/Animation/PoseSearch/Source/Runtime/Public/PoseSearch/PoseSearchFeatureChannel_Heading.h
new file mode 100644
--- /dev/null
+++ b/Engine/Plugins/Animation/PoseSearch/Source/Runtime/Public/PoseSearch/PoseSearchFeatureChannel_Heading.h
@@ -0,0 +1,127 @@
+// Copyright Epic Games, Inc. All Rights Reserved.
+
+#pragma once
+
+#include "PoseSearch/PoseSearchFeatureChannel.h"
+#include "BoneContainer.h"
+#include "PoseSearchFeatureChannel_Heading.generated.h"
+
+UENUM()
+enum class EHeadingAxis : uint8
+{
+	X,
+	Y,
+	Z,
+
+	Num UMETA(Hidden)
+};
+
+UCLASS(EditInlineNew, Blueprintable, meta = (DisplayName = "Heading Channel"), CollapseCategories)
+class POSESEARCH_API UPoseSearchFeatureChannel_Heading : public UPoseSearchFeatureChannel
+{
+	GENERATED_BODY()
+
+public:
+	// Experimental, this feature might be removed without warning, not for production use
+	// if bDefaultWithRootBone is true Bone.BoneName and OriginBone.BoneName get initialized to the associated
+	// skeleton root bone if equals to NAME_None, otherwise transforms from the trajectory will be used
+	UPROPERTY(EditAnywhere, Category = "Experimental")
+	bool bDefaultWithRootBone = true;
+
+	UPROPERTY(EditAnywhere, Category = "Settings")
+	FBoneReference Bone;
+
+	UPROPERTY(EditAnywhere, Category = "Settings")
+	FName SampleRole = UE::PoseSearch::DefaultRole;
+
+	UPROPERTY(EditAnywhere, Category = "Settings")
+	FBoneReference OriginBone;
+
+	UPROPERTY(EditAnywhere, Category = "Settings")
+	FName OriginRole = UE::PoseSearch::DefaultRole;
+
+#if WITH_EDITORONLY_DATA
+	UPROPERTY(EditAnywhere, Category = "Settings")
+	float Weight = 1.f;
+#endif // WITH_EDITORONLY_DATA
+
+	// if SamplingAttributeId >= 0, ALL the animations contained in the pose search database referencing the schema containing this channel are expected to have 
+	// UAnimNotifyState_PoseSearchSamplingAttribute notify state with a matching SamplingAttributeId, and the UAnimNotifyState_PoseSearchSamplingAttribute properties
+	// will be used as source of data instead of this channel "Bone". UAnimNotifyState_PoseSearchSamplingAttribute properties will be then converted into OriginBone space
+	UPROPERTY(EditAnywhere, Category = "Settings")
+	int32 SamplingAttributeId = INDEX_NONE;
+	
+	// the data relative to the sampling time associated to this channel will be offsetted by SampleTimeOffset seconds.
+	// For example, if Bone is the head bone, and SampleTimeOffset is 0.5, this channel will try to match the future heading of the character head bone 0.5 seconds ahead
+	UPROPERTY(EditAnywhere, Category = "Settings")
+	float SampleTimeOffset = 0.f;
+
+	// the data relative to the sampling time associated to this channel origin (root / trajectory bone) will be offsetted by OriginTimeOffset seconds.
+	// For example, if Bone is the head bone, SampleTimeOffset is 0.5, and OriginTimeOffset is 0.5, this channel will try to match 
+	// the future heading of the character head bone 0.5 seconds ahead, relative to the future root bone 0.5 seconds ahead
+	UPROPERTY(EditAnywhere, Category = "Settings")
+	float OriginTimeOffset = 0.f;
+
+	UPROPERTY(EditAnywhere, Category = "Settings")
+	EHeadingAxis HeadingAxis = EHeadingAxis::X;	
+
+	// index referencing the associated bone in UPoseSearchSchema::BoneReferences
+	UPROPERTY(Transient)
+	int8 SchemaBoneIdx = 0;
+
+	UPROPERTY(Transient)
+	int8 SchemaOriginBoneIdx = 0;
+
+#if WITH_EDITORONLY_DATA
+	UPROPERTY(EditAnywhere, Category = "Settings", meta = (ExcludeFromHash, DisplayPriority = 0))
+	FLinearColor DebugColor = FLinearColor::White;
+#endif // WITH_EDITORONLY_DATA
+
+	UPROPERTY(EditAnywhere, Category = "Settings")
+	EInputQueryPose InputQueryPose = EInputQueryPose::UseContinuingPose;
+
+	UPROPERTY(EditAnywhere, Category = "Settings")
+	EComponentStrippingVector ComponentStripping = EComponentStrippingVector::None;
+
+	UPROPERTY(EditAnywhere, Category = "Settings")
+	EPermutationTimeType PermutationTimeType = EPermutationTimeType::UseSampleTime;
+
+#if WITH_EDITORONLY_DATA
+	// if set, all the channels of the same class with the same cardinality, and the same NormalizationGroup, will be normalized together.
+	// for example in a locomotion database of a character holding a weapon, containing non mirrorable animations, you'd still want to normalize together 
+	// left foot and right foot position and velocity
+	UPROPERTY(EditAnywhere, Category = "Settings")
+	FName NormalizationGroup;
+#endif //WITH_EDITORONLY_DATA
+
+	UFUNCTION(BlueprintPure, BlueprintImplementableEvent, meta=(BlueprintThreadSafe, DisplayName = "Get World Rotation"), Category = "Settings")
+	FQuat BP_GetWorldRotation(const UAnimInstance* AnimInstance) const;
+
+	bool bUseBlueprintQueryOverride = false;
+
+	UPoseSearchFeatureChannel_Heading();
+
+	// UPoseSearchFeatureChannel interface
+	virtual bool Finalize(UPoseSearchSchema* Schema) override;
+	virtual void BuildQuery(UE::PoseSearch::FSearchContext& SearchContext) const override;
+
+	virtual EPermutationTimeType GetPermutationTimeType() const override { return PermutationTimeType; }
+	virtual void AddDependentChannels(UPoseSearchSchema* Schema) const override;
+
+#if ENABLE_DRAW_DEBUG
+	virtual void DebugDraw(const UE::PoseSearch::FDebugDrawParams& DrawParams, TConstArrayView<float> PoseVector) const override;
+#endif // ENABLE_DRAW_DEBUG
+
+#if WITH_EDITOR
+	virtual void FillWeights(TArrayView<float> Weights) const override;
+	virtual bool IndexAsset(UE::PoseSearch::FAssetIndexer& Indexer) const override;
+	virtual UE::PoseSearch::TLabelBuilder& GetLabel(UE::PoseSearch::TLabelBuilder& LabelBuilder, UE::PoseSearch::ELabelFormat LabelFormat = UE::PoseSearch::ELabelFormat::Full_Horizontal) const override;
+	virtual FName GetNormalizationGroup() const override { return NormalizationGroup; }
+
+	// IBoneReferenceSkeletonProvider interface
+	USkeleton* GetSkeleton(bool& bInvalidSkeletonIsError, const IPropertyHandle* PropertyHandle) override;
+#endif
+
+	FVector GetAxis(const FQuat& Rotation) const;
+	static void FindOrAddToSchema(UPoseSearchSchema* Schema, float SampleTimeOffset, const FName& BoneName, const UE::PoseSearch::FRole& Role, EHeadingAxis HeadingAxis = EHeadingAxis::X, EPermutationTimeType PermutationTimeType = EPermutationTimeType::UseSampleTime);
+};
\ No newline at end of file
diff --git a/Engine/Plugins/Animation/PoseSearch/Source/Runtime/Public/PoseSearch/PoseSearchFeatureChannel_Padding.h b/Engine/Plugins/Animation/PoseSearch/Source/Runtime/Public/PoseSearch/PoseSearchFeatureChannel_Padding.h
new file mode 100644
--- /dev/null
+++ b/Engine/Plugins/Animation/PoseSearch/Source/Runtime/Public/PoseSearch/PoseSearchFeatureChannel_Padding.h
@@ -0,0 +1,28 @@
+// Copyright Epic Games, Inc. All Rights Reserved.
+
+#pragma once
+
+#include "PoseSearch/PoseSearchFeatureChannel.h"
+#include "PoseSearchFeatureChannel_Padding.generated.h"
+
+UCLASS(EditInlineNew, meta = (DisplayName = "Padding Channel"), CollapseCategories)
+class POSESEARCH_API UPoseSearchFeatureChannel_Padding : public UPoseSearchFeatureChannel
+{
+	GENERATED_BODY()
+
+public:
+	UPROPERTY(EditAnywhere, Category = "Settings")
+	int32 PaddingSize = 1;
+
+	// UPoseSearchFeatureChannel interface
+	virtual bool Finalize(UPoseSearchSchema* Schema) override;
+	virtual void BuildQuery(UE::PoseSearch::FSearchContext& SearchContext) const override;
+
+#if WITH_EDITOR
+	virtual void FillWeights(TArrayView<float> Weights) const override;
+	virtual bool IndexAsset(UE::PoseSearch::FAssetIndexer& Indexer) const override;
+	virtual UE::PoseSearch::TLabelBuilder& GetLabel(UE::PoseSearch::TLabelBuilder& LabelBuilder, UE::PoseSearch::ELabelFormat LabelFormat = UE::PoseSearch::ELabelFormat::Full_Horizontal) const override;
+#endif
+
+	static void AddToSchema(UPoseSearchSchema* Schema, int32 PaddingSize);
+};
diff --git a/Engine/Plugins/Animation/PoseSearch/Source/Runtime/Public/PoseSearch/PoseSearchFeatureChannel_PermutationTime.h b/Engine/Plugins/Animation/PoseSearch/Source/Runtime/Public/PoseSearch/PoseSearchFeatureChannel_PermutationTime.h
new file mode 100644
--- /dev/null
+++ b/Engine/Plugins/Animation/PoseSearch/Source/Runtime/Public/PoseSearch/PoseSearchFeatureChannel_PermutationTime.h
@@ -0,0 +1,32 @@
+// Copyright Epic Games, Inc. All Rights Reserved.
+
+#pragma once
+
+#include "PoseSearch/PoseSearchFeatureChannel.h"
+#include "PoseSearchFeatureChannel_PermutationTime.generated.h"
+
+UCLASS(Experimental, EditInlineNew, meta = (DisplayName = "Permutation Time Channel"), CollapseCategories)
+class POSESEARCH_API UPoseSearchFeatureChannel_PermutationTime : public UPoseSearchFeatureChannel
+{
+	GENERATED_BODY()
+
+public:
+#if WITH_EDITORONLY_DATA
+	UPROPERTY(EditAnywhere, Category = "Settings")
+	float Weight = 1.f;
+#endif // WITH_EDITORONLY_DATA
+
+	// UPoseSearchFeatureChannel interface
+	virtual bool Finalize(UPoseSearchSchema* Schema) override;
+	virtual void BuildQuery(UE::PoseSearch::FSearchContext& SearchContext) const override;
+
+#if WITH_EDITOR
+	virtual void FillWeights(TArrayView<float> Weights) const override;
+	virtual bool IndexAsset(UE::PoseSearch::FAssetIndexer& Indexer) const override;
+	virtual UE::PoseSearch::TLabelBuilder& GetLabel(UE::PoseSearch::TLabelBuilder& LabelBuilder, UE::PoseSearch::ELabelFormat LabelFormat = UE::PoseSearch::ELabelFormat::Full_Horizontal) const override;
+#endif
+
+	float GetPermutationTime(TConstArrayView<float> PoseVector) const;
+
+	static void FindOrAddToSchema(UPoseSearchSchema* Schema);
+};
diff --git a/Engine/Plugins/Animation/PoseSearch/Source/Runtime/Public/PoseSearch/PoseSearchFeatureChannel_Phase.h b/Engine/Plugins/Animation/PoseSearch/Source/Runtime/Public/PoseSearch/PoseSearchFeatureChannel_Phase.h
new file mode 100644
--- /dev/null
+++ b/Engine/Plugins/Animation/PoseSearch/Source/Runtime/Public/PoseSearch/PoseSearchFeatureChannel_Phase.h
@@ -0,0 +1,63 @@
+// Copyright Epic Games, Inc. All Rights Reserved.
+
+#pragma once
+
+#include "PoseSearch/PoseSearchFeatureChannel.h"
+#include "BoneContainer.h"
+#include "PoseSearchFeatureChannel_Phase.generated.h"
+
+UCLASS(Experimental, BlueprintType, EditInlineNew, meta = (DisplayName = "Phase Channel"), CollapseCategories)
+class POSESEARCH_API UPoseSearchFeatureChannel_Phase : public UPoseSearchFeatureChannel
+{
+	GENERATED_BODY()
+
+public:
+	UPROPERTY(EditAnywhere, Category = "Settings")
+	FBoneReference Bone;
+
+	UPROPERTY(EditAnywhere, Category = "Settings")
+	FName SampleRole = UE::PoseSearch::DefaultRole;
+
+#if WITH_EDITORONLY_DATA
+	UPROPERTY(EditAnywhere, Category = "Settings")
+	float Weight = 1.f;
+#endif // WITH_EDITORONLY_DATA
+
+	// index referencing the associated bone in UPoseSearchSchema::BoneReferences
+	UPROPERTY(Transient)
+	int8 SchemaBoneIdx = 0;
+
+#if WITH_EDITORONLY_DATA
+	UPROPERTY(EditAnywhere, Category = "Settings", meta = (ExcludeFromHash, DisplayPriority = 0))
+	FLinearColor DebugColor = FLinearColor::Yellow;
+#endif // WITH_EDITORONLY_DATA
+
+	UPROPERTY(EditAnywhere, Category = "Settings")
+	EInputQueryPose InputQueryPose = EInputQueryPose::UseContinuingPose;
+
+#if WITH_EDITORONLY_DATA
+	// if set, all the channels of the same class with the same cardinality, and the same NormalizationGroup, will be normalized together.
+	// for example in a locomotion database of a character holding a weapon, containing non mirrorable animations, you'd still want to normalize together 
+	// left foot and right foot position and velocity
+	UPROPERTY(EditAnywhere, Category = "Settings")
+	FName NormalizationGroup;
+#endif //WITH_EDITORONLY_DATA
+
+	// UPoseSearchFeatureChannel interface
+	virtual bool Finalize(UPoseSearchSchema* Schema) override;
+	virtual void BuildQuery(UE::PoseSearch::FSearchContext& SearchContext) const override;
+
+	virtual void AddDependentChannels(UPoseSearchSchema* Schema) const override;
+
+#if ENABLE_DRAW_DEBUG
+	virtual void DebugDraw(const UE::PoseSearch::FDebugDrawParams& DrawParams, TConstArrayView<float> PoseVector) const override;
+#endif // ENABLE_DRAW_DEBUG
+
+#if WITH_EDITOR
+	virtual void FillWeights(TArrayView<float> Weights) const override;
+	virtual bool IndexAsset(UE::PoseSearch::FAssetIndexer& Indexer) const override;
+	virtual UE::PoseSearch::TLabelBuilder& GetLabel(UE::PoseSearch::TLabelBuilder& LabelBuilder, UE::PoseSearch::ELabelFormat LabelFormat = UE::PoseSearch::ELabelFormat::Full_Horizontal) const override;
+	virtual FName GetNormalizationGroup() const override { return NormalizationGroup; }
+	virtual const UE::PoseSearch::FRole GetDefaultRole() const override { return SampleRole; }
+#endif
+};
\ No newline at end of file
diff --git a/Engine/Plugins/Animation/PoseSearch/Source/Runtime/Public/PoseSearch/PoseSearchFeatureChannel_Pose.h b/Engine/Plugins/Animation/PoseSearch/Source/Runtime/Public/PoseSearch/PoseSearchFeatureChannel_Pose.h
new file mode 100644
--- /dev/null
+++ b/Engine/Plugins/Animation/PoseSearch/Source/Runtime/Public/PoseSearch/PoseSearchFeatureChannel_Pose.h
@@ -0,0 +1,90 @@
+// Copyright Epic Games, Inc. All Rights Reserved.
+
+#pragma once
+
+#include "PoseSearchFeatureChannel_Group.h"
+#include "PoseSearchFeatureChannel_Pose.generated.h"
+
+UENUM(meta = (Bitflags, UseEnumValuesAsMaskValuesInEditor = "true"))
+enum class EPoseSearchBoneFlags : uint32
+{
+	Velocity = 1 << 0,
+	Position = 1 << 1,
+	Rotation = 1 << 2,
+	Phase = 1 << 3,
+};
+ENUM_CLASS_FLAGS(EPoseSearchBoneFlags);
+constexpr bool EnumHasAnyFlags(int32 Flags, EPoseSearchBoneFlags Contains) { return (Flags & int32(Contains)) != 0; }
+inline int32& operator|=(int32& Lhs, EPoseSearchBoneFlags Rhs) { return Lhs |= int32(Rhs); }
+
+USTRUCT()
+struct POSESEARCH_API FPoseSearchBone
+{
+	GENERATED_BODY()
+
+	UPROPERTY(EditAnywhere, Category = Config)
+	FBoneReference Reference;
+
+	// This allows the user to define what information from the channel you want to compare to.
+	UPROPERTY(EditAnywhere, meta = (Bitmask, BitmaskEnum = "/Script/PoseSearch.EPoseSearchBoneFlags"), Category = Config)
+	int32 Flags = int32(EPoseSearchBoneFlags::Position);
+
+#if WITH_EDITORONLY_DATA
+	// if set, all the channels of the same class with the same cardinality, and the same NormalizationGroup, will be normalized together.
+	// for example in a locomotion database of a character holding a weapon, containing non mirrorable animations, you'd still want to normalize together 
+	// left foot and right foot position and velocity
+	UPROPERTY(EditAnywhere, Category = "Settings")
+	FName NormalizationGroup;
+
+	UPROPERTY(EditAnywhere, Category = Config)
+	float Weight = 1.f;
+
+	UPROPERTY(EditAnywhere, Category = Config, meta=(ExcludeFromHash, DisplayPriority = 0))
+	FLinearColor DebugColor = FLinearColor::Green;
+#endif // WITH_EDITORONLY_DATA
+};
+
+// UPoseSearchFeatureChannel_Pose
+UCLASS(BlueprintType, EditInlineNew, meta = (DisplayName = "Pose Channel"), CollapseCategories)
+class POSESEARCH_API UPoseSearchFeatureChannel_Pose : public UPoseSearchFeatureChannel_GroupBase
+{
+	GENERATED_BODY()
+
+public:
+#if WITH_EDITORONLY_DATA
+	UPROPERTY(EditAnywhere, Category = "Settings")
+	float Weight = 1.f;
+#endif // WITH_EDITORONLY_DATA
+
+	// List of skeletal joints and associated Flags (Velocity, Position, etc) to sample.
+	UPROPERTY(EditAnywhere, Category = "Settings")
+	TArray<FPoseSearchBone> SampledBones;
+
+	UPROPERTY()
+	TArray<int8> SchemaBoneIdx;
+
+	UPROPERTY(EditAnywhere, Category = "Settings")
+	EInputQueryPose InputQueryPose = EInputQueryPose::UseContinuingPose;
+
+	// if bUseCharacterSpaceVelocities is true, velocities will be calculated from the positions in character space, otherwise they will be calculated using global space positions
+	UPROPERTY(EditAnywhere, Category = "Settings")
+	bool bUseCharacterSpaceVelocities = true;
+
+	UPROPERTY(Transient)
+	TArray<TObjectPtr<UPoseSearchFeatureChannel>> SubChannels;
+
+	UPoseSearchFeatureChannel_Pose();
+
+	// UPoseSearchFeatureChannel_GroupBase interface
+	virtual TArrayView<TObjectPtr<UPoseSearchFeatureChannel>> GetSubChannels() override { return SubChannels; }
+	virtual TConstArrayView<TObjectPtr<UPoseSearchFeatureChannel>> GetSubChannels() const override { return SubChannels; }
+
+	// UPoseSearchFeatureChannel interface
+	virtual bool Finalize(UPoseSearchSchema* Schema) override;
+
+#if WITH_EDITOR
+	virtual UE::PoseSearch::TLabelBuilder& GetLabel(UE::PoseSearch::TLabelBuilder& LabelBuilder, UE::PoseSearch::ELabelFormat LabelFormat = UE::PoseSearch::ELabelFormat::Full_Horizontal) const override;
+#endif
+};
+
+
diff --git a/Engine/Plugins/Animation/PoseSearch/Source/Runtime/Public/PoseSearch/PoseSearchFeatureChannel_Position.h b/Engine/Plugins/Animation/PoseSearch/Source/Runtime/Public/PoseSearch/PoseSearchFeatureChannel_Position.h
new file mode 100644
--- /dev/null
+++ b/Engine/Plugins/Animation/PoseSearch/Source/Runtime/Public/PoseSearch/PoseSearchFeatureChannel_Position.h
@@ -0,0 +1,126 @@
+// Copyright Epic Games, Inc. All Rights Reserved.
+
+#pragma once
+
+#include "PoseSearch/PoseSearchFeatureChannel.h"
+#include "BoneContainer.h"
+#include "PoseSearchFeatureChannel_Position.generated.h"
+
+UCLASS(EditInlineNew, Blueprintable, meta = (DisplayName = "Position Channel"), CollapseCategories)
+class POSESEARCH_API UPoseSearchFeatureChannel_Position : public UPoseSearchFeatureChannel
+{
+	GENERATED_BODY()
+
+public:
+	// Experimental, this feature might be removed without warning, not for production use
+	// if bDefaultWithRootBone is true Bone.BoneName and OriginBone.BoneName get initialized to the associated
+	// skeleton root bone if equals to NAME_None, otherwise transforms from the trajectory will be used
+	UPROPERTY(EditAnywhere, Category = "Experimental")
+	bool bDefaultWithRootBone = true;
+
+	UPROPERTY(EditAnywhere, Category = "Settings")
+	FBoneReference Bone;
+
+	UPROPERTY(EditAnywhere, Category = "Settings")
+	FName SampleRole = UE::PoseSearch::DefaultRole;
+
+	UPROPERTY(EditAnywhere, Category = "Settings")
+	FBoneReference OriginBone;
+
+	UPROPERTY(EditAnywhere, Category = "Settings")
+	FName OriginRole = UE::PoseSearch::DefaultRole;
+
+#if WITH_EDITORONLY_DATA
+	UPROPERTY(EditAnywhere, Category = "Settings")
+	float Weight = 1.f;
+#endif // WITH_EDITORONLY_DATA
+
+	// if SamplingAttributeId >= 0, ALL the animations contained in the pose search database referencing the schema containing this channel are expected to have 
+	// UAnimNotifyState_PoseSearchSamplingAttribute notify state with a matching SamplingAttributeId, and the UAnimNotifyState_PoseSearchSamplingAttribute properties
+	// will be used as source of data instead of this channel "Bone". UAnimNotifyState_PoseSearchSamplingAttribute properties will be then converted into OriginBone space
+	UPROPERTY(EditAnywhere, Category = "Settings")
+	int32 SamplingAttributeId = INDEX_NONE;
+	
+	// the data relative to the sampling time associated to this channel will be offsetted by SampleTimeOffset seconds.
+	// For example, if Bone is the head bone, and SampleTimeOffset is 0.5, this channel will try to match the future position of the character head bone 0.5 seconds ahead
+	UPROPERTY(EditAnywhere, Category = "Settings")
+	float SampleTimeOffset = 0.f;
+
+	// the data relative to the sampling time associated to this channel origin (root / trajectory bone) will be offsetted by OriginTimeOffset seconds.
+	// For example, if Bone is the head bone, SampleTimeOffset is 0.5, and OriginTimeOffset is 0.5, this channel will try to match 
+	// the future position of the character head bone 0.5 seconds ahead, relative to the future root bone 0.5 seconds ahead
+	UPROPERTY(EditAnywhere, Category = "Settings")
+	float OriginTimeOffset = 0.f;
+
+	// index referencing the associated bone in UPoseSearchSchema::BoneReferences
+	UPROPERTY(Transient)
+	int8 SchemaBoneIdx = 0;
+
+	UPROPERTY(Transient)
+	int8 SchemaOriginBoneIdx = 0;
+
+#if WITH_EDITORONLY_DATA
+	UPROPERTY(EditAnywhere, Category = "Settings", meta = (ExcludeFromHash, DisplayPriority = 0))
+	FLinearColor DebugColor = FLinearColor::Blue;
+#endif // WITH_EDITORONLY_DATA
+
+	UPROPERTY(EditAnywhere, Category = "Settings")
+	EInputQueryPose InputQueryPose = EInputQueryPose::UseContinuingPose;
+
+	UPROPERTY(EditAnywhere, Category = "Settings")
+	EComponentStrippingVector ComponentStripping = EComponentStrippingVector::None;
+
+	UPROPERTY(EditAnywhere, Category = "Settings")
+	EPermutationTimeType PermutationTimeType = EPermutationTimeType::UseSampleTime;
+
+	// Experimental, this feature might be removed without warning, not for production use
+	// During pose selection if the squared distance between query versus candidate poses for this 
+	// position channel is greater than MaxPositionDistanceSquared the candidate will be discarded.
+	// the filtering will be enabled only for MaxPositionDistanceSquared > 0
+	UPROPERTY(EditAnywhere, Category = "Experimental|Filter")
+	float MaxPositionDistanceSquared = 0.f;
+
+#if WITH_EDITORONLY_DATA
+	// if set, all the channels of the same class with the same cardinality, and the same NormalizationGroup, will be normalized together.
+	// for example in a locomotion database of a character holding a weapon, containing non mirrorable animations, you'd still want to normalize together 
+	// left foot and right foot position and velocity
+	UPROPERTY(EditAnywhere, Category = "Settings")
+	FName NormalizationGroup;
+#endif //WITH_EDITORONLY_DATA
+
+	UFUNCTION(BlueprintPure, BlueprintImplementableEvent, meta=(BlueprintThreadSafe, DisplayName = "Get World Position"), Category = "Settings")
+	FVector BP_GetWorldPosition(const UAnimInstance* AnimInstance) const;
+
+	bool bUseBlueprintQueryOverride = false;
+
+	UPoseSearchFeatureChannel_Position();
+
+	// UPoseSearchFeatureChannel interface
+	virtual bool Finalize(UPoseSearchSchema* Schema) override;
+	virtual void BuildQuery(UE::PoseSearch::FSearchContext& SearchContext) const override;
+
+	virtual EPermutationTimeType GetPermutationTimeType() const override { return PermutationTimeType; }
+	virtual void AddDependentChannels(UPoseSearchSchema* Schema) const override;
+
+	// IPoseSearchFilter interface
+	// Experimental, this feature might be removed without warning, not for production use
+	virtual bool IsFilterActive() const override { return MaxPositionDistanceSquared > 0.f; }
+	// Experimental, this feature might be removed without warning, not for production use
+	virtual bool IsFilterValid(TConstArrayView<float> PoseValues, TConstArrayView<float> QueryValues, int32 PoseIdx, const UE::PoseSearch::FPoseMetadata& Metadata) const override;
+
+#if ENABLE_DRAW_DEBUG
+	virtual void DebugDraw(const UE::PoseSearch::FDebugDrawParams& DrawParams, TConstArrayView<float> PoseVector) const override;
+#endif // ENABLE_DRAW_DEBUG
+
+#if WITH_EDITOR
+	virtual void FillWeights(TArrayView<float> Weights) const override;
+	virtual bool IndexAsset(UE::PoseSearch::FAssetIndexer& Indexer) const override;
+	virtual UE::PoseSearch::TLabelBuilder& GetLabel(UE::PoseSearch::TLabelBuilder& LabelBuilder, UE::PoseSearch::ELabelFormat LabelFormat = UE::PoseSearch::ELabelFormat::Full_Horizontal) const override;
+	virtual FName GetNormalizationGroup() const override { return NormalizationGroup; }
+
+	// IBoneReferenceSkeletonProvider interface
+	USkeleton* GetSkeleton(bool& bInvalidSkeletonIsError, const IPropertyHandle* PropertyHandle) override;
+#endif
+
+	static void FindOrAddToSchema(UPoseSearchSchema* Schema, float SampleTimeOffset, const FName& BoneName, const UE::PoseSearch::FRole& Role, EPermutationTimeType PermutationTimeType = EPermutationTimeType::UseSampleTime);
+};
diff --git a/Engine/Plugins/Animation/PoseSearch/Source/Runtime/Public/PoseSearch/PoseSearchFeatureChannel_SamplingTime.h b/Engine/Plugins/Animation/PoseSearch/Source/Runtime/Public/PoseSearch/PoseSearchFeatureChannel_SamplingTime.h
new file mode 100644
--- /dev/null
+++ b/Engine/Plugins/Animation/PoseSearch/Source/Runtime/Public/PoseSearch/PoseSearchFeatureChannel_SamplingTime.h
@@ -0,0 +1,30 @@
+// Copyright Epic Games, Inc. All Rights Reserved.
+
+#pragma once
+
+#include "PoseSearch/PoseSearchFeatureChannel.h"
+#include "PoseSearchFeatureChannel_SamplingTime.generated.h"
+
+// this channel is mainly for debug purposes to augment the features data with the sampling time (default weight is set to zero to be irrelevant during searches)
+UCLASS(EditInlineNew, meta = (DisplayName = "Sampling Time Channel"), CollapseCategories)
+class POSESEARCH_API UPoseSearchFeatureChannel_SamplingTime : public UPoseSearchFeatureChannel
+{
+	GENERATED_BODY()
+
+public:
+	UPROPERTY(EditAnywhere, Category = "Settings")
+	float Weight = 0.f;
+
+	UPROPERTY(EditAnywhere, Category = "Settings")
+	float TimeToMatch = 0.f;
+
+	// UPoseSearchFeatureChannel interface
+	virtual bool Finalize(UPoseSearchSchema* Schema) override;
+	virtual void BuildQuery(UE::PoseSearch::FSearchContext& SearchContext) const override;
+
+#if WITH_EDITOR
+	virtual void FillWeights(TArrayView<float> Weights) const override;
+	virtual bool IndexAsset(UE::PoseSearch::FAssetIndexer& Indexer) const override;
+	virtual UE::PoseSearch::TLabelBuilder& GetLabel(UE::PoseSearch::TLabelBuilder& LabelBuilder, UE::PoseSearch::ELabelFormat LabelFormat = UE::PoseSearch::ELabelFormat::Full_Horizontal) const override;
+#endif
+};
diff --git a/Engine/Plugins/Animation/PoseSearch/Source/Runtime/Public/PoseSearch/PoseSearchFeatureChannel_TimeToEvent.h b/Engine/Plugins/Animation/PoseSearch/Source/Runtime/Public/PoseSearch/PoseSearchFeatureChannel_TimeToEvent.h
new file mode 100644
--- /dev/null
+++ b/Engine/Plugins/Animation/PoseSearch/Source/Runtime/Public/PoseSearch/PoseSearchFeatureChannel_TimeToEvent.h
@@ -0,0 +1,36 @@
+// Copyright Epic Games, Inc. All Rights Reserved.
+
+#pragma once
+
+#include "PoseSearch/PoseSearchFeatureChannel.h"
+#include "PoseSearchFeatureChannel_TimeToEvent.generated.h"
+
+UCLASS(EditInlineNew, Blueprintable, meta = (DisplayName = "Time To Event Channel"), CollapseCategories)
+class POSESEARCH_API UPoseSearchFeatureChannel_TimeToEvent : public UPoseSearchFeatureChannel
+{
+	GENERATED_BODY()
+
+public:
+	UPROPERTY(EditAnywhere, Category = "Settings")
+	float Weight = 0.f;
+
+	UPROPERTY(EditAnywhere, Category = "Settings")
+	int32 SamplingAttributeId = 0;
+
+	UFUNCTION(BlueprintPure, BlueprintImplementableEvent, meta=(BlueprintThreadSafe, DisplayName = "Get Time To Event"), Category = "Settings")
+	float BP_GetTimeToEvent(const UAnimInstance* AnimInstance) const;
+
+	bool bUseBlueprintQueryOverride = false;
+
+	UPoseSearchFeatureChannel_TimeToEvent();
+
+	// UPoseSearchFeatureChannel interface
+	virtual bool Finalize(UPoseSearchSchema* Schema) override;
+	virtual void BuildQuery(UE::PoseSearch::FSearchContext& SearchContext) const override;
+
+#if WITH_EDITOR
+	virtual void FillWeights(TArrayView<float> Weights) const override;
+	virtual bool IndexAsset(UE::PoseSearch::FAssetIndexer& Indexer) const override;
+	virtual UE::PoseSearch::TLabelBuilder& GetLabel(UE::PoseSearch::TLabelBuilder& LabelBuilder, UE::PoseSearch::ELabelFormat LabelFormat = UE::PoseSearch::ELabelFormat::Full_Horizontal) const override;
+#endif
+};
diff --git a/Engine/Plugins/Animation/PoseSearch/Source/Runtime/Public/PoseSearch/PoseSearchFeatureChannel_Trajectory.h b/Engine/Plugins/Animation/PoseSearch/Source/Runtime/Public/PoseSearch/PoseSearchFeatureChannel_Trajectory.h
new file mode 100644
--- /dev/null
+++ b/Engine/Plugins/Animation/PoseSearch/Source/Runtime/Public/PoseSearch/PoseSearchFeatureChannel_Trajectory.h
@@ -0,0 +1,91 @@
+// Copyright Epic Games, Inc. All Rights Reserved.
+
+#pragma once
+
+#include "PoseSearchFeatureChannel_Group.h"
+#include "UObject/ObjectSaveContext.h"
+#include "PoseSearchFeatureChannel_Trajectory.generated.h"
+
+UENUM(meta = (Bitflags, UseEnumValuesAsMaskValuesInEditor = "true"))
+enum class EPoseSearchTrajectoryFlags : uint32
+{
+	Velocity = 1 << 0,
+	Position = 1 << 1,
+	VelocityDirection = 1 << 2,
+	FacingDirection = 1 << 3,
+	VelocityXY = 1 << 4,
+	PositionXY = 1 << 5,
+	VelocityDirectionXY = 1 << 6,
+	FacingDirectionXY = 1 << 7,
+};
+ENUM_CLASS_FLAGS(EPoseSearchTrajectoryFlags);
+constexpr bool EnumHasAnyFlags(int32 Flags, EPoseSearchTrajectoryFlags Contains) { return (Flags & int32(Contains)) != 0; }
+inline int32& operator|=(int32& Lhs, EPoseSearchTrajectoryFlags Rhs) { return Lhs |= int32(Rhs); }
+
+USTRUCT()
+struct POSESEARCH_API FPoseSearchTrajectorySample
+{
+	GENERATED_BODY()
+
+	// the data relative to the sampling time associated to this FPoseSearchTrajectorySample will be offsetted by Offset seconds.
+	// For example, Flags is Position, and Offset is 0.5, this channel will try to match the future position of the trajectory 0.5 seconds ahead
+	UPROPERTY(EditAnywhere, Category = Config)
+	float Offset = 0.f;
+
+	// This allows the user to define what information from the channel you want to compare to.
+	UPROPERTY(EditAnywhere, meta = (Bitmask, BitmaskEnum = "/Script/PoseSearch.EPoseSearchTrajectoryFlags"), Category = Config)
+	int32 Flags = int32(EPoseSearchTrajectoryFlags::Position);
+
+#if WITH_EDITORONLY_DATA
+	UPROPERTY(EditAnywhere, Category = Config)
+	float Weight = 1.f;
+
+	// if set, all the channels of the same class with the same cardinality, and the same NormalizationGroup, will be normalized together.
+	// for example in a locomotion database of a character holding a weapon, containing non mirrorable animations, you'd still want to normalize together 
+	// left foot and right foot position and velocity
+	UPROPERTY(EditAnywhere, Category = Config)
+	FName NormalizationGroup;
+
+	UPROPERTY(EditAnywhere, Category = Config, meta = (ExcludeFromHash, DisplayPriority = 0))
+	FLinearColor DebugColor = FLinearColor::Blue;
+#endif // WITH_EDITORONLY_DATA
+};
+
+UCLASS(BlueprintType, EditInlineNew, meta = (DisplayName = "Trajectory Channel"), CollapseCategories)
+class POSESEARCH_API UPoseSearchFeatureChannel_Trajectory : public UPoseSearchFeatureChannel_GroupBase
+{
+	GENERATED_BODY()
+
+public:
+#if WITH_EDITORONLY_DATA
+	UPROPERTY(EditAnywhere, Category = "Settings")
+	float Weight = 1.f;
+#endif // WITH_EDITORONLY_DATA
+
+	UPROPERTY(EditAnywhere, Category = "Settings")
+	TArray<FPoseSearchTrajectorySample> Samples;
+
+	UPROPERTY(Transient)
+	TArray<TObjectPtr<UPoseSearchFeatureChannel>> SubChannels;
+
+	UPoseSearchFeatureChannel_Trajectory();
+
+	// UPoseSearchFeatureChannel_GroupBase interface
+	virtual TArrayView<TObjectPtr<UPoseSearchFeatureChannel>> GetSubChannels() override { return SubChannels; }
+	virtual TConstArrayView<TObjectPtr<UPoseSearchFeatureChannel>> GetSubChannels() const override { return SubChannels; }
+
+	// UPoseSearchFeatureChannel interface
+	virtual bool Finalize(UPoseSearchSchema* Schema) override;
+
+#if ENABLE_DRAW_DEBUG
+	virtual void DebugDraw(const UE::PoseSearch::FDebugDrawParams& DrawParams, TConstArrayView<float> PoseVector) const override;
+#endif // ENABLE_DRAW_DEBUG
+
+#if WITH_EDITOR
+	virtual UE::PoseSearch::TLabelBuilder& GetLabel(UE::PoseSearch::TLabelBuilder& LabelBuilder, UE::PoseSearch::ELabelFormat LabelFormat = UE::PoseSearch::ELabelFormat::Full_Horizontal) const override;
+#endif
+
+	float GetEstimatedSpeedRatio(TConstArrayView<float> QueryVector, TConstArrayView<float> PoseVector) const;
+	FVector GetEstimatedFutureRootMotionVelocity(TConstArrayView<float> PoseVector) const;
+};
+
diff --git a/Engine/Plugins/Animation/PoseSearch/Source/Runtime/Public/PoseSearch/PoseSearchFeatureChannel_Velocity.h b/Engine/Plugins/Animation/PoseSearch/Source/Runtime/Public/PoseSearch/PoseSearchFeatureChannel_Velocity.h
new file mode 100644
--- /dev/null
+++ b/Engine/Plugins/Animation/PoseSearch/Source/Runtime/Public/PoseSearch/PoseSearchFeatureChannel_Velocity.h
@@ -0,0 +1,119 @@
+// Copyright Epic Games, Inc. All Rights Reserved.
+
+#pragma once
+
+#include "PoseSearch/PoseSearchFeatureChannel.h"
+#include "BoneContainer.h"
+#include "PoseSearchFeatureChannel_Velocity.generated.h"
+
+UCLASS(EditInlineNew, Blueprintable, meta = (DisplayName = "Velocity Channel"), CollapseCategories)
+class POSESEARCH_API UPoseSearchFeatureChannel_Velocity : public UPoseSearchFeatureChannel
+{
+	GENERATED_BODY()
+
+public:
+	// Experimental, this feature might be removed without warning, not for production use
+	// if bDefaultWithRootBone is true Bone.BoneName and OriginBone.BoneName get initialized to the associated
+	// skeleton root bone if equals to NAME_None, otherwise transforms from the trajectory will be used
+	UPROPERTY(EditAnywhere, Category = "Experimental")
+	bool bDefaultWithRootBone = true;
+
+	UPROPERTY(EditAnywhere, Category = "Settings")
+	FBoneReference Bone;
+
+	UPROPERTY(EditAnywhere, Category = "Settings")
+	FName SampleRole = UE::PoseSearch::DefaultRole;
+
+	UPROPERTY(EditAnywhere, Category = "Settings")
+	FBoneReference OriginBone;
+
+	UPROPERTY(EditAnywhere, Category = "Settings")
+	FName OriginRole = UE::PoseSearch::DefaultRole;
+
+#if WITH_EDITORONLY_DATA
+	UPROPERTY(EditAnywhere, Category = "Settings")
+	float Weight = 1.f;
+#endif // WITH_EDITORONLY_DATA
+
+	// if SamplingAttributeId >= 0, ALL the animations contained in the pose search database referencing the schema containing this channel are expected to have 
+	// UAnimNotifyState_PoseSearchSamplingAttribute notify state with a matching SamplingAttributeId, and the UAnimNotifyState_PoseSearchSamplingAttribute properties
+	// will be used as source of data instead of this channel "Bone". UAnimNotifyState_PoseSearchSamplingAttribute properties will be then converted into OriginBone space
+	UPROPERTY(EditAnywhere, Category = "Settings")
+	int32 SamplingAttributeId = INDEX_NONE;
+	
+	// the data relative to the sampling time associated to this channel will be offsetted by SampleTimeOffset seconds.
+	// For example, if Bone is the head bone, and SampleTimeOffset is 0.5, this channel will try to match the future velocity of the character head bone 0.5 seconds ahead
+	UPROPERTY(EditAnywhere, Category = "Settings")
+	float SampleTimeOffset = 0.f;
+
+	// the data relative to the sampling time associated to this channel origin (root / trajectory bone) will be offsetted by OriginTimeOffset seconds.
+	// For example, if Bone is the head bone, SampleTimeOffset is 0.5, and OriginTimeOffset is 0.5, this channel will try to match 
+	// the future velocity of the character head bone 0.5 seconds ahead, relative to the future root bone 0.5 seconds ahead
+	UPROPERTY(EditAnywhere, Category = "Settings")
+	float OriginTimeOffset = 0.f;
+
+	// index referencing the associated bone in UPoseSearchSchema::BoneReferences
+	UPROPERTY(Transient)
+	int8 SchemaBoneIdx = 0;
+
+	UPROPERTY(Transient)
+	int8 SchemaOriginBoneIdx = 0;
+
+#if WITH_EDITORONLY_DATA
+	UPROPERTY(EditAnywhere, Category = "Settings", meta = (ExcludeFromHash, DisplayPriority = 0))
+	FLinearColor DebugColor = FLinearColor::Green;
+#endif // WITH_EDITORONLY_DATA
+
+	UPROPERTY(EditAnywhere, Category = "Settings")
+	EInputQueryPose InputQueryPose = EInputQueryPose::UseContinuingPose;
+
+	// if bUseCharacterSpaceVelocities is true, velocities will be calculated from the positions in character space, otherwise they will be calculated using global space positions
+	UPROPERTY(EditAnywhere, Category = "Settings")
+	bool bUseCharacterSpaceVelocities = true;
+
+	// if bNormalize is true, velocity will be clamped to a maximum size of one (UPoseSearchFeatureChannel_Velocity will be matching the "direction" of the velocity)
+	UPROPERTY(EditAnywhere, Category = "Settings")
+	bool bNormalize = false;
+
+	UPROPERTY(EditAnywhere, Category = "Settings")
+	EComponentStrippingVector ComponentStripping = EComponentStrippingVector::None;
+
+	UPROPERTY(EditAnywhere, Category = "Settings")
+	EPermutationTimeType PermutationTimeType = EPermutationTimeType::UseSampleTime;
+
+#if WITH_EDITORONLY_DATA
+	// if set, all the channels of the same class with the same cardinality, and the same NormalizationGroup, will be normalized together.
+	// for example in a locomotion database of a character holding a weapon, containing non mirrorable animations, you'd still want to normalize together 
+	// left foot and right foot position and velocity
+	UPROPERTY(EditAnywhere, Category = "Settings")
+	FName NormalizationGroup;
+#endif //WITH_EDITORONLY_DATA
+
+	UFUNCTION(BlueprintPure, BlueprintImplementableEvent, meta=(BlueprintThreadSafe, DisplayName = "Get World Velocity"), Category = "Settings")
+	FVector BP_GetWorldVelocity(const UAnimInstance* AnimInstance) const;
+
+	bool bUseBlueprintQueryOverride = false;
+
+	UPoseSearchFeatureChannel_Velocity();
+
+	// UPoseSearchFeatureChannel interface
+	virtual bool Finalize(UPoseSearchSchema* Schema) override;
+	virtual void BuildQuery(UE::PoseSearch::FSearchContext& SearchContext) const override;
+
+	virtual EPermutationTimeType GetPermutationTimeType() const override { return PermutationTimeType; }
+	virtual void AddDependentChannels(UPoseSearchSchema* Schema) const override;
+
+#if ENABLE_DRAW_DEBUG
+	virtual void DebugDraw(const UE::PoseSearch::FDebugDrawParams& DrawParams, TConstArrayView<float> PoseVector) const override;
+#endif // ENABLE_DRAW_DEBUG
+
+#if WITH_EDITOR
+	virtual void FillWeights(TArrayView<float> Weights) const override;
+	virtual bool IndexAsset(UE::PoseSearch::FAssetIndexer& Indexer) const override;
+	virtual UE::PoseSearch::TLabelBuilder& GetLabel(UE::PoseSearch::TLabelBuilder& LabelBuilder, UE::PoseSearch::ELabelFormat LabelFormat = UE::PoseSearch::ELabelFormat::Full_Horizontal) const override;
+	virtual FName GetNormalizationGroup() const override { return NormalizationGroup; }
+	
+	// IBoneReferenceSkeletonProvider interface
+	USkeleton* GetSkeleton(bool& bInvalidSkeletonIsError, const IPropertyHandle* PropertyHandle) override;
+#endif
+};
diff --git a/Engine/Plugins/Animation/PoseSearch/Source/Runtime/Public/PoseSearch/PoseSearchLibrary.h b/Engine/Plugins/Animation/PoseSearch/Source/Runtime/Public/PoseSearch/PoseSearchLibrary.h
--- a/Engine/Plugins/Animation/PoseSearch/Source/Runtime/Public/PoseSearch/PoseSearchLibrary.h
+++ b/Engine/Plugins/Animation/PoseSearch/Source/Runtime/Public/PoseSearch/PoseSearchLibrary.h
@@ -58,7 +58,7 @@ PRAGMA_ENABLE_DEPRECATION_WARNINGS
 	// state back to the default state.
 	POSESEARCH_API void AdjustAssetTime(float AssetTime);
 
-	void UpdateWantedPlayRate(const UE::PoseSearch::FSearchContext& SearchContext, const FFloatInterval& PlayRate, float TrajectorySpeedMultiplier);
+	POSESEARCH_API void UpdateWantedPlayRate(const UE::PoseSearch::FSearchContext& SearchContext, const FFloatInterval& PlayRate, float TrajectorySpeedMultiplier);
 
 	FVector GetEstimatedFutureRootMotionVelocity() const;
 
