PoseSearch - fix to avoid loading every asset in FPoseSearchEditorUtils::IsAssetCompatibleWithDatabase when searching for a UMultiAnimAsset

diff --git a/Engine/Plugins/Animation/PoseSearch/Source/Editor/Private/PoseSearchDatabaseEditorUtils.cpp b/Engine/Plugins/Animation/PoseSearch/Source/Editor/Private/PoseSearchDatabaseEditorUtils.cpp
--- a/Engine/Plugins/Animation/PoseSearch/Source/Editor/Private/PoseSearchDatabaseEditorUtils.cpp
+++ b/Engine/Plugins/Animation/PoseSearch/Source/Editor/Private/PoseSearchDatabaseEditorUtils.cpp
@@ -10,53 +10,60 @@ namespace UE::PoseSearch
 {
 	bool FPoseSearchEditorUtils::IsAssetCompatibleWithDatabase(const UPoseSearchDatabase* InDatabase, const FAssetData& InAssetData)
 	{
-		if (InDatabase)
+		if (!InDatabase || !InDatabase->Schema)
 		{
-			if (InDatabase->Schema)
+			return false;
+		}
+			
+		const TArray<FPoseSearchRoledSkeleton> RoledSkeletons = InDatabase->Schema->GetRoledSkeletons();
+		if (RoledSkeletons.IsEmpty())
+		{
+			return false;
+		}
+
+		if (RoledSkeletons.Num() == 1)
+		{
+			if (RoledSkeletons[0].Skeleton && InAssetData.GetClass()->IsChildOf(UAnimationAsset::StaticClass()) &&
+				RoledSkeletons[0].Skeleton->IsCompatibleForEditor(InAssetData))
 			{
-				const TArray<FPoseSearchRoledSkeleton> RoledSkeletons = InDatabase->Schema->GetRoledSkeletons();
-				if (RoledSkeletons.Num() == 1 )
-				{
-					if (InAssetData.GetClass()->IsChildOf(UAnimationAsset::StaticClass()) &&
-						RoledSkeletons[0].Skeleton &&
-						RoledSkeletons[0].Skeleton->IsCompatibleForEditor(InAssetData))
-					{
-						// We found a compatible skeleton in the schema.
-						return true;
-					}
-				}
-				else if (RoledSkeletons.Num() > 1)
-				{
-					if (const UMultiAnimAsset* MultiAnimAsset = Cast<UMultiAnimAsset>(InAssetData.GetAsset()))
-					{
-						if (MultiAnimAsset->GetNumRoles() != RoledSkeletons.Num())
-						{
-							return false;
-						}
-						
-						for (const FPoseSearchRoledSkeleton& RoledSkeleton : RoledSkeletons)
-						{
-							if (const UAnimationAsset* RoledAnimationAsset = MultiAnimAsset->GetAnimationAsset(RoledSkeleton.Role))
-							{
-								if (!RoledAnimationAsset->GetSkeleton()->IsCompatibleForEditor(RoledSkeleton.Skeleton))
-								{
-									return false;
-								}
-							}
-							else
-							{
-								// couldn't find a necessary asset for the RoledSkeleton
-								return false;
-							}
-						}
+				// We found a compatible skeleton in the schema.
+				return true;
+			}
+			
+			return false;
+		}
+				
+		if (!InAssetData.GetClass()->IsChildOf(UMultiAnimAsset::StaticClass()))
+		{
+			return false;
+		}
+		
+		// loading the UMultiAnimAsset
+		const UMultiAnimAsset* MultiAnimAsset = Cast<UMultiAnimAsset>(InAssetData.GetAsset());
+		check(MultiAnimAsset)
+
+		if (MultiAnimAsset->GetNumRoles() != RoledSkeletons.Num())
+		{
+			return false;
+		}
 
-						// we passed all the compatibility requirements!
-						return true;
-					}
+		for (const FPoseSearchRoledSkeleton& RoledSkeleton : RoledSkeletons)
+		{
+			if (const UAnimationAsset* RoledAnimationAsset = MultiAnimAsset->GetAnimationAsset(RoledSkeleton.Role))
+			{
+				if (!RoledAnimationAsset->GetSkeleton()->IsCompatibleForEditor(RoledSkeleton.Skeleton))
+				{
+					return false;
 				}
 			}
+			else
+			{
+				// couldn't find a necessary asset for the RoledSkeleton
+				return false;
+			}
 		}
-		
-		return false;
+
+		// we passed all the compatibility requirements!
+		return true;
 	}
 }
