MM interaction - added greedy policy to allocate islands search results. warping improvements via UAnimNotifyState_PoseSearchWarping

[RN] skip

diff --git a/Engine/Plugins/Animation/PoseSearch/Source/Runtime/Private/AnimNode_MotionMatchingInteraction.cpp b/Engine/Plugins/Animation/PoseSearch/Source/Runtime/Private/AnimNode_MotionMatchingInteraction.cpp
--- a/Engine/Plugins/Animation/PoseSearch/Source/Runtime/Private/AnimNode_MotionMatchingInteraction.cpp
+++ b/Engine/Plugins/Animation/PoseSearch/Source/Runtime/Private/AnimNode_MotionMatchingInteraction.cpp
@@ -5,6 +5,7 @@
 #include "Animation/AnimInstanceProxy.h"
 #include "Animation/AnimRootMotionProvider.h"
 #include "HAL/IConsoleManager.h"
+#include "PoseSearch/PoseSearchAnimNotifies.h"
 #include "PoseSearch/PoseSearchInteractionAsset.h"
 #include "PoseSearch/PoseSearchInteractionIsland.h"
 #include "PoseSearch/PoseSearchInteractionLibrary.h"
@@ -142,6 +143,44 @@ void FAnimNode_MotionMatchingInteraction::Evaluate_AnyThread(FPoseContext& Outpu
 			const FInteractionSearchContext InteractionSearchContext = UPoseSearchInteractionLibrary::FindSearchContext_AnyThread(Output.AnimInstanceProxy->GetAnimInstanceObject());
 			if (InteractionSearchContext.IsValid())
 			{
+				// calculating the animation time we want the full alignment transforms from
+				float WarpingTransformsExtractionTime = GetAccumulatedTime();
+				
+				// time to full warp alignment in seconds
+				float TimeToFullAlignment = 0.f;
+
+				// time after full alignment. if 0 it means we're still in full alignment, > 0 we can start relaxing the warping
+				float TimePastFullAlignment = 0.f;
+				if (const UAnimationAsset* AnimationAsset = MultiAnimAsset->GetAnimationAsset(CurrentResult.Role))
+				{
+					// @todo: support multiple UAnimNotifyState_PoseSearchWarping?
+					// sampler used only to extract the notify states. RootTransformOrigin can be set as Identity, since will not be relevant
+					const FAnimationAssetSampler SequenceBaseSampler(AnimationAsset, FTransform::Identity, CurrentResult.BlendParameters, FAnimationAssetSampler::DefaultRootTransformSamplingRate, false);
+					for (const FAnimNotifyEvent& NotifyEvent : SequenceBaseSampler.GetAllAnimNotifyEvents())
+					{
+						if (NotifyEvent.NotifyStateClass->GetClass()->IsChildOf<UAnimNotifyState_PoseSearchWarping>())
+						{
+							const float NotifyEventStartTime = NotifyEvent.GetTime();
+							const float NotifyEventEndTime = NotifyEventStartTime + NotifyEvent.GetDuration();
+
+							if (WarpingTransformsExtractionTime > NotifyEventEndTime)
+							{
+								TimePastFullAlignment = WarpingTransformsExtractionTime - NotifyEventEndTime;
+							}
+							 
+							if (NotifyEventStartTime > WarpingTransformsExtractionTime)
+							{
+								TimeToFullAlignment = NotifyEventStartTime - WarpingTransformsExtractionTime;
+								WarpingTransformsExtractionTime = NotifyEventStartTime;
+							}
+							
+							break;
+						}
+					}
+				}
+
+				// @todo: predicting ActorTransforms TimeToFullAlignment in the future! Maybe by using the trajectories..
+				
 				// populating ActorTransforms containing all the SkelMeshComponent transforms of all the characters involved into this interaction,
 				// to be able to calculate the required warp transform for this actor via UMultiAnimAsset::CalculateWarpTransforms
 				check(InteractionSearchContext.CheckForConsistency());
@@ -159,7 +198,7 @@ void FAnimNode_MotionMatchingInteraction::Evaluate_AnyThread(FPoseContext& Outpu
 					{
 						CurrentResultRoleIndex = MultiAnimAssetRoleIndex;
 					}
-					
+
 					if (const int32* InteractionSearchContextRoleIndex = InteractionSearchContextRoleToIndex.Find(MultiAnimAssetRole))
 					{
 						if (const UAnimInstance* AnimInstance = Cast<UAnimInstance>(InteractionSearchContext.AnimContexts[*InteractionSearchContextRoleIndex].Get()))
@@ -181,20 +220,23 @@ void FAnimNode_MotionMatchingInteraction::Evaluate_AnyThread(FPoseContext& Outpu
 
 				if (bAllRolesResolved && CurrentResultRoleIndex != INDEX_NONE)
 				{
+					
 					TArray<FTransform, TInlineAllocator<PreallocatedRolesNum>> FullAlignedActorRootBoneTransforms;
 					FullAlignedActorRootBoneTransforms.SetNum(MultiAnimAsset->GetNumRoles());
 
-					MultiAnimAsset->CalculateWarpTransforms(GetAccumulatedTime(), ActorTransforms, FullAlignedActorRootBoneTransforms);
+					MultiAnimAsset->CalculateWarpTransforms(WarpingTransformsExtractionTime, ActorTransforms, FullAlignedActorRootBoneTransforms);
 
 					const FTransform& FullAlignedActorRootBoneTransform = FullAlignedActorRootBoneTransforms[CurrentResultRoleIndex];
 
 					const FTransform RootBoneTransform = Output.AnimInstanceProxy->GetSkelMeshComponent()->GetBoneTransform(RootBoneIndexType);
 					FTransform WarpRootMotionTransform = FullAlignedActorRootBoneTransform.GetRelativeTransform(RootBoneTransform);
-
+					 
 					check(WarpRootMotionTransform.IsRotationNormalized());
 
-					const float DampenWarpingTranslation = WarpingTranslationHalflife > 0.f ? DamperImplicit(WarpingTranslationHalflife, CachedDeltaTime) : 1.f;
-					const float DampenWarpingRotation = WarpingRotationHalfLife > 0.f ? DamperImplicit(WarpingRotationHalfLife, CachedDeltaTime) : 1.f;
+					// @todo: use TimePastFullAlignment to "blend out" the warping transform
+					const float RatioToFullAlignment = TimeToFullAlignment > UE_KINDA_SMALL_NUMBER ? CachedDeltaTime / TimeToFullAlignment : UE_BIG_NUMBER;
+					const float DampenWarpingTranslation = FMath::Min(RatioToFullAlignment, WarpingTranslationHalflife > 0.f ? DamperImplicit(WarpingTranslationHalflife, CachedDeltaTime) : 1.f);
+					const float DampenWarpingRotation = FMath::Min(RatioToFullAlignment, WarpingRotationHalfLife > 0.f ? DamperImplicit(WarpingRotationHalfLife, CachedDeltaTime) : 1.f);
 
 					WarpRootMotionTransform.SetTranslation(WarpRootMotionTransform.GetTranslation() * DampenWarpingTranslation);
 					WarpRootMotionTransform.SetRotation(FQuat::Slerp(FQuat::Identity, WarpRootMotionTransform.GetRotation(), DampenWarpingRotation));
diff --git a/Engine/Plugins/Animation/PoseSearch/Source/Runtime/Private/PoseSearchAssetIndexer.cpp b/Engine/Plugins/Animation/PoseSearch/Source/Runtime/Private/PoseSearchAssetIndexer.cpp
--- a/Engine/Plugins/Animation/PoseSearch/Source/Runtime/Private/PoseSearchAssetIndexer.cpp
+++ b/Engine/Plugins/Animation/PoseSearch/Source/Runtime/Private/PoseSearchAssetIndexer.cpp
@@ -247,14 +247,16 @@ void FAssetIndexer::Process(int32 AssetIdx)
 		float CostAddend = SamplingContext.BaseCostBias;
 		bool bBlockTransition = false;
 
-		AssetSamplers.ExtractPoseSearchNotifyStates(SampleTime, [&bBlockTransition, &CostAddend](const UAnimNotifyState_PoseSearchBase* PoseSearchNotify)
+		AssetSamplers.ExtractPoseSearchNotifyStates(SampleTime, [&bBlockTransition, &CostAddend](const UAnimNotifyState_PoseSearchBase* PoseSearchBase)
 			{
-				if (PoseSearchNotify->GetClass()->IsChildOf<UAnimNotifyState_PoseSearchBlockTransition>())
+				if (PoseSearchBase->GetClass()->IsChildOf<UAnimNotifyState_PoseSearchBlockTransition>())
 				{
 					bBlockTransition = true;
 				}
-				else if (const UAnimNotifyState_PoseSearchModifyCost* ModifyCostNotify = Cast<const UAnimNotifyState_PoseSearchModifyCost>(PoseSearchNotify))
+				else if (PoseSearchBase->GetClass()->IsChildOf<UAnimNotifyState_PoseSearchModifyCost>())
 				{
+					const UAnimNotifyState_PoseSearchModifyCost* ModifyCostNotify = Cast<const UAnimNotifyState_PoseSearchModifyCost>(PoseSearchBase);
+					check(ModifyCostNotify);
 					CostAddend = ModifyCostNotify->CostAddend;
 				}
 				return true;
diff --git a/Engine/Plugins/Animation/PoseSearch/Source/Runtime/Private/PoseSearchAssetSampler.cpp b/Engine/Plugins/Animation/PoseSearch/Source/Runtime/Private/PoseSearchAssetSampler.cpp
--- a/Engine/Plugins/Animation/PoseSearch/Source/Runtime/Private/PoseSearchAssetSampler.cpp
+++ b/Engine/Plugins/Animation/PoseSearch/Source/Runtime/Private/PoseSearchAssetSampler.cpp
@@ -222,12 +222,12 @@ static int32 GetHighestWeightSample(const TArray<struct FBlendSampleData>& Sampl
 
 //////////////////////////////////////////////////////////////////////////
 // FAssetSamplerBase
-FAnimationAssetSampler::FAnimationAssetSampler(TObjectPtr<const UAnimationAsset> InAnimationAsset, const FTransform& InRootTransformOrigin, const FVector& InBlendParameters, int32 InRootTransformSamplingRate, bool bPreProcessRootTransform)
+FAnimationAssetSampler::FAnimationAssetSampler(const UAnimationAsset* InAnimationAsset, const FTransform& InRootTransformOrigin, const FVector& InBlendParameters, int32 InRootTransformSamplingRate, bool bPreProcessRootTransform)
 {
 	Init(InAnimationAsset, InRootTransformOrigin, InBlendParameters, InRootTransformSamplingRate, bPreProcessRootTransform);
 }
 
-void FAnimationAssetSampler::Init(TObjectPtr<const UAnimationAsset> InAnimationAsset, const FTransform& InRootTransformOrigin, const FVector& InBlendParameters, int32 InRootTransformSamplingRate, bool bPreProcessRootTransform)
+void FAnimationAssetSampler::Init(const UAnimationAsset* InAnimationAsset, const FTransform& InRootTransformOrigin, const FVector& InBlendParameters, int32 InRootTransformSamplingRate, bool bPreProcessRootTransform)
 {
 	AnimationAssetPtr = InAnimationAsset;
 	RootTransformOrigin = InRootTransformOrigin;
diff --git a/Engine/Plugins/Animation/PoseSearch/Source/Runtime/Private/PoseSearchDatabase.cpp b/Engine/Plugins/Animation/PoseSearch/Source/Runtime/Private/PoseSearchDatabase.cpp
--- a/Engine/Plugins/Animation/PoseSearch/Source/Runtime/Private/PoseSearchDatabase.cpp
+++ b/Engine/Plugins/Animation/PoseSearch/Source/Runtime/Private/PoseSearchDatabase.cpp
@@ -1330,14 +1330,14 @@ UE::PoseSearch::FSearchResult UPoseSearchDatabase::SearchContinuingPose(UE::Pose
 		if (const UAnimationAsset* AnimationAsset = DatabaseAnimationAssetBase->GetAnimationAssetForRole(DatabaseAnimationAssetBase->GetRole(RoleIndex)))
 		{
 			// sampler used only to extract the notify states. RootTransformOrigin can be set as Identity, since will not be relevant
-			const FAnimationAssetSampler SequenceBaseSampler(AnimationAsset, FTransform::Identity, SearchIndexAsset.GetBlendParameters());
+			const FAnimationAssetSampler SequenceBaseSampler(AnimationAsset, FTransform::Identity, SearchIndexAsset.GetBlendParameters(), FAnimationAssetSampler::DefaultRootTransformSamplingRate, false);
 
 			bool bDone = false;
-			SequenceBaseSampler.ExtractPoseSearchNotifyStates(SampleTime, [&UpdatedContinuingPoseCostBias, &bDone](const UAnimNotifyState_PoseSearchBase* PoseSearchNotify)
+			SequenceBaseSampler.ExtractPoseSearchNotifyStates(SampleTime, [&UpdatedContinuingPoseCostBias, &bDone](const UAnimNotifyState_PoseSearchBase* PoseSearchBase)
 				{
-					if (const UAnimNotifyState_PoseSearchOverrideContinuingPoseCostBias* ContinuingPoseCostBiasNotify = Cast<const UAnimNotifyState_PoseSearchOverrideContinuingPoseCostBias>(PoseSearchNotify))
+					if (const UAnimNotifyState_PoseSearchOverrideContinuingPoseCostBias* NotifyStateContinuingPoseCostBias = Cast<const UAnimNotifyState_PoseSearchOverrideContinuingPoseCostBias>(PoseSearchBase))
 					{
-						UpdatedContinuingPoseCostBias = ContinuingPoseCostBiasNotify->CostAddend;
+						UpdatedContinuingPoseCostBias = NotifyStateContinuingPoseCostBias->CostAddend;
 						bDone = true;
 						return false;
 					}
diff --git a/Engine/Plugins/Animation/PoseSearch/Source/Runtime/Private/PoseSearchInteractionIsland.cpp b/Engine/Plugins/Animation/PoseSearch/Source/Runtime/Private/PoseSearchInteractionIsland.cpp
--- a/Engine/Plugins/Animation/PoseSearch/Source/Runtime/Private/PoseSearchInteractionIsland.cpp
+++ b/Engine/Plugins/Animation/PoseSearch/Source/Runtime/Private/PoseSearchInteractionIsland.cpp
@@ -21,9 +21,12 @@ static bool GVarPoseSearchInteractionShowIslands = false;
 static FAutoConsoleVariableRef CVarPoseSearchInteractionShowIslands(TEXT("a.PoseSearchInteraction.ShowIslands"), GVarPoseSearchInteractionShowIslands, TEXT("Show Pose Search Interaction Islands"));
 #endif // ENABLE_DRAW_DEBUG && ENABLE_ANIM_DEBUG
 
+static int32 GVarPoseSearchInteractionIslandResultPolicy = 1;
+static FAutoConsoleVariableRef CVarPoseSearchInteractionIslandResultPolicy(TEXT("a.PoseSearchInteraction.IslandResultPolicy"), GVarPoseSearchInteractionIslandResultPolicy, TEXT("0: only best result is valid\n1: assign from best to worst result"));
+
 // Utils functions
 ///////////////////////////////////////////////////////////
-static FInteractionSearchResult InitSearchResult(const FSearchResult& SearchResult, int32 SearchIndex, const FInteractionSearchContext& SearchContext)
+static FInteractionSearchResult InitSearchResult(const FSearchResult& SearchResult, int32 SearchIndex)
 {
 	FInteractionSearchResult InteractionSearchResult;
 	static_cast<FSearchResult&>(InteractionSearchResult) = SearchResult;
@@ -31,36 +34,111 @@ static FInteractionSearchResult InitSearchResult(const FSearchResult& SearchResu
 	return InteractionSearchResult;
 }
 
-static void InitSearchResults(TArray<FInteractionSearchResult>& SearchResults, TConstArrayView<FSearchResult> PoseSearchResults, TConstArrayView<FInteractionSearchContext> SearchContexts)
+typedef TSet<const UObject*, DefaultKeyFuncs<const UObject*>, TInlineSetAllocator<64, TMemStackSetAllocator<>>> FVisitedAnimContexts;
+
+static bool IsPoseSearchResultUsable(int32 SearchIndex, TConstArrayView<FSearchResult> PoseSearchResults, TConstArrayView<FInteractionSearchContext> SearchContexts, const FVisitedAnimContexts& VisitedAnimContexts)
 {
-	// WIP!
-	// @todo: figure out multiple policies to use the most characters? right now only the best search is "valid" with the most characters
-	int32 BestSearchIndex = INDEX_NONE;
-	for (int32 SearchIndex = 0; SearchIndex < PoseSearchResults.Num(); ++SearchIndex)
+	if (!PoseSearchResults[SearchIndex].IsValid())
+	{
+		return false;
+	}
+
+	for (const TWeakObjectPtr<const UObject>& SearchAnimContext : SearchContexts[SearchIndex].AnimContexts)
 	{
-		if (PoseSearchResults[SearchIndex].IsValid())
+		if (const UObject* ValidSearchAnimContext = SearchAnimContext.Get())
 		{
-			if (BestSearchIndex == INDEX_NONE)
+			if (VisitedAnimContexts.Find(ValidSearchAnimContext))
 			{
-				BestSearchIndex = SearchIndex;
+				return false;
 			}
-			else if (SearchContexts[SearchIndex].Roles.Num() > SearchContexts[BestSearchIndex].Roles.Num())
+		}
+		else
+		{
+			return false;
+		}
+	}
+
+	return true;
+}
+
+static void InitSearchResults(TArray<FInteractionSearchResult>& SearchResults, TConstArrayView<FSearchResult> PoseSearchResults, TConstArrayView<FInteractionSearchContext> SearchContexts)
+{
+	SearchResults.Reset();
+
+	if (!PoseSearchResults.IsEmpty())
+	{
+		TArray<int32, TMemStackAllocator<>> SortedPoseSearchResults;
+		SortedPoseSearchResults.SetNum(PoseSearchResults.Num());
+		for (int32 Index = 0; Index < PoseSearchResults.Num(); ++Index)
+		{
+			SortedPoseSearchResults[Index] = Index;
+		}
+
+		SortedPoseSearchResults.StableSort([&SearchContexts, &PoseSearchResults](int32 IndexA, int32 IndexB)
+			{
+				const bool IsValidA = PoseSearchResults[IndexA].IsValid();
+				const bool IsValidB = PoseSearchResults[IndexB].IsValid();
+
+				if (IsValidA && !IsValidB)
+				{
+					return true;
+				}
+
+				if (!IsValidA && IsValidB)
+				{
+					return false;
+				}
+
+				if (!IsValidA && !IsValidB)
+				{
+					return true;
+				}
+
+				const int32 NumRolesA = SearchContexts[IndexA].Roles.Num();
+				const int32 NumRolesB = SearchContexts[IndexB].Roles.Num();
+
+				if (NumRolesA > NumRolesB)
+				{
+					return true;
+				}
+
+				if (NumRolesA < NumRolesB)
+				{
+					return false;
+				}
+
+				return PoseSearchResults[IndexA].PoseCost < PoseSearchResults[IndexB].PoseCost;
+			});
+
+		if (GVarPoseSearchInteractionIslandResultPolicy == 0)
+		{
+			// only best, valid with the highest cardinality in roles, result is valid
+			const int32 BestSearchIndex = SortedPoseSearchResults[0];
+			if (PoseSearchResults[BestSearchIndex].IsValid())
 			{
-				BestSearchIndex = SearchIndex;
+				SearchResults.SetNum(1, EAllowShrinking::No);
+				SearchResults[0] = InitSearchResult(PoseSearchResults[BestSearchIndex], BestSearchIndex);
 			}
-			else if (SearchContexts[SearchIndex].Roles.Num() == SearchContexts[BestSearchIndex].Roles.Num() &&
-				PoseSearchResults[SearchIndex].PoseCost < PoseSearchResults[BestSearchIndex].PoseCost)
+		}
+		else
+		{
+			// assign from best to worst result
+			FVisitedAnimContexts VisitedAnimContexts;
+			for (int32 SearchIndex : SortedPoseSearchResults)
 			{
-				BestSearchIndex = SearchIndex;
+				if (IsPoseSearchResultUsable(SearchIndex, PoseSearchResults, SearchContexts, VisitedAnimContexts))
+				{
+					SearchResults.Add(InitSearchResult(PoseSearchResults[SearchIndex], SearchIndex));
+					
+					for (const TWeakObjectPtr<const UObject>& SearchAnimContext : SearchContexts[SearchIndex].AnimContexts)
+					{
+						check(SearchAnimContext.Get());
+						VisitedAnimContexts.Add(SearchAnimContext.Get());
+					}
+				}
 			}
 		}
 	}
-
-	if (BestSearchIndex != INDEX_NONE)
-	{
-		SearchResults.SetNum(1, EAllowShrinking::No);
-		SearchResults[0] = InitSearchResult(PoseSearchResults[BestSearchIndex], BestSearchIndex, SearchContexts[BestSearchIndex]);
-	}
 }
 
 // FInteractionSearchContext
diff --git a/Engine/Plugins/Animation/PoseSearch/Source/Runtime/Private/PoseSearchLibrary.cpp b/Engine/Plugins/Animation/PoseSearch/Source/Runtime/Private/PoseSearchLibrary.cpp
--- a/Engine/Plugins/Animation/PoseSearch/Source/Runtime/Private/PoseSearchLibrary.cpp
+++ b/Engine/Plugins/Animation/PoseSearch/Source/Runtime/Private/PoseSearchLibrary.cpp
@@ -978,7 +978,7 @@ UE::PoseSearch::FSearchResult UPoseSearchLibrary::MotionMatch(
 					}
 				}
 
-				const FAnimationAssetSampler Sampler(AnimationAsset);
+				const FAnimationAssetSampler Sampler(AnimationAsset, FTransform::Identity, FVector::ZeroVector, FAnimationAssetSampler::DefaultRootTransformSamplingRate, false);
 				for (int32 i = 0; i < 2; ++i)
 				{
 					const float FuturePoseExtractionTime = FutureAnimationTime + (i - 1) * FiniteDelta;
diff --git a/Engine/Plugins/Animation/PoseSearch/Source/Runtime/Public/PoseSearch/PoseSearchAnimNotifies.h b/Engine/Plugins/Animation/PoseSearch/Source/Runtime/Public/PoseSearch/PoseSearchAnimNotifies.h
--- a/Engine/Plugins/Animation/PoseSearch/Source/Runtime/Public/PoseSearch/PoseSearchAnimNotifies.h
+++ b/Engine/Plugins/Animation/PoseSearch/Source/Runtime/Public/PoseSearch/PoseSearchAnimNotifies.h
@@ -117,3 +117,9 @@ public:
 
 	uint32 GetBranchInId() const;
 };
+
+UCLASS(Experimental, Blueprintable, meta = (DisplayName = "Pose Search: Warping"))
+class POSESEARCH_API UAnimNotifyState_PoseSearchWarping : public UAnimNotifyState_PoseSearchBase
+{
+	GENERATED_BODY()
+};
\ No newline at end of file
diff --git a/Engine/Plugins/Animation/PoseSearch/Source/Runtime/Public/PoseSearch/PoseSearchAssetSampler.h b/Engine/Plugins/Animation/PoseSearch/Source/Runtime/Public/PoseSearch/PoseSearchAssetSampler.h
--- a/Engine/Plugins/Animation/PoseSearch/Source/Runtime/Public/PoseSearch/PoseSearchAssetSampler.h
+++ b/Engine/Plugins/Animation/PoseSearch/Source/Runtime/Public/PoseSearch/PoseSearchAssetSampler.h
@@ -18,8 +18,8 @@ namespace UE::PoseSearch
 struct POSESEARCH_API FAnimationAssetSampler
 {
 	enum { DefaultRootTransformSamplingRate = 30 };
-	FAnimationAssetSampler(TObjectPtr<const UAnimationAsset> InAnimationAsset = nullptr, const FTransform& InRootTransformOrigin = FTransform::Identity, const FVector& InBlendParameters = FVector::ZeroVector, int32 InRootTransformSamplingRate = DefaultRootTransformSamplingRate, bool bPreProcessRootTransform = true);
-	void Init(TObjectPtr<const UAnimationAsset> InAnimationAsset, const FTransform& InRootTransformOrigin = FTransform::Identity, const FVector& InBlendParameters = FVector::ZeroVector, int32 InRootTransformSamplingRate = DefaultRootTransformSamplingRate, bool bPreProcessRootTransform = true);
+	FAnimationAssetSampler(const UAnimationAsset* InAnimationAsset = nullptr, const FTransform& InRootTransformOrigin = FTransform::Identity, const FVector& InBlendParameters = FVector::ZeroVector, int32 InRootTransformSamplingRate = DefaultRootTransformSamplingRate, bool bPreProcessRootTransform = true);
+	void Init(const UAnimationAsset* InAnimationAsset, const FTransform& InRootTransformOrigin = FTransform::Identity, const FVector& InBlendParameters = FVector::ZeroVector, int32 InRootTransformSamplingRate = DefaultRootTransformSamplingRate, bool bPreProcessRootTransform = true);
 
 	bool IsInitialized() const;
 	float GetPlayLength() const;
diff --git a/Engine/Plugins/Animation/PoseSearch/Source/Runtime/Public/PoseSearch/PoseSearchInteractionAsset.h b/Engine/Plugins/Animation/PoseSearch/Source/Runtime/Public/PoseSearch/PoseSearchInteractionAsset.h
--- a/Engine/Plugins/Animation/PoseSearch/Source/Runtime/Public/PoseSearch/PoseSearchInteractionAsset.h
+++ b/Engine/Plugins/Animation/PoseSearch/Source/Runtime/Public/PoseSearch/PoseSearchInteractionAsset.h
@@ -54,7 +54,7 @@ public:
 	UPROPERTY(Transient, EditAnywhere, Category = "Debug", meta = (EditCondition=bEnableDebugWarp, EditConditionHides))
 	TArray<FTransform> DebugWarpOffsets;
 
-	// used to test warping: 0 - no warping applied, 1 - full warping/aligment applied
+	// used to test warping: 0 - no warping applied, 1 - full warping/alignment applied
 	// test warping actors will be offsetted by Items::DebugWarpOffset transforms from the original
 	// UMultiAnimAsset::GetOrigin() definition and warped accordingly with CalculateWarpTransforms
 	// following the rotation and translation weights defined in Items::WarpingWeightRotation and
