Fixed crash when opening a PoseSearchDatabase with items or when adding items if its schema has no skeleton

[RN] minor

diff --git a/Engine/Plugins/Animation/PoseSearch/Source/Runtime/Private/PoseSearchDatabase.cpp b/Engine/Plugins/Animation/PoseSearch/Source/Runtime/Private/PoseSearchDatabase.cpp
--- a/Engine/Plugins/Animation/PoseSearch/Source/Runtime/Private/PoseSearchDatabase.cpp
+++ b/Engine/Plugins/Animation/PoseSearch/Source/Runtime/Private/PoseSearchDatabase.cpp
@@ -267,7 +267,7 @@ bool FPoseSearchDatabaseAnimationAssetBase::IsSkeletonCompatible(TObjectPtr<cons
 					if (RoledSkeleton.Role == Role)
 					{
 						// Match skeleton
-						if (RoledSkeleton.Skeleton->IsCompatibleForEditor(AssetData))
+						if (RoledSkeleton.Skeleton && RoledSkeleton.Skeleton->IsCompatibleForEditor(AssetData))
 						{
 							return true;
 						}
