posesearch - preserving only valid objects from GC in FPoseSearchDatabaseAsyncCacheTask::AddReferencedObjects

[RN] skip

diff --git a/Engine/Plugins/Animation/PoseSearch/Source/Runtime/Private/PoseSearchDerivedData.cpp b/Engine/Plugins/Animation/PoseSearch/Source/Runtime/Private/PoseSearchDerivedData.cpp
--- a/Engine/Plugins/Animation/PoseSearch/Source/Runtime/Private/PoseSearchDerivedData.cpp
+++ b/Engine/Plugins/Animation/PoseSearch/Source/Runtime/Private/PoseSearchDerivedData.cpp
@@ -1970,12 +1970,18 @@ void FPoseSearchDatabaseAsyncCacheTask::AddReferencedObjects(FReferenceCollector
 	const EState State = GetState();
 	if (State != EState::Ended && State != EState::Failed)
 	{
-		// keeping around the assets for starting or in progress tasks
-		Collector.AddReferencedObject(Database);
+		if (Database.IsValid())
+		{
+			// keeping around the assets for starting or in progress tasks
+			Collector.AddReferencedObject(Database);
+		}
 
 		for (TWeakObjectPtr<const UObject>& Dependency : DatabaseDependencies)
 		{
-			Collector.AddReferencedObject(Dependency);
+			if (Dependency.IsValid())
+			{
+				Collector.AddReferencedObject(Dependency);
+			}
 		}
 	}
 }
