posesearch - motion matching interaction warping improvements

[RN] skip

diff --git a/Engine/Plugins/Animation/PoseSearch/Source/Runtime/Private/AnimNode_MotionMatchingInteraction.cpp b/Engine/Plugins/Animation/PoseSearch/Source/Runtime/Private/AnimNode_MotionMatchingInteraction.cpp
--- a/Engine/Plugins/Animation/PoseSearch/Source/Runtime/Private/AnimNode_MotionMatchingInteraction.cpp
+++ b/Engine/Plugins/Animation/PoseSearch/Source/Runtime/Private/AnimNode_MotionMatchingInteraction.cpp
@@ -27,19 +27,6 @@ namespace UE::PoseSearch
 	}
 } // UE::PoseSearch
 
-
-void FAnimNode_MotionMatchingInteraction::GatherDebugData(FNodeDebugData& DebugData)
-{
-	Super::GatherDebugData(DebugData);
-
-	FString DebugLine = DebugData.GetNodeName(this);
-#if ENABLE_ANIM_DEBUG
-	static bool bActive = true;
-	DebugLine += FString::Printf(TEXT("\n - Active: (%s)"), bActive ? TEXT("True") : TEXT("False"));
-#endif
-	DebugData.AddDebugItem(DebugLine);
-}
-
 void FAnimNode_MotionMatchingInteraction::Reset()
 {
 	Super::Reset();
@@ -117,15 +104,16 @@ void FAnimNode_MotionMatchingInteraction::UpdateAssetPlayer(const FAnimationUpda
 	UpdatePlayRate(CurrentResult.WantedPlayRate);
 	UpdateBlendspaceParameters(BlendspaceUpdateMode, CurrentResult.BlendParameters);
 
-	// bypassing FAnimNode_BlendStack::UpdateAssetPlayer, since we overridden its behaviour
-	FAnimNode_BlendStack_Standalone::UpdateAssetPlayer(Context);
-
 #if ENABLE_ANIM_DEBUG
 	if (UE_TRACE_CHANNELEXPR_IS_ENABLED(AnimationChannel))
 	{
-		TRACE_ANIM_NODE_VALUE(Context, *FString("BlendToExecuted"), bBlendToExecuted);
+		TRACE_ANIM_NODE_VALUE(Context, *FString("WarpingRotationRatio"), WarpingRotationRatio);
+		TRACE_ANIM_NODE_VALUE(Context, *FString("WarpingTranslationRatio"), WarpingTranslationRatio);
 	}
 #endif // ENABLE_ANIM_DEBUG
+
+	// bypassing FAnimNode_BlendStack::UpdateAssetPlayer, since we overridden its behaviour
+	FAnimNode_BlendStack_Standalone::UpdateAssetPlayer(Context);
 }
 
 void FAnimNode_MotionMatchingInteraction::Evaluate_AnyThread(FPoseContext& Output)
@@ -136,129 +124,109 @@ void FAnimNode_MotionMatchingInteraction::Evaluate_AnyThread(FPoseContext& Outpu
 
 	Super::Evaluate_AnyThread(Output);
 
-	if (const UE::Anim::IAnimRootMotionProvider* RootMotionProvider = UE::Anim::IAnimRootMotionProvider::Get())
+	// calculating the animation time we want the full alignment transforms from
+	const float AccumulatedTime = GetAccumulatedTime();
+				
+	// warping alignment times in seconds
+	float WarpingAlignmentExtractionTime = AccumulatedTime;
+
+	const UE::Anim::IAnimRootMotionProvider* RootMotionProvider = UE::Anim::IAnimRootMotionProvider::Get();
+
+	if (!RootMotionProvider)
 	{
-		if (const UMultiAnimAsset* MultiAnimAsset = Cast<UMultiAnimAsset>(CurrentResult.SelectedAnimation))
+		UE_LOG(LogPoseSearch, Error, TEXT("FAnimNode_MotionMatchingInteraction::Evaluate_AnyThread couldn't find the IAnimRootMotionProvider"));
+	}
+	else if (WarpingTranslationRatio <= UE_KINDA_SMALL_NUMBER && WarpingRotationRatio <= UE_KINDA_SMALL_NUMBER)
+	{
+		// nothing to warp about...
+	}
+	else if (const UMultiAnimAsset* MultiAnimAsset = Cast<UMultiAnimAsset>(CurrentResult.SelectedAnimation))
+	{
+		const FInteractionSearchContext InteractionSearchContext = UPoseSearchInteractionLibrary::FindSearchContext_AnyThread(Output.AnimInstanceProxy->GetAnimInstanceObject());
+		if (InteractionSearchContext.IsValid())
 		{
-			const FInteractionSearchContext InteractionSearchContext = UPoseSearchInteractionLibrary::FindSearchContext_AnyThread(Output.AnimInstanceProxy->GetAnimInstanceObject());
-			if (InteractionSearchContext.IsValid())
+			if (bUsePoseSearchWarpingNotifyStates)
 			{
-				// calculating the animation time we want the full alignment transforms from
-				float WarpingTransformsExtractionTime = GetAccumulatedTime();
-				
-				// time to full warp alignment in seconds
-				float TimeToFullAlignment = 0.f;
-
-				// time after full alignment. if 0 it means we're still in full alignment, > 0 we can start relaxing the warping
-				float TimePastFullAlignment = 0.f;
+				// using UAnimNotifyState_PoseSearchWarping to change the WarpingAlignmentExtractionTime (time we sample the full alignment transforms at)
 				if (const UAnimationAsset* AnimationAsset = MultiAnimAsset->GetAnimationAsset(CurrentResult.Role))
 				{
-					// @todo: support multiple UAnimNotifyState_PoseSearchWarping?
-					// sampler used only to extract the notify states. RootTransformOrigin can be set as Identity, since will not be relevant
 					const FAnimationAssetSampler SequenceBaseSampler(AnimationAsset, FTransform::Identity, CurrentResult.BlendParameters, FAnimationAssetSampler::DefaultRootTransformSamplingRate, false);
 					for (const FAnimNotifyEvent& NotifyEvent : SequenceBaseSampler.GetAllAnimNotifyEvents())
 					{
 						if (NotifyEvent.NotifyStateClass->GetClass()->IsChildOf<UAnimNotifyState_PoseSearchWarping>())
 						{
-							const float NotifyEventStartTime = NotifyEvent.GetTime();
-							const float NotifyEventEndTime = NotifyEventStartTime + NotifyEvent.GetDuration();
-
-							if (WarpingTransformsExtractionTime > NotifyEventEndTime)
-							{
-								TimePastFullAlignment = WarpingTransformsExtractionTime - NotifyEventEndTime;
-							}
-							 
-							if (NotifyEventStartTime > WarpingTransformsExtractionTime)
-							{
-								TimeToFullAlignment = NotifyEventStartTime - WarpingTransformsExtractionTime;
-								WarpingTransformsExtractionTime = NotifyEventStartTime;
-							}
-							
-							break;
-						}
-					}
-				}
+							const UAnimNotifyState_PoseSearchWarping* NotifyStatePoseSearchWarping = Cast<UAnimNotifyState_PoseSearchWarping>(NotifyEvent.NotifyStateClass);
+							check(NotifyStatePoseSearchWarping);
 
-				// @todo: predicting ActorTransforms TimeToFullAlignment in the future! Maybe by using the trajectories..
-				
-				// populating ActorTransforms containing all the SkelMeshComponent transforms of all the characters involved into this interaction,
-				// to be able to calculate the required warp transform for this actor via UMultiAnimAsset::CalculateWarpTransforms
-				check(InteractionSearchContext.CheckForConsistency());
-				const FRoleToIndex InteractionSearchContextRoleToIndex = MakeRoleToIndex(InteractionSearchContext.Roles);
-				TArray<FTransform, TInlineAllocator<PreallocatedRolesNum>> ActorTransforms;
-				ActorTransforms.SetNum(MultiAnimAsset->GetNumRoles());
-
-				int32 CurrentResultRoleIndex = INDEX_NONE;
-				bool bAllRolesResolved = true;
-				for (int32 MultiAnimAssetRoleIndex = 0; MultiAnimAssetRoleIndex < MultiAnimAsset->GetNumRoles(); ++MultiAnimAssetRoleIndex)
-				{
-					const FRole& MultiAnimAssetRole = MultiAnimAsset->GetRole(MultiAnimAssetRoleIndex);
-
-					if (CurrentResultRoleIndex == INDEX_NONE && MultiAnimAssetRole == CurrentResult.Role)
-					{
-						CurrentResultRoleIndex = MultiAnimAssetRoleIndex;
-					}
-
-					if (const int32* InteractionSearchContextRoleIndex = InteractionSearchContextRoleToIndex.Find(MultiAnimAssetRole))
-					{
-						if (const UAnimInstance* AnimInstance = Cast<UAnimInstance>(InteractionSearchContext.AnimContexts[*InteractionSearchContextRoleIndex].Get()))
-						{
-							ActorTransforms[MultiAnimAssetRoleIndex] = AnimInstance->GetSkelMeshComponent()->GetComponentTransform();
-						}
-						else
-						{
-							bAllRolesResolved = false;
+							WarpingAlignmentExtractionTime = FMath::Max(NotifyEvent.GetTime(), AccumulatedTime);
 							break;
 						}
 					}
-					else
-					{
-						bAllRolesResolved = false;
-						break;
-					}
 				}
+			}
 
-				if (bAllRolesResolved && CurrentResultRoleIndex != INDEX_NONE)
+			int32 CurrentResultRoleIndex = INDEX_NONE;
+			for (int32 MultiAnimAssetRoleIndex = 0; MultiAnimAssetRoleIndex < MultiAnimAsset->GetNumRoles(); ++MultiAnimAssetRoleIndex)
+			{
+				if (MultiAnimAsset->GetRole(MultiAnimAssetRoleIndex) == CurrentResult.Role)
 				{
-					
-					TArray<FTransform, TInlineAllocator<PreallocatedRolesNum>> FullAlignedActorRootBoneTransforms;
-					FullAlignedActorRootBoneTransforms.SetNum(MultiAnimAsset->GetNumRoles());
+					CurrentResultRoleIndex = MultiAnimAssetRoleIndex;
+					break;
+				}
+			}
 
-					MultiAnimAsset->CalculateWarpTransforms(WarpingTransformsExtractionTime, ActorTransforms, FullAlignedActorRootBoneTransforms);
+			check(InteractionSearchContext.CheckForConsistency());
 
-					const FTransform& FullAlignedActorRootBoneTransform = FullAlignedActorRootBoneTransforms[CurrentResultRoleIndex];
+			TArray<FTransform, TInlineAllocator<PreallocatedRolesNum>> ActorTransforms;
+			ActorTransforms.SetNum(MultiAnimAsset->GetNumRoles());
+			// @todo: predicting ActorTransforms TimeToFullAlignment in the future! Maybe by using the trajectories..
+			if (CurrentResultRoleIndex != INDEX_NONE && InteractionSearchContext.GetActorTransforms(ActorTransforms, MultiAnimAsset))
+			{
+				// calculating the world space warp tranform to have the character in its "original" no reference frame (where no warping associated to MultiAnimAsset is impacting its transform)
+				FTransform RootMotionTransform;
+				RootMotionProvider->ExtractRootMotion(Output.CustomAttributes, RootMotionTransform);
+				// NoTe that ActorTransforms[CurrentResultRoleIndex] is Output.AnimInstanceProxy->GetSkelMeshComponent()->GetComponentTransform();
+				const FTransform AlignmentTransformNoWarpAdjustment = RootMotionTransform * ActorTransforms[CurrentResultRoleIndex];
+					
+				// calculating the world space warp tranform to have the character fully aligned to what MultiAnimAsset wants
+				TArray<FTransform, TInlineAllocator<PreallocatedRolesNum>> FullAlignedActorRootBoneTransforms;
+				FullAlignedActorRootBoneTransforms.SetNum(MultiAnimAsset->GetNumRoles());
+				MultiAnimAsset->CalculateWarpTransforms(WarpingAlignmentExtractionTime, ActorTransforms, FullAlignedActorRootBoneTransforms);
+				const FTransform AlignmentTransformFullWarpAdjustment = FullAlignedActorRootBoneTransforms[CurrentResultRoleIndex];
+					
+				// @todo: compensating AlignmentTransformFullWarpAdjustment by the delta time between AccumulatedTime and WarpingAlignmentExtractionTime?				
 
-					const FTransform RootBoneTransform = Output.AnimInstanceProxy->GetSkelMeshComponent()->GetBoneTransform(RootBoneIndexType);
-					FTransform WarpRootMotionTransform = FullAlignedActorRootBoneTransform.GetRelativeTransform(RootBoneTransform);
-					 
-					check(WarpRootMotionTransform.IsRotationNormalized());
+				// calculating the world space blended warp aligment transform
+				const FTransform AlignmentTransform(
+					FMath::Lerp(AlignmentTransformNoWarpAdjustment.GetRotation(), AlignmentTransformFullWarpAdjustment.GetRotation(), WarpingRotationRatio),
+					FMath::Lerp(AlignmentTransformNoWarpAdjustment.GetTranslation(), AlignmentTransformFullWarpAdjustment.GetTranslation(), WarpingTranslationRatio),
+					FVector::OneVector);
 
-					// @todo: use TimePastFullAlignment to "blend out" the warping transform
-					const float RatioToFullAlignment = TimeToFullAlignment > UE_KINDA_SMALL_NUMBER ? CachedDeltaTime / TimeToFullAlignment : UE_BIG_NUMBER;
-					const float DampenWarpingTranslation = FMath::Min(RatioToFullAlignment, WarpingTranslationHalflife > 0.f ? DamperImplicit(WarpingTranslationHalflife, CachedDeltaTime) : 1.f);
-					const float DampenWarpingRotation = FMath::Min(RatioToFullAlignment, WarpingRotationHalfLife > 0.f ? DamperImplicit(WarpingRotationHalfLife, CachedDeltaTime) : 1.f);
+				// calculating WarpRootMotionTransform as AlignmentTransform in root bone space
+				const FTransform RootBoneTransform = Output.AnimInstanceProxy->GetSkelMeshComponent()->GetBoneTransform(RootBoneIndexType);
+				FTransform WarpRootMotionTransform = AlignmentTransform.GetRelativeTransform(RootBoneTransform);
+				check(WarpRootMotionTransform.IsRotationNormalized());
 
-					WarpRootMotionTransform.SetTranslation(WarpRootMotionTransform.GetTranslation() * DampenWarpingTranslation);
-					WarpRootMotionTransform.SetRotation(FQuat::Slerp(FQuat::Identity, WarpRootMotionTransform.GetRotation(), DampenWarpingRotation));
+				// apply some damping to WarpRootMotionTransform in root bone space
+				const float DampenWarpingTranslation = WarpingTranslationHalflife > 0.f ? DamperImplicit(WarpingTranslationHalflife, CachedDeltaTime) : 1.f;
+				const float DampenWarpingRotation = WarpingRotationHalfLife > 0.f ? DamperImplicit(WarpingRotationHalfLife, CachedDeltaTime) : 1.f;
+				WarpRootMotionTransform.SetTranslation(WarpRootMotionTransform.GetTranslation() * DampenWarpingTranslation);
+				WarpRootMotionTransform.SetRotation(FQuat::Slerp(FQuat::Identity, WarpRootMotionTransform.GetRotation(), DampenWarpingRotation));
 
 #if ENABLE_VISUAL_LOG
-					if (FVisualLogger::IsRecording())
-					{
-						static const TCHAR* LogName = TEXT("MotionMatchingInteraction");
+				if (FVisualLogger::IsRecording())
+				{
+					static const TCHAR* LogName = TEXT("MotionMatchingInteraction");
 
-						UObject* LogOwner = Output.AnimInstanceProxy->GetAnimInstanceObject();
-						const FTransform WorldWarpRootMotionTransform = WarpRootMotionTransform * RootBoneTransform;
-						UE_VLOG_ARROW(LogOwner, LogName, Display, WorldWarpRootMotionTransform.GetLocation(), WorldWarpRootMotionTransform.GetLocation() + WorldWarpRootMotionTransform.GetRotation().GetForwardVector() * 100, FColor::Green, TEXT(""));
-						UE_VLOG_ARROW(LogOwner, LogName, Display, FullAlignedActorRootBoneTransform.GetLocation(), FullAlignedActorRootBoneTransform.GetLocation() + FullAlignedActorRootBoneTransform.GetRotation().GetForwardVector() * 100, FColor::Blue, TEXT(""));
-					}
-#endif
-					RootMotionProvider->OverrideRootMotion(WarpRootMotionTransform, Output.CustomAttributes);
+					UObject* LogOwner = Output.AnimInstanceProxy->GetAnimInstanceObject();
+					const FTransform WorldWarpRootMotionTransform = WarpRootMotionTransform * RootBoneTransform;
+					UE_VLOG_ARROW(LogOwner, LogName, Display, WorldWarpRootMotionTransform.GetLocation(), WorldWarpRootMotionTransform.GetLocation() + WorldWarpRootMotionTransform.GetRotation().GetForwardVector() * 100, FColor::Green, TEXT(""));
+					UE_VLOG_ARROW(LogOwner, LogName, Display, AlignmentTransform.GetLocation(), AlignmentTransform.GetLocation() + AlignmentTransform.GetRotation().GetForwardVector() * 100, FColor::Blue, TEXT(""));
+					UE_VLOG_ARROW(LogOwner, LogName, Display, RootBoneTransform.GetLocation(), RootBoneTransform.GetLocation() + RootBoneTransform.GetRotation().GetForwardVector() * 100, FColor::Orange, TEXT(""));
 				}
+#endif
+				RootMotionProvider->OverrideRootMotion(WarpRootMotionTransform, Output.CustomAttributes);
 			}
 		}
 	}
-	else
-	{
-		UE_LOG(LogPoseSearch, Error, TEXT("FAnimNode_MotionMatchingInteraction::Evaluate_AnyThread couldn't find the IAnimRootMotionProvider"));
-	}
 }
diff --git a/Engine/Plugins/Animation/PoseSearch/Source/Runtime/Private/PoseSearchInteractionIsland.cpp b/Engine/Plugins/Animation/PoseSearch/Source/Runtime/Private/PoseSearchInteractionIsland.cpp
--- a/Engine/Plugins/Animation/PoseSearch/Source/Runtime/Private/PoseSearchInteractionIsland.cpp
+++ b/Engine/Plugins/Animation/PoseSearch/Source/Runtime/Private/PoseSearchInteractionIsland.cpp
@@ -250,6 +250,41 @@ bool FInteractionSearchContext::IsEquivalent(const FInteractionSearchContext& Ot
 	return false;
 }
 
+bool FInteractionSearchContext::GetActorTransforms(TArrayView<FTransform> ActorTransforms, const UMultiAnimAsset* MultiAnimAsset) const
+{
+	check(MultiAnimAsset);
+	check(ActorTransforms.Num() == MultiAnimAsset->GetNumRoles());
+
+	const FRoleToIndex InteractionSearchContextRoleToIndex = MakeRoleToIndex(Roles);
+
+	bool bAllRolesResolved = true;
+	for (int32 MultiAnimAssetRoleIndex = 0; MultiAnimAssetRoleIndex < MultiAnimAsset->GetNumRoles(); ++MultiAnimAssetRoleIndex)
+	{
+		const FRole& MultiAnimAssetRole = MultiAnimAsset->GetRole(MultiAnimAssetRoleIndex);
+
+		if (const int32* InteractionSearchContextRoleIndex = InteractionSearchContextRoleToIndex.Find(MultiAnimAssetRole))
+		{
+			if (const UAnimInstance* AnimInstance = Cast<UAnimInstance>(AnimContexts[*InteractionSearchContextRoleIndex].Get()))
+			{
+				ActorTransforms[MultiAnimAssetRoleIndex] = AnimInstance->GetSkelMeshComponent()->GetComponentTransform();
+			}
+			else
+			{
+				ActorTransforms[MultiAnimAssetRoleIndex] = FTransform::Identity;
+				bAllRolesResolved = false;
+				break;
+			}
+		}
+		else
+		{
+			ActorTransforms[MultiAnimAssetRoleIndex] = FTransform::Identity;
+			bAllRolesResolved = false;
+			break;
+		}
+	}
+	return bAllRolesResolved;
+}
+
 // FInteractionSearchResult
 ///////////////////////////////////////////////////////////
 bool FInteractionSearchResult::operator==(const FInteractionSearchResult& Other) const
diff --git a/Engine/Plugins/Animation/PoseSearch/Source/Runtime/Public/PoseSearch/AnimNode_MotionMatchingInteraction.h b/Engine/Plugins/Animation/PoseSearch/Source/Runtime/Public/PoseSearch/AnimNode_MotionMatchingInteraction.h
--- a/Engine/Plugins/Animation/PoseSearch/Source/Runtime/Public/PoseSearch/AnimNode_MotionMatchingInteraction.h
+++ b/Engine/Plugins/Animation/PoseSearch/Source/Runtime/Public/PoseSearch/AnimNode_MotionMatchingInteraction.h
@@ -19,16 +19,28 @@ struct POSESEARCH_API FAnimNode_MotionMatchingInteraction : public FAnimNode_Ble
 	UPROPERTY(EditAnywhere, Category = Search, meta = (PinHiddenByDefault))
 	bool bValidateResultAgainstAvailabilities = true;
 
+	// amount or translation warping to apply
+	UPROPERTY(EditAnywhere, Category = Warping, meta = (PinHiddenByDefault, ClampMin = "0", ClampMax = "1"))
+	float WarpingTranslationRatio = 1.f;
+
 	// Controls how fast the warp translation offset is blended in
 	// Values closer to 0 make it faster. if 0 the warp translation is applied in full
 	UPROPERTY(EditAnywhere, Category = Warping, meta = (PinHiddenByDefault))
 	float WarpingTranslationHalflife = 0.05f;
 
+	// amount or rotation warping to apply
+	UPROPERTY(EditAnywhere, Category = Warping, meta = (PinHiddenByDefault, ClampMin = "0", ClampMax = "1"))
+	float WarpingRotationRatio = 1.f;
+
 	// Controls how fast the warp rotation offset is blended in
 	// Values closer to 0 make it faster. if 0 the warp rotation is applied in full
 	UPROPERTY(EditAnywhere, Category = Warping, meta = (PinHiddenByDefault))
 	float WarpingRotationHalfLife = 0.05f;
 
+	// enable / disable the usage of UAnimNotifyState_PoseSearchWarping(s) to change the WarpingAlignmentExtractionTime (time we sample the full alignment transforms at)
+	UPROPERTY(EditAnywhere, Category = Warping, meta = (PinHiddenByDefault))
+	bool bUsePoseSearchWarpingNotifyStates = true;
+
 	// Reset the blend stack if it has become relevant to the graph after not being updated on previous frames.
 	UPROPERTY(EditAnywhere, Category = Settings)
 	bool bResetOnBecomingRelevant = true;
@@ -52,7 +64,6 @@ struct POSESEARCH_API FAnimNode_MotionMatchingInteraction : public FAnimNode_Ble
 	bool bUseInertialBlend = false;
 
 	// FAnimNode_Base interface
-	virtual void GatherDebugData(FNodeDebugData& DebugData) override;
 	virtual void Evaluate_AnyThread(FPoseContext& Output) override;
 
 	virtual void UpdateAssetPlayer(const FAnimationUpdateContext& Context) override;
diff --git a/Engine/Plugins/Animation/PoseSearch/Source/Runtime/Public/PoseSearch/PoseSearchInteractionIsland.h b/Engine/Plugins/Animation/PoseSearch/Source/Runtime/Public/PoseSearch/PoseSearchInteractionIsland.h
--- a/Engine/Plugins/Animation/PoseSearch/Source/Runtime/Public/PoseSearch/PoseSearchInteractionIsland.h
+++ b/Engine/Plugins/Animation/PoseSearch/Source/Runtime/Public/PoseSearch/PoseSearchInteractionIsland.h
@@ -7,6 +7,7 @@
 struct FPoseSearchInteractionBlueprintResult;
 struct FPoseSearchContinuingProperties;
 class UCharacterMovementComponent;
+class UMultiAnimAsset;
 
 namespace UE::PoseSearch
 {
@@ -33,6 +34,7 @@ struct FInteractionSearchContext
 
 	bool IsValid() const { return !AnimContexts.IsEmpty(); }
 	bool IsEquivalent(const FInteractionSearchContext& Other) const;
+	bool GetActorTransforms(TArrayView<FTransform> ActorTransforms, const UMultiAnimAsset* MultiAnimAsset) const;
 
 #if DO_CHECK
 	bool CheckForConsistency() const;
