Fix trajectory generation method for moving platforms and rapid changes in controller orientation

#tests AnimSanbox win64 editor
[RN] fix

diff --git a/Engine/Plugins/Animation/PoseSearch/Source/Runtime/Private/PoseSearchTrajectoryLibrary.cpp b/Engine/Plugins/Animation/PoseSearch/Source/Runtime/Private/PoseSearchTrajectoryLibrary.cpp
--- a/Engine/Plugins/Animation/PoseSearch/Source/Runtime/Private/PoseSearchTrajectoryLibrary.cpp
+++ b/Engine/Plugins/Animation/PoseSearch/Source/Runtime/Private/PoseSearchTrajectoryLibrary.cpp
@@ -248,17 +248,10 @@ void UPoseSearchTrajectoryLibrary::UpdateHistory_TransformHistory(
 
 		check(NumHistorySamples <= Trajectory.Samples.Num());
 
-		// Our trajectory's "current" position assumes the we have the same delta time as the previous frame.
-		// so use predicted trajectory with current time step.
-		const FVector PredictedPositionAdjusted = Trajectory.GetSampleAtTime(DeltaTime).Position;
-
-		// converting all the history samples relative to the previous character position (Trajectory.Samples[NumHistorySamples].Position)
-		for (int32 Index = 0; Index < NumHistorySamples; ++Index)
-		{
-			Trajectory.Samples[Index].Position = PredictedPositionAdjusted - Trajectory.Samples[Index].Position;
-		}
-
-		FVector CurrentTranslation = CurrentVelocity * DeltaTime;
+		// Trajectory.Samples[NumHistorySamples] is last frame position! (assuming this is called every frame)
+		const FVector CurrentTranslationFromMover = CurrentVelocity * DeltaTime;
+		const FVector TranslationSinceLastFrame = CurrentPosition - Trajectory.Samples[NumHistorySamples].Position;
+		const FVector InferredGroundTranslation = TranslationSinceLastFrame - CurrentTranslationFromMover;
 
 		// Shift history Samples when it's time to record a new one.
 		if (SecondsPerHistorySample <= 0.f || FMath::Abs(Trajectory.Samples[NumHistorySamples - 1].AccumulatedSeconds) >= SecondsPerHistorySample)
@@ -266,28 +259,26 @@ void UPoseSearchTrajectoryLibrary::UpdateHistory_TransformHistory(
 			for (int32 Index = 0; Index < NumHistorySamples - 1; ++Index)
 			{
 				Trajectory.Samples[Index].AccumulatedSeconds = Trajectory.Samples[Index + 1].AccumulatedSeconds - DeltaTime;
-				Trajectory.Samples[Index].Position = Trajectory.Samples[Index + 1].Position + CurrentTranslation;
+				Trajectory.Samples[Index].Position = Trajectory.Samples[Index + 1].Position + InferredGroundTranslation;
 				Trajectory.Samples[Index].Facing = Trajectory.Samples[Index + 1].Facing;
 			}
 
-			Trajectory.Samples[NumHistorySamples - 1].AccumulatedSeconds = 0.f;
-			Trajectory.Samples[NumHistorySamples - 1].Position = CurrentTranslation;
+			// Adding a new history record
+			// Note we add current frame at delta time in the future since t = 0 should be the previous frame
+			Trajectory.Samples[NumHistorySamples - 1].AccumulatedSeconds = DeltaTime;
+			Trajectory.Samples[NumHistorySamples - 1].Position = CurrentPosition;
 			Trajectory.Samples[NumHistorySamples - 1].Facing = Trajectory.Samples[NumHistorySamples].Facing;
 		}
 		else
 		{
+			// Didn't record a new history position, update timers and shift by ground translation
+
 			for (int32 Index = 0; Index < NumHistorySamples; ++Index)
 			{
 				Trajectory.Samples[Index].AccumulatedSeconds -= DeltaTime;
-				Trajectory.Samples[Index].Position += CurrentTranslation;
+				Trajectory.Samples[Index].Position += InferredGroundTranslation;
 			}
 		}
-
-		// converting the history sample positions in world space by applying the current world position.
-		for (int32 Index = 0; Index < NumHistorySamples; ++Index)
-		{
-			Trajectory.Samples[Index].Position = CurrentPosition - Trajectory.Samples[Index].Position;
-		}
 	}
 }
 
