posesearch - motion matching interaction IK constraints

[RN] skip

diff --git a/Engine/Plugins/Animation/PoseSearch/Source/Runtime/Private/AnimNode_MotionMatchingInteraction.cpp b/Engine/Plugins/Animation/PoseSearch/Source/Runtime/Private/AnimNode_MotionMatchingInteraction.cpp
--- a/Engine/Plugins/Animation/PoseSearch/Source/Runtime/Private/AnimNode_MotionMatchingInteraction.cpp
+++ b/Engine/Plugins/Animation/PoseSearch/Source/Runtime/Private/AnimNode_MotionMatchingInteraction.cpp
@@ -12,19 +12,27 @@
 #include "PoseSearch/PoseHistoryProvider.h"
 #include "PoseSearch/PoseSearchDatabase.h"
 #include "PoseSearch/PoseSearchLibrary.h"
+#include "PoseSearch/PoseSearchAssetSamplerLibrary.h"
 #include "PoseSearch/PoseSearchSchema.h"
 #include "VisualLogger/VisualLogger.h"
 
 namespace UE::PoseSearch
 {
-	// Taken from https://theorangeduck.com/page/spring-roll-call#implicitspringdamper
-	static float DamperImplicit(float Halflife, float DeltaTime, float Epsilon = 1e-8f)
-	{
-		// Halflife values very close to 0 approach infinity, and result in big motion spikes when Halflife < DeltaTime
-		// This is a hack, and adds some degree of frame-rate dependency, but it holds up even at lower frame-rates.
-		Halflife = FMath::Max(Halflife, DeltaTime);
-		return FMath::Clamp((1.0f - FMath::InvExpApprox((0.69314718056f * DeltaTime) / (Halflife + Epsilon))), 0.0f, 1.0f);
-	}
+
+#if ENABLE_DRAW_DEBUG && ENABLE_ANIM_DEBUG
+static bool GVarAnimNodeMotionMatchingInteractionDrawFullAlignedPose = false;
+static FAutoConsoleVariableRef CVarAnimNodeMotionMatchingInteractionDrawFullAlignedPose(TEXT("a.AnimNode.MotionMatchingInteraction.DrawFullAlignedPose"), GVarAnimNodeMotionMatchingInteractionDrawFullAlignedPose, TEXT("Draw Full Aligned Pose"));
+#endif // ENABLE_DRAW_DEBUG && ENABLE_ANIM_DEBUG
+
+// Taken from https://theorangeduck.com/page/spring-roll-call#implicitspringdamper
+static float DamperImplicit(float Halflife, float DeltaTime, float Epsilon = 1e-8f)
+{
+	// Halflife values very close to 0 approach infinity, and result in big motion spikes when Halflife < DeltaTime
+	// This is a hack, and adds some degree of frame-rate dependency, but it holds up even at lower frame-rates.
+	Halflife = FMath::Max(Halflife, DeltaTime);
+	return FMath::Clamp((1.0f - FMath::InvExpApprox((0.69314718056f * DeltaTime) / (Halflife + Epsilon))), 0.0f, 1.0f);
+}
+
 } // UE::PoseSearch
 
 void FAnimNode_MotionMatchingInteraction::Reset()
@@ -47,6 +55,8 @@ void FAnimNode_MotionMatchingInteraction::UpdateAssetPlayer(const FAnimationUpda
 {
 	using namespace UE::PoseSearch;
 
+	CheckInteractionThreadSafety(Context.AnimInstanceProxy->GetAnimInstanceObject(), FInteractionValidator::EPhase::MotionMatchingInteractionUpdate);
+
 	if (NeedsReset(Context))
 	{
 		Reset();
@@ -68,6 +78,7 @@ void FAnimNode_MotionMatchingInteraction::UpdateAssetPlayer(const FAnimationUpda
 
 		check(Context.AnimInstanceProxy);
 		CurrentResult = UPoseSearchInteractionLibrary::MotionMatchInteraction(Availabilities, Context.AnimInstanceProxy->GetAnimInstanceObject(), ContinuingProperties, PoseHistoryProvider->GetHistoryCollector(), bValidateResultAgainstAvailabilities);
+		check(CurrentResult.ActorRootTransforms.Num() == CurrentResult.ActorRootBoneTransforms.Num());
 
 		UAnimationAsset* RoledAnimAsset = nullptr;
 		if (const UMultiAnimAsset* MultiAnimAsset = Cast<UMultiAnimAsset>(CurrentResult.SelectedAnimation))
@@ -122,6 +133,8 @@ void FAnimNode_MotionMatchingInteraction::Evaluate_AnyThread(FPoseContext& Outpu
 
 	check(Output.AnimInstanceProxy);
 
+	CheckInteractionThreadSafety(Output.AnimInstanceProxy->GetAnimInstanceObject(), FInteractionValidator::EPhase::MotionMatchingInteractionEvaluate);
+
 	Super::Evaluate_AnyThread(Output);
 
 	// calculating the animation time we want the full alignment transforms from
@@ -130,103 +143,243 @@ void FAnimNode_MotionMatchingInteraction::Evaluate_AnyThread(FPoseContext& Outpu
 	// warping alignment times in seconds
 	float WarpingAlignmentExtractionTime = AccumulatedTime;
 
+	UAnimInstance* AnimInstance = Cast<UAnimInstance>(Output.AnimInstanceProxy->GetAnimInstanceObject());
+	if (!AnimInstance)
+	{
+		UE_LOG(LogPoseSearch, Error, TEXT("FAnimNode_MotionMatchingInteraction::Evaluate_AnyThread couldn't find the AnimInstance!?"));
+		return;
+	}
+
 	const UE::Anim::IAnimRootMotionProvider* RootMotionProvider = UE::Anim::IAnimRootMotionProvider::Get();
 
 	if (!RootMotionProvider)
 	{
 		UE_LOG(LogPoseSearch, Error, TEXT("FAnimNode_MotionMatchingInteraction::Evaluate_AnyThread couldn't find the IAnimRootMotionProvider"));
+		return;
 	}
-	else if (WarpingTranslationRatio <= UE_KINDA_SMALL_NUMBER && WarpingRotationRatio <= UE_KINDA_SMALL_NUMBER)
+
+	const UMultiAnimAsset* MultiAnimAsset = Cast<UMultiAnimAsset>(CurrentResult.SelectedAnimation);
+	if (!MultiAnimAsset)
 	{
-		// nothing to warp about...
+		return;
 	}
-	else if (const UMultiAnimAsset* MultiAnimAsset = Cast<UMultiAnimAsset>(CurrentResult.SelectedAnimation))
+
+	const int32 NumRoles = MultiAnimAsset->GetNumRoles();
+	if (CurrentResult.ActorRootTransforms.Num() != NumRoles)
 	{
-		const FInteractionSearchContext InteractionSearchContext = UPoseSearchInteractionLibrary::FindSearchContext_AnyThread(Output.AnimInstanceProxy->GetAnimInstanceObject());
-		if (InteractionSearchContext.IsValid())
+		// warping is supported only for UMultiAnimAsset(s)
+		return;
+	}
+
+	int32 CurrentResultRoleIndex = INDEX_NONE;
+	for (int32 MultiAnimAssetRoleIndex = 0; MultiAnimAssetRoleIndex < NumRoles; ++MultiAnimAssetRoleIndex)
+	{
+		if (MultiAnimAsset->GetRole(MultiAnimAssetRoleIndex) == CurrentResult.Role)
 		{
-			if (bUsePoseSearchWarpingNotifyStates)
+			CurrentResultRoleIndex = MultiAnimAssetRoleIndex;
+			break;
+		}
+	}
+	if (CurrentResultRoleIndex == INDEX_NONE)
+	{
+		return;
+	}
+
+	const UAnimationAsset* AnimationAsset = MultiAnimAsset->GetAnimationAsset(CurrentResult.Role);
+	if (!AnimationAsset)
+	{
+		UE_LOG(LogPoseSearch, Error, TEXT("FAnimNode_MotionMatchingInteraction::Evaluate_AnyThread MultiAnimAsset %s for Role %s is invalid!"), *GetNameSafe(MultiAnimAsset), *CurrentResult.Role.ToString());
+		return;
+	}
+
+	FMemMark Mark(FMemStack::Get());
+
+	// @todo: predicting ActorRootTransforms TimeToFullAlignment in the future! Maybe by using the trajectories..
+	const TArrayView<FTransform> ActorRootTransforms = CurrentResult.ActorRootTransforms;
+	const TArrayView<FTransform> ActorRootBoneTransforms = CurrentResult.ActorRootBoneTransforms;
+
+	// NoTe: applying the previous frame historical root bone orientation to the ActorRootTransforms (array of ComponentTransform(s)), instead of applying the translation as well via
+	// ActorRootBoneTransforms[i] * ActorRootTransforms[i], to get the ActorTransforms[i] where the animation wants it to be without being affected by OffsetRootBone trnaslations,
+	// so we'll be able to warp into full aligment and blend back once the aligment is no longer required, ultimately converging back where the capsule is located 
+	// (calculated below in AlignmentTransformNoWarpAdjustment, and adjusted by any RootMotionTransform)
+	TArray<FTransform, TInlineAllocator<PreallocatedRolesNum, TMemStackAllocator<>>> ActorTransforms;
+	ActorTransforms.SetNum(NumRoles);
+	for (int32 MultiAnimAssetRoleIndex = 0; MultiAnimAssetRoleIndex < NumRoles; ++MultiAnimAssetRoleIndex)
+	{
+		ActorTransforms[MultiAnimAssetRoleIndex] = FTransform(
+			ActorRootTransforms[MultiAnimAssetRoleIndex].GetRotation() * ActorRootBoneTransforms[MultiAnimAssetRoleIndex].GetRotation(),
+			ActorRootTransforms[MultiAnimAssetRoleIndex].GetTranslation(),
+			ActorRootTransforms[MultiAnimAssetRoleIndex].GetScale3D());
+	}
+
+	const FAnimationAssetSampler Sampler(AnimationAsset, FTransform::Identity, CurrentResult.BlendParameters, FAnimationAssetSampler::DefaultRootTransformSamplingRate, false);
+
+	if (bUsePoseSearchWarpingNotifyStates)
+	{
+		// using UAnimNotifyState_PoseSearchWarping to change the WarpingAlignmentExtractionTime (time we sample the full alignment transforms at)
+		for (const FAnimNotifyEvent& NotifyEvent : Sampler.GetAllAnimNotifyEvents())
+		{
+			if (NotifyEvent.NotifyStateClass->GetClass()->IsChildOf<UAnimNotifyState_PoseSearchWarping>())
 			{
-				// using UAnimNotifyState_PoseSearchWarping to change the WarpingAlignmentExtractionTime (time we sample the full alignment transforms at)
-				if (const UAnimationAsset* AnimationAsset = MultiAnimAsset->GetAnimationAsset(CurrentResult.Role))
+				const UAnimNotifyState_PoseSearchWarping* NotifyStatePoseSearchWarping = Cast<UAnimNotifyState_PoseSearchWarping>(NotifyEvent.NotifyStateClass);
+				check(NotifyStatePoseSearchWarping);
+
+				WarpingAlignmentExtractionTime = FMath::Max(NotifyEvent.GetTime(), AccumulatedTime);
+				break;
+			}
+		}
+	}
+	
+	TSet<FName, DefaultKeyFuncs<FName>, TInlineSetAllocator<16, TMemStackSetAllocator<>>> GoalNames;
+	if (bUsePoseSearchIKWindowNotifyStates)
+	{
+		Sampler.ExtractPoseSearchNotifyStates(AccumulatedTime, [&GoalNames](const UAnimNotifyState_PoseSearchBase* PoseSearchBase)
+			{
+				if (const UAnimNotifyState_PoseSearchIKWindow* PoseSearchIKWindow = Cast<const UAnimNotifyState_PoseSearchIKWindow>(PoseSearchBase))
 				{
-					const FAnimationAssetSampler SequenceBaseSampler(AnimationAsset, FTransform::Identity, CurrentResult.BlendParameters, FAnimationAssetSampler::DefaultRootTransformSamplingRate, false);
-					for (const FAnimNotifyEvent& NotifyEvent : SequenceBaseSampler.GetAllAnimNotifyEvents())
+					for (FName GoalName : PoseSearchIKWindow->GoalNames)
 					{
-						if (NotifyEvent.NotifyStateClass->GetClass()->IsChildOf<UAnimNotifyState_PoseSearchWarping>())
-						{
-							const UAnimNotifyState_PoseSearchWarping* NotifyStatePoseSearchWarping = Cast<UAnimNotifyState_PoseSearchWarping>(NotifyEvent.NotifyStateClass);
-							check(NotifyStatePoseSearchWarping);
-
-							WarpingAlignmentExtractionTime = FMath::Max(NotifyEvent.GetTime(), AccumulatedTime);
-							break;
-						}
+						GoalNames.Add(GoalName);
 					}
 				}
-			}
+				// continue for other notify states
+				return true;
+			});
+	}
 
-			int32 CurrentResultRoleIndex = INDEX_NONE;
-			for (int32 MultiAnimAssetRoleIndex = 0; MultiAnimAssetRoleIndex < MultiAnimAsset->GetNumRoles(); ++MultiAnimAssetRoleIndex)
-			{
-				if (MultiAnimAsset->GetRole(MultiAnimAssetRoleIndex) == CurrentResult.Role)
-				{
-					CurrentResultRoleIndex = MultiAnimAssetRoleIndex;
-					break;
-				}
-			}
+	const bool bNeedCalculatingWarping = WarpingTranslationRatio > UE_KINDA_SMALL_NUMBER || WarpingRotationRatio > UE_KINDA_SMALL_NUMBER;
+	if (!bNeedCalculatingWarping && GoalNames.IsEmpty())
+	{
+		// no warping nor IK necessary. we're done
+		return;
+	}
 
-			check(InteractionSearchContext.CheckForConsistency());
+	// calculating the world space warp tranform to have the character in its "original" no reference frame (where no warping associated to MultiAnimAsset is impacting its transform)
+	FTransform RootMotionTransform;
+	RootMotionProvider->ExtractRootMotion(Output.CustomAttributes, RootMotionTransform);
 
-			TArray<FTransform, TInlineAllocator<PreallocatedRolesNum>> ActorTransforms;
-			ActorTransforms.SetNum(MultiAnimAsset->GetNumRoles());
-			// @todo: predicting ActorTransforms TimeToFullAlignment in the future! Maybe by using the trajectories..
-			if (CurrentResultRoleIndex != INDEX_NONE && InteractionSearchContext.GetActorTransforms(ActorTransforms, MultiAnimAsset))
-			{
-				// calculating the world space warp tranform to have the character in its "original" no reference frame (where no warping associated to MultiAnimAsset is impacting its transform)
-				FTransform RootMotionTransform;
-				RootMotionProvider->ExtractRootMotion(Output.CustomAttributes, RootMotionTransform);
-				// NoTe that ActorTransforms[CurrentResultRoleIndex] is Output.AnimInstanceProxy->GetSkelMeshComponent()->GetComponentTransform();
-				const FTransform AlignmentTransformNoWarpAdjustment = RootMotionTransform * ActorTransforms[CurrentResultRoleIndex];
+	// NoTe that ActorRootTransforms[CurrentResultRoleIndex] is Output.AnimInstanceProxy->GetSkelMeshComponent()->GetComponentTransform();
+	const FTransform AlignmentTransformNoWarpAdjustment = RootMotionTransform * ActorRootTransforms[CurrentResultRoleIndex];
 					
-				// calculating the world space warp tranform to have the character fully aligned to what MultiAnimAsset wants
-				TArray<FTransform, TInlineAllocator<PreallocatedRolesNum>> FullAlignedActorRootBoneTransforms;
-				FullAlignedActorRootBoneTransforms.SetNum(MultiAnimAsset->GetNumRoles());
-				MultiAnimAsset->CalculateWarpTransforms(WarpingAlignmentExtractionTime, ActorTransforms, FullAlignedActorRootBoneTransforms);
-				const FTransform AlignmentTransformFullWarpAdjustment = FullAlignedActorRootBoneTransforms[CurrentResultRoleIndex];
+	// calculating the world space warp tranform to have the character fully aligned to what MultiAnimAsset wants.
+	// using ActorTransforms instead of ActorRootTransforms, since we want the previous frame root bone offset rotation to be embedded into the calculation
+	TArray<FTransform, TInlineAllocator<PreallocatedRolesNum, TMemStackAllocator<>>> FullAlignedActorRootBoneTransformsAtExtractionTime;
+	FullAlignedActorRootBoneTransformsAtExtractionTime.SetNum(NumRoles);
+	MultiAnimAsset->CalculateWarpTransforms(WarpingAlignmentExtractionTime, ActorTransforms, FullAlignedActorRootBoneTransformsAtExtractionTime);
+	const FTransform AlignmentTransformFullWarpAdjustment = FullAlignedActorRootBoneTransformsAtExtractionTime[CurrentResultRoleIndex];
 					
-				// @todo: compensating AlignmentTransformFullWarpAdjustment by the delta time between AccumulatedTime and WarpingAlignmentExtractionTime?				
+	// @todo: compensating AlignmentTransformFullWarpAdjustment by the delta time between AccumulatedTime and WarpingAlignmentExtractionTime?				
+	// @todo: optimize me! we can skip some math if (!bNeedCalculatingWarping)
+	
+	// calculating the world space blended warp aligment transform
+	const FTransform AlignmentTransform(
+		FMath::Lerp(AlignmentTransformNoWarpAdjustment.GetRotation(), AlignmentTransformFullWarpAdjustment.GetRotation(), WarpingRotationRatio),
+		FMath::Lerp(AlignmentTransformNoWarpAdjustment.GetTranslation(), AlignmentTransformFullWarpAdjustment.GetTranslation(), WarpingTranslationRatio),
+		FVector::OneVector);
+
+	// calculating WarpRootMotionTransform as AlignmentTransform in root bone space
+	const FTransform RootBoneTransform = Output.AnimInstanceProxy->GetSkelMeshComponent()->GetBoneTransform(RootBoneIndexType);
+	FTransform WarpRootMotionTransform = AlignmentTransform.GetRelativeTransform(RootBoneTransform);
+	check(WarpRootMotionTransform.IsRotationNormalized());
+
+	// apply some damping to WarpRootMotionTransform in root bone space
+	const float DampenWarpingTranslation = WarpingTranslationHalflife > 0.f ? DamperImplicit(WarpingTranslationHalflife, CachedDeltaTime) : 1.f;
+	const float DampenWarpingRotation = WarpingRotationHalfLife > 0.f ? DamperImplicit(WarpingRotationHalfLife, CachedDeltaTime) : 1.f;
+	WarpRootMotionTransform.SetTranslation(WarpRootMotionTransform.GetTranslation() * DampenWarpingTranslation);
+	WarpRootMotionTransform.SetRotation(FQuat::Slerp(FQuat::Identity, WarpRootMotionTransform.GetRotation(), DampenWarpingRotation));
 
-				// calculating the world space blended warp aligment transform
-				const FTransform AlignmentTransform(
-					FMath::Lerp(AlignmentTransformNoWarpAdjustment.GetRotation(), AlignmentTransformFullWarpAdjustment.GetRotation(), WarpingRotationRatio),
-					FMath::Lerp(AlignmentTransformNoWarpAdjustment.GetTranslation(), AlignmentTransformFullWarpAdjustment.GetTranslation(), WarpingTranslationRatio),
-					FVector::OneVector);
+#if ENABLE_VISUAL_LOG
+	if (FVisualLogger::IsRecording())
+	{
+		static const TCHAR* LogName = TEXT("MotionMatchingInteraction");
 
-				// calculating WarpRootMotionTransform as AlignmentTransform in root bone space
-				const FTransform RootBoneTransform = Output.AnimInstanceProxy->GetSkelMeshComponent()->GetBoneTransform(RootBoneIndexType);
-				FTransform WarpRootMotionTransform = AlignmentTransform.GetRelativeTransform(RootBoneTransform);
-				check(WarpRootMotionTransform.IsRotationNormalized());
+		const FTransform WorldWarpRootMotionTransform = WarpRootMotionTransform * RootBoneTransform;
+		UE_VLOG_ARROW(AnimInstance, LogName, Display, WorldWarpRootMotionTransform.GetLocation(), WorldWarpRootMotionTransform.GetLocation() + WorldWarpRootMotionTransform.GetRotation().GetForwardVector() * 100, FColor::Green, TEXT(""));
+		UE_VLOG_ARROW(AnimInstance, LogName, Display, AlignmentTransform.GetLocation(), AlignmentTransform.GetLocation() + AlignmentTransform.GetRotation().GetForwardVector() * 100, FColor::Blue, TEXT(""));
+		UE_VLOG_ARROW(AnimInstance, LogName, Display, RootBoneTransform.GetLocation(), RootBoneTransform.GetLocation() + RootBoneTransform.GetRotation().GetForwardVector() * 100, FColor::Orange, TEXT(""));
+	}
+#endif
+	RootMotionProvider->OverrideRootMotion(WarpRootMotionTransform, Output.CustomAttributes);
 
-				// apply some damping to WarpRootMotionTransform in root bone space
-				const float DampenWarpingTranslation = WarpingTranslationHalflife > 0.f ? DamperImplicit(WarpingTranslationHalflife, CachedDeltaTime) : 1.f;
-				const float DampenWarpingRotation = WarpingRotationHalfLife > 0.f ? DamperImplicit(WarpingRotationHalfLife, CachedDeltaTime) : 1.f;
-				WarpRootMotionTransform.SetTranslation(WarpRootMotionTransform.GetTranslation() * DampenWarpingTranslation);
-				WarpRootMotionTransform.SetRotation(FQuat::Slerp(FQuat::Identity, WarpRootMotionTransform.GetRotation(), DampenWarpingRotation));
+	if (GoalNames.IsEmpty())
+	{
+		return;
+	}
 
+	FTransform FullAlignedPoseRootTransform;
+	if (FMath::IsNearlyEqual(WarpingAlignmentExtractionTime, AccumulatedTime))
+	{
+		FullAlignedPoseRootTransform = FullAlignedActorRootBoneTransformsAtExtractionTime[CurrentResultRoleIndex];
+	}
+	else
+	{
+		// if the we extracted the full alignment transforms NOT at AccumulatedTime (usually to make aligment calculation more stable - closer to the interaction time)
+		// we need to reextract them. This could happen only if users want to align the IK before aligning the character, things that I think should never happen, 
+		// but just in case the code supports it
+		// using ActorTransforms instead of ActorRootTransforms, since we want the previous frame root bone offset rotation to be embedded into the calculation
+		TArray<FTransform, TInlineAllocator<PreallocatedRolesNum, TMemStackAllocator<>>> FullAlignedActorRootBoneTransformsAtAccumulatedTime;
+		FullAlignedActorRootBoneTransformsAtAccumulatedTime.SetNum(NumRoles);
+		MultiAnimAsset->CalculateWarpTransforms(AccumulatedTime, ActorTransforms, FullAlignedActorRootBoneTransformsAtAccumulatedTime);
+		FullAlignedPoseRootTransform = FullAlignedActorRootBoneTransformsAtAccumulatedTime[CurrentResultRoleIndex];
+	}
+								
+	// solving the IK constraints for this character, knowing where its FullAlignedPose should be
+	const FBoneContainer& BoneContainer = AnimInstance->GetRequiredBonesOnAnyThread();
+	const USkeleton* Skeleton = BoneContainer.GetSkeletonAsset();
+	check(Skeleton);
+
+	FCompactPose FullAlignedPose;
+	FullAlignedPose.SetBoneContainer(&BoneContainer);
+	Sampler.ExtractPose(AccumulatedTime, FullAlignedPose);
+	FCSPose<FCompactPose> ComponentSpaceFullAlignedPose;
+	ComponentSpaceFullAlignedPose.InitPose(FullAlignedPose);
+
+	FCSPose<FCompactPose> ComponentSpacePose;
+	ComponentSpacePose.InitPose(Output.Pose);
+	const FTransform WorldWarpRootMotionTransform = WarpRootMotionTransform * RootBoneTransform;
+
+	for (FName GoalName : GoalNames)
+	{
+		FBoneReference BoneReference;
+		BoneReference.BoneName = GoalName;
+		BoneReference.Initialize(Skeleton);
+		if (!BoneReference.HasValidSetup())
+		{
+			UE_LOG(LogPoseSearch, Error, TEXT("FAnimNode_MotionMatchingInteraction::Evaluate_AnyThread invalid GoalName %s for Skeleton %s"), *GoalName.ToString(), *GetNameSafe(Skeleton));
+		}
+		else
+		{
+			const FCompactPoseBoneIndex CompactBoneIdx = BoneContainer.GetCompactPoseIndexFromSkeletonPoseIndex(FSkeletonPoseBoneIndex(BoneReference.BoneIndex));
+			if (!CompactBoneIdx.IsValid())
+			{
+				UE_LOG(LogPoseSearch, Error, TEXT("FAnimNode_MotionMatchingInteraction::Evaluate_AnyThread invalid CompactPoseBoneIndex for GoalName %s for Skeleton %s with BoneIndex %d"), *GoalName.ToString(), *GetNameSafe(Skeleton), BoneReference.BoneIndex);
+			}
+			else
+			{
+				const FTransform GoalWorldTransform = ComponentSpaceFullAlignedPose.GetComponentSpaceTransform(CompactBoneIdx) * FullAlignedPoseRootTransform;
+				const FCompactPoseBoneIndex CompactParentBoneIdx = BoneContainer.GetParentBoneIndex(CompactBoneIdx);
+				const FTransform WorldPoseParentBone = ComponentSpacePose.GetComponentSpaceTransform(CompactParentBoneIdx) * WorldWarpRootMotionTransform;
+				const FTransform GoalLocalTransform = GoalWorldTransform.GetRelativeTransform(WorldPoseParentBone);
+				Output.Pose[CompactBoneIdx] = GoalLocalTransform;
 #if ENABLE_VISUAL_LOG
 				if (FVisualLogger::IsRecording())
 				{
 					static const TCHAR* LogName = TEXT("MotionMatchingInteraction");
-
-					UObject* LogOwner = Output.AnimInstanceProxy->GetAnimInstanceObject();
-					const FTransform WorldWarpRootMotionTransform = WarpRootMotionTransform * RootBoneTransform;
-					UE_VLOG_ARROW(LogOwner, LogName, Display, WorldWarpRootMotionTransform.GetLocation(), WorldWarpRootMotionTransform.GetLocation() + WorldWarpRootMotionTransform.GetRotation().GetForwardVector() * 100, FColor::Green, TEXT(""));
-					UE_VLOG_ARROW(LogOwner, LogName, Display, AlignmentTransform.GetLocation(), AlignmentTransform.GetLocation() + AlignmentTransform.GetRotation().GetForwardVector() * 100, FColor::Blue, TEXT(""));
-					UE_VLOG_ARROW(LogOwner, LogName, Display, RootBoneTransform.GetLocation(), RootBoneTransform.GetLocation() + RootBoneTransform.GetRotation().GetForwardVector() * 100, FColor::Orange, TEXT(""));
+					UE_VLOG_ARROW(AnimInstance, LogName, Display, FullAlignedPoseRootTransform.GetLocation(), GoalWorldTransform.GetLocation(), FColor::Black, TEXT(""));
+					UE_VLOG_ARROW(AnimInstance, LogName, Display, ActorRootTransforms[CurrentResultRoleIndex].GetLocation(), WorldPoseParentBone.GetLocation(), FColor::Red, TEXT(""));
+					const FTransform WorldPoseBone = ComponentSpacePose.GetComponentSpaceTransform(CompactBoneIdx) * WorldWarpRootMotionTransform;
+					UE_VLOG_ARROW(AnimInstance, LogName, Display, WorldPoseParentBone.GetLocation(), WorldPoseBone.GetLocation(), FColor::Green, TEXT(""));
 				}
 #endif
-				RootMotionProvider->OverrideRootMotion(WarpRootMotionTransform, Output.CustomAttributes);
 			}
 		}
 	}
+
+#if ENABLE_DRAW_DEBUG && ENABLE_ANIM_DEBUG
+	if (GVarAnimNodeMotionMatchingInteractionDrawFullAlignedPose)
+	{
+		UPoseSearchAssetSamplerLibrary::Draw(AnimInstance->GetWorld(), ComponentSpaceFullAlignedPose, FullAlignedPoseRootTransform);
+	}
+#endif // ENABLE_DRAW_DEBUG && ENABLE_ANIM_DEBUG
+
 }
diff --git a/Engine/Plugins/Animation/PoseSearch/Source/Runtime/Private/AnimNode_PoseSearchHistoryCollector.cpp b/Engine/Plugins/Animation/PoseSearch/Source/Runtime/Private/AnimNode_PoseSearchHistoryCollector.cpp
--- a/Engine/Plugins/Animation/PoseSearch/Source/Runtime/Private/AnimNode_PoseSearchHistoryCollector.cpp
+++ b/Engine/Plugins/Animation/PoseSearch/Source/Runtime/Private/AnimNode_PoseSearchHistoryCollector.cpp
@@ -5,6 +5,7 @@
 #include "Animation/AnimNodeMessages.h"
 #include "Animation/AnimStats.h"
 #include "Engine/SkinnedAsset.h"
+#include "PoseSearch/PoseSearchInteractionValidator.h"
 
 #include UE_INLINE_GENERATED_CPP_BY_NAME(AnimNode_PoseSearchHistoryCollector)
 
@@ -112,6 +113,9 @@ void FAnimNode_PoseSearchHistoryCollector::Evaluate_AnyThread(FPoseContext& Outp
 	DECLARE_SCOPE_HIERARCHICAL_COUNTER_ANIMNODE(Evaluate_AnyThread);
 	ANIM_MT_SCOPE_CYCLE_COUNTER_VERBOSE(PoseSearchHistoryCollector, !IsInGameThread());
 
+	using namespace UE::PoseSearch;
+
+	CheckInteractionThreadSafety(Output.AnimInstanceProxy->GetAnimInstanceObject(), FInteractionValidator::EPhase::HistoryCollectorEvaluate);
 	check(Output.AnimInstanceProxy);
 
 	Super::Evaluate_AnyThread(Output);
@@ -151,6 +155,10 @@ void FAnimNode_PoseSearchHistoryCollector::Update_AnyThread(const FAnimationUpda
 {
 	DECLARE_SCOPE_HIERARCHICAL_COUNTER_ANIMNODE(Update_AnyThread);
 
+	using namespace UE::PoseSearch;
+
+	CheckInteractionThreadSafety(Context.AnimInstanceProxy->GetAnimInstanceObject(), FInteractionValidator::EPhase::HistoryCollectorUpdate);
+
 	if (bGenerateTrajectory)
 	{
 		GenerateTrajectory(Cast<UAnimInstance>(Context.AnimInstanceProxy->GetAnimInstanceObject()));
@@ -190,6 +198,10 @@ void FAnimNode_PoseSearchComponentSpaceHistoryCollector::EvaluateComponentSpace_
 	DECLARE_SCOPE_HIERARCHICAL_COUNTER_ANIMNODE(EvaluateComponentSpace_AnyThread);
 	ANIM_MT_SCOPE_CYCLE_COUNTER_VERBOSE(PoseSearchComponentSpaceHistoryCollector, !IsInGameThread());
 
+	using namespace UE::PoseSearch;
+
+	CheckInteractionThreadSafety(Output.AnimInstanceProxy->GetAnimInstanceObject(), FInteractionValidator::EPhase::HistoryCollectorEvaluate);
+
 	check(Output.AnimInstanceProxy);
 
 	Super::EvaluateComponentSpace_AnyThread(Output);
@@ -226,6 +238,10 @@ void FAnimNode_PoseSearchComponentSpaceHistoryCollector::Update_AnyThread(const
 {
 	DECLARE_SCOPE_HIERARCHICAL_COUNTER_ANIMNODE(Update_AnyThread);
 
+	using namespace UE::PoseSearch;
+
+	CheckInteractionThreadSafety(Context.AnimInstanceProxy->GetAnimInstanceObject(), FInteractionValidator::EPhase::HistoryCollectorUpdate);
+
 	if (bGenerateTrajectory)
 	{
 		GenerateTrajectory(Cast<UAnimInstance>(Context.AnimInstanceProxy->GetAnimInstanceObject()));
diff --git a/Engine/Plugins/Animation/PoseSearch/Source/Runtime/Private/PoseSearchAssetSamplerLibrary.cpp b/Engine/Plugins/Animation/PoseSearch/Source/Runtime/Private/PoseSearchAssetSamplerLibrary.cpp
--- a/Engine/Plugins/Animation/PoseSearch/Source/Runtime/Private/PoseSearchAssetSamplerLibrary.cpp
+++ b/Engine/Plugins/Animation/PoseSearch/Source/Runtime/Private/PoseSearchAssetSamplerLibrary.cpp
@@ -51,45 +51,9 @@ FPoseSearchAssetSamplerPose UPoseSearchAssetSamplerLibrary::SamplePose(const UAn
 	return AssetSamplerPose;
 }
 
-FTransform UPoseSearchAssetSamplerLibrary::GetTransform(UPARAM(ref) FPoseSearchAssetSamplerPose& AssetSamplerPose, int32 BoneIndex, EPoseSearchAssetSamplerSpace Space)
+FTransform UPoseSearchAssetSamplerLibrary::GetTransform(FPoseSearchAssetSamplerPose& AssetSamplerPose, FCompactPoseBoneIndex CompactPoseBoneIndex, EPoseSearchAssetSamplerSpace Space)
 {
-	const FCompactPoseBoneIndex CompactPoseBoneIndex(BoneIndex);
-
-	if (!AssetSamplerPose.Pose.IsValid())
-	{
-		UE_LOG(LogPoseSearch, Error, TEXT("UPoseSearchAssetSamplerLibrary::GetBoneTransform invalid AssetSamplerPose.Pose"));
-		return FTransform::Identity;
-	}
-
-	if (BoneIndex == INDEX_NONE)
-	{
-		if (Space != EPoseSearchAssetSamplerSpace::World)
-		{
-			UE_LOG(LogPoseSearch, Error, TEXT("UPoseSearchAssetSamplerLibrary::GetBoneTransform invalid Space %s to get the RootTransform. Expected space: %s"), *UEnum::GetDisplayValueAsText(Space).ToString(), *UEnum::GetDisplayValueAsText(EPoseSearchAssetSamplerSpace::World).ToString());
-		}
-		return AssetSamplerPose.RootTransform;
-	}
-	
-	if (!AssetSamplerPose.Pose.IsValidIndex(CompactPoseBoneIndex))
-	{
-		UE_LOG(LogPoseSearch, Error, TEXT("UPoseSearchAssetSamplerLibrary::GetBoneTransform invalid BoneIndex %d"), BoneIndex);
-		return FTransform::Identity;
-	}
-	
-	switch (Space)
-	{
-	case EPoseSearchAssetSamplerSpace::Local:
-		return AssetSamplerPose.Pose[CompactPoseBoneIndex];
-
-	case EPoseSearchAssetSamplerSpace::Component:
-		return AssetSamplerPose.ComponentSpacePose.GetComponentSpaceTransform(CompactPoseBoneIndex);
-
-	case EPoseSearchAssetSamplerSpace::World:
-		return AssetSamplerPose.ComponentSpacePose.GetComponentSpaceTransform(CompactPoseBoneIndex) * AssetSamplerPose.RootTransform;
-	}
-
-	checkNoEntry();
-	return FTransform::Identity;
+	return GetTransform(AssetSamplerPose.ComponentSpacePose, AssetSamplerPose.RootTransform, CompactPoseBoneIndex, Space);
 }
 
 FTransform UPoseSearchAssetSamplerLibrary::GetTransformByName(UPARAM(ref) FPoseSearchAssetSamplerPose& AssetSamplerPose, FName BoneName, EPoseSearchAssetSamplerSpace Space)
@@ -100,7 +64,8 @@ FTransform UPoseSearchAssetSamplerLibrary::GetTransformByName(UPARAM(ref) FPoseS
 		return FTransform::Identity;
 	}
 
-	const USkeleton* Skeleton = AssetSamplerPose.Pose.GetBoneContainer().GetSkeletonAsset();
+	const FBoneContainer& BoneContainer = AssetSamplerPose.Pose.GetBoneContainer();
+	const USkeleton* Skeleton = BoneContainer.GetSkeletonAsset();
 
 	FBoneReference BoneReference;
 	BoneReference.BoneName = BoneName;
@@ -111,32 +76,27 @@ FTransform UPoseSearchAssetSamplerLibrary::GetTransformByName(UPARAM(ref) FPoseS
 		return FTransform::Identity;
 	}
 
-	return GetTransform(AssetSamplerPose, BoneReference.BoneIndex, Space);
+	const FCompactPoseBoneIndex CompactPoseBoneIndex = BoneContainer.GetCompactPoseIndexFromSkeletonPoseIndex(FSkeletonPoseBoneIndex(BoneReference.BoneIndex));
+	if (CompactPoseBoneIndex.IsValid())
+	{
+		UE_LOG(LogPoseSearch, Error, TEXT("UPoseSearchAssetSamplerLibrary::GetTransformByName invalid FCompactPoseBoneIndex for BoneName %s for Skeleton %s"), *BoneName.ToString(), *GetNameSafe(Skeleton));
+		return FTransform::Identity;
+	}
+
+	return GetTransform(AssetSamplerPose, CompactPoseBoneIndex, Space);
 }
 
 void UPoseSearchAssetSamplerLibrary::Draw(const UAnimInstance* AnimInstance, UPARAM(ref) FPoseSearchAssetSamplerPose& AssetSamplerPose)
 {
+#if ENABLE_DRAW_DEBUG
 	static float DebugDrawSamplerRootAxisLength = 20.f;
 	static float DebugDrawSamplerSize = 6.f;
 
 	if (AnimInstance)
 	{
-		if (UWorld* World = AnimInstance->GetWorld())
-		{
-			if (DebugDrawSamplerRootAxisLength > 0.f)
-			{
-				const FTransform RootTransform = GetTransform(AssetSamplerPose);
-				DrawDebugLine(World, RootTransform.GetTranslation(), RootTransform.GetTranslation() + RootTransform.GetScaledAxis(EAxis::X) * DebugDrawSamplerRootAxisLength, FColor::Red, false, 0.f, SDPG_Foreground);
-				DrawDebugLine(World, RootTransform.GetTranslation(), RootTransform.GetTranslation() + RootTransform.GetScaledAxis(EAxis::Y) * DebugDrawSamplerRootAxisLength, FColor::Green, false, 0.f, SDPG_Foreground);
-				DrawDebugLine(World, RootTransform.GetTranslation(), RootTransform.GetTranslation() + RootTransform.GetScaledAxis(EAxis::Z) * DebugDrawSamplerRootAxisLength, FColor::Blue, false, 0.f, SDPG_Foreground);
-			}
-
-			for (int32 BoneIndex = 0; BoneIndex < AssetSamplerPose.ComponentSpacePose.GetPose().GetNumBones(); ++BoneIndex)
-			{
-				const FTransform BoneWorldTransforms = GetTransform(AssetSamplerPose, BoneIndex, EPoseSearchAssetSamplerSpace::World);
-				DrawDebugPoint(World, BoneWorldTransforms.GetTranslation(), DebugDrawSamplerSize, FColor::Red, false, 0.f, SDPG_Foreground);
-			}
-		}
+		Draw(AnimInstance->GetWorld(), AssetSamplerPose.ComponentSpacePose, AssetSamplerPose.RootTransform);
 	}
+#endif // ENABLE_DRAW_DEBUG
 }
 
+
diff --git a/Engine/Plugins/Animation/PoseSearch/Source/Runtime/Private/PoseSearchHistory.cpp b/Engine/Plugins/Animation/PoseSearch/Source/Runtime/Private/PoseSearchHistory.cpp
--- a/Engine/Plugins/Animation/PoseSearch/Source/Runtime/Private/PoseSearchHistory.cpp
+++ b/Engine/Plugins/Animation/PoseSearch/Source/Runtime/Private/PoseSearchHistory.cpp
@@ -48,8 +48,8 @@ static FAutoConsoleVariableRef CVarAnimPoseHistoryDebugDrawTrajectoryMaxNumOfPre
 *
 * @returns Position of the first element >= Value, may be position after last element in range
 */
-template <typename IteratorType, typename ValueType, typename ProjectionType, typename SortPredicateType>
-FORCEINLINE auto LowerBound(IteratorType First, IteratorType Last, const ValueType& Value, ProjectionType Projection, SortPredicateType SortPredicate) -> decltype(First.GetIndex())
+template <typename IteratorType, typename ValueType, typename SortPredicateType = TLess<>()>
+FORCEINLINE auto LowerBound(IteratorType First, IteratorType Last, const ValueType& Value, SortPredicateType SortPredicate) -> decltype(First.GetIndex())
 {
 	using SizeType = decltype(First.GetIndex());
 
@@ -69,19 +69,11 @@ FORCEINLINE auto LowerBound(IteratorType First, IteratorType Last, const ValueTy
 
 		const SizeType CheckIndex = Start + Size;
 		const SizeType StartIfLess = CheckIndex + LeftoverSize;
-
-		auto&& CheckValue = Invoke(Projection, *(First + CheckIndex));
-		Start = SortPredicate(CheckValue, Value) ? StartIfLess : Start;
+		Start = SortPredicate(*(First + CheckIndex), Value) ? StartIfLess : Start;
 	}
 	return Start;
 }
 
-template <typename IteratorType, typename ValueType, typename SortPredicateType = TLess<>()>
-FORCEINLINE auto LowerBound(IteratorType First, IteratorType Last, const ValueType& Value, SortPredicateType SortPredicate) -> decltype(First.GetIndex())
-{
-	return LowerBound(First, Last, Value, FIdentityFunctor(), SortPredicate);
-}
-
 static FBoneIndexType GetRemappedBoneIndexType(FBoneIndexType BoneIndexType, const USkeleton* BoneIndexSkeleton, const USkeleton* LastUpdateSkeleton)
 {
 	// remapping the skeleton bone index (encoded as BoneIndexType) in case the skeleton we used to store the history (LastUpdateSkeleton) is different from BoneIndexSkeleton
@@ -629,13 +621,9 @@ FPoseHistory::FPoseHistory(FPoseHistory&& Other)
 FPoseHistory& FPoseHistory::operator=(const FPoseHistory& Other)
 {
 #if ENABLE_ANIM_DEBUG
-	CheckThreadSafetyWrite(ReadPoseDataThreadSafeCounter);
-	CheckThreadSafetyWrite(WritePoseDataThreadSafeCounter);
-
-	FThreadSafeCounter& OtherReadPoseDataThreadSafeCounter = Other.ReadPoseDataThreadSafeCounter;
-	FThreadSafeCounter& OtherWritePoseDataThreadSafeCounter = Other.WritePoseDataThreadSafeCounter;
-	CheckThreadSafetyWrite(OtherReadPoseDataThreadSafeCounter);
-	CheckThreadSafetyWrite(OtherWritePoseDataThreadSafeCounter);
+	FThreadSafeCounter& OtherPoseDataThreadSafeCounter = Other.PoseDataThreadSafeCounter;
+	CheckThreadSafetyWrite(OtherPoseDataThreadSafeCounter);
+	CheckThreadSafetyWrite(PoseDataThreadSafeCounter);
 #endif // ENABLE_ANIM_DEBUG
 
 	MaxNumPoses = Other.MaxNumPoses;
@@ -652,13 +640,9 @@ FPoseHistory& FPoseHistory::operator=(const FPoseHistory& Other)
 FPoseHistory& FPoseHistory::operator=(FPoseHistory&& Other)
 {
 #if ENABLE_ANIM_DEBUG
-	CheckThreadSafetyWrite(ReadPoseDataThreadSafeCounter);
-	CheckThreadSafetyWrite(WritePoseDataThreadSafeCounter);
-
-	FThreadSafeCounter& OtherReadPoseDataThreadSafeCounter = Other.ReadPoseDataThreadSafeCounter;
-	FThreadSafeCounter& OtherWritePoseDataThreadSafeCounter = Other.WritePoseDataThreadSafeCounter;
-	CheckThreadSafetyWrite(OtherReadPoseDataThreadSafeCounter);
-	CheckThreadSafetyWrite(OtherWritePoseDataThreadSafeCounter);
+	FThreadSafeCounter& OtherPoseDataThreadSafeCounter = Other.PoseDataThreadSafeCounter;
+	CheckThreadSafetyWrite(OtherPoseDataThreadSafeCounter);
+	CheckThreadSafetyWrite(PoseDataThreadSafeCounter);
 #endif // ENABLE_ANIM_DEBUG
 
 	MaxNumPoses = MoveTemp(Other.MaxNumPoses);
@@ -674,7 +658,7 @@ FPoseHistory& FPoseHistory::operator=(FPoseHistory&& Other)
 
 void FPoseHistory::Initialize_AnyThread(int32 InNumPoses, float InSamplingInterval)
 {
-	CheckThreadSafetyWrite(WritePoseDataThreadSafeCounter);
+	CheckThreadSafetyWrite(PoseDataThreadSafeCounter);
 	check(InNumPoses >= 2);
 
 	MaxNumPoses = InNumPoses;
@@ -689,7 +673,7 @@ void FPoseHistory::Initialize_AnyThread(int32 InNumPoses, float InSamplingInterv
 
 bool FPoseHistory::GetTransformAtTime(float Time, FTransform& OutBoneTransform, const USkeleton* BoneIndexSkeleton, FBoneIndexType BoneIndexType, FBoneIndexType ReferenceBoneIndexType, bool bExtrapolate) const
 {
-	CheckThreadSafetyRead(ReadPoseDataThreadSafeCounter);
+	CheckThreadSafetyRead(PoseDataThreadSafeCounter);
 
 	static_assert(RootBoneIndexType == 0 && ComponentSpaceIndexType == FBoneIndexType(-1) && WorldSpaceIndexType == FBoneIndexType(-2)); // some assumptions
 	
@@ -747,7 +731,7 @@ bool FPoseHistory::GetTransformAtTime(float Time, FTransform& OutBoneTransform,
 
 bool FPoseHistory::GetCurveValueAtTime(float Time, const FName& CurveName, float& OutCurveValue, bool bExtrapolate /*= false*/) const
 {
-	CheckThreadSafetyRead(ReadPoseDataThreadSafeCounter);
+	CheckThreadSafetyRead(PoseDataThreadSafeCounter);
 
 	const int32 NumEntries = PoseData.Entries.Num();
 	if (NumEntries > 0)
@@ -774,43 +758,43 @@ bool FPoseHistory::GetCurveValueAtTime(float Time, const FName& CurveName, float
 
 const FPoseSearchQueryTrajectory& FPoseHistory::GetTrajectory() const
 {
-	CheckThreadSafetyRead(ReadPoseDataThreadSafeCounter);
+	CheckThreadSafetyRead(PoseDataThreadSafeCounter);
 	return Trajectory;
 }
 
 float FPoseHistory::GetTrajectorySpeedMultiplier() const
 {
-	CheckThreadSafetyRead(ReadPoseDataThreadSafeCounter);
+	CheckThreadSafetyRead(PoseDataThreadSafeCounter);
 	return TrajectorySpeedMultiplier;
 }
 
 bool FPoseHistory::IsEmpty() const
 {
-	CheckThreadSafetyRead(ReadPoseDataThreadSafeCounter);
+	CheckThreadSafetyRead(PoseDataThreadSafeCounter);
 	return PoseData.Entries.IsEmpty();
 }
 
 const FBoneToTransformMap& FPoseHistory::GetBoneToTransformMap() const
 {
-	CheckThreadSafetyRead(ReadPoseDataThreadSafeCounter);
+	CheckThreadSafetyRead(PoseDataThreadSafeCounter);
 	return PoseData.BoneToTransformMap;
 }
 
 const TConstArrayView<FName> FPoseHistory::GetCollectedCurves() const
 {
-	CheckThreadSafetyRead(ReadPoseDataThreadSafeCounter);
+	CheckThreadSafetyRead(PoseDataThreadSafeCounter);
 	return MakeConstArrayView(PoseData.CollectedCurves);
 }
 
 int32 FPoseHistory::GetNumEntries() const
 {
-	CheckThreadSafetyRead(ReadPoseDataThreadSafeCounter);
+	CheckThreadSafetyRead(PoseDataThreadSafeCounter);
 	return PoseData.Entries.Num();
 }
 
 const FPoseHistoryEntry& FPoseHistory::GetEntry(int32 EntryIndex) const
 {
-	CheckThreadSafetyRead(ReadPoseDataThreadSafeCounter);
+	CheckThreadSafetyRead(PoseDataThreadSafeCounter);
 	return PoseData.Entries[EntryIndex];
 }
 
@@ -833,7 +817,7 @@ void FPoseHistory::SetTrajectory(const FPoseSearchQueryTrajectory& InTrajectory,
 {
 	if (!InTrajectory.Samples.IsEmpty())
 	{
-		CheckThreadSafetyWrite(ReadPoseDataThreadSafeCounter);
+		CheckThreadSafetyWrite(PoseDataThreadSafeCounter);
 
 		// @todo: THIS IS NOT THREAD SAFE! in the contex of multi character motion matching (CheckThreadSafetyWrite will assert in case of improper usage)
 		Trajectory = InTrajectory;
@@ -856,8 +840,7 @@ void FPoseHistory::EvaluateComponentSpace_AnyThread(float DeltaTime, IComponentS
 	bool bNeedsReset, bool bCacheBones, const TArray<FBoneIndexType>& RequiredBones,
 	const FBlendedCurve& Curves, const TConstArrayView<FName>& CollectedCurves)
 {
-	CheckThreadSafetyWrite(WritePoseDataThreadSafeCounter);
-	CheckThreadSafetyRead(ReadPoseDataThreadSafeCounter);
+	CheckThreadSafetyWrite(PoseDataThreadSafeCounter);
 
 	check(MaxNumPoses >= 2);
 
@@ -969,7 +952,7 @@ void FPoseHistory::EvaluateComponentSpace_AnyThread(float DeltaTime, FCSPose<FCo
 #if ENABLE_DRAW_DEBUG && ENABLE_ANIM_DEBUG
 void FPoseHistory::DebugDraw(FAnimInstanceProxy& AnimInstanceProxy, FColor Color) const
 {
-	CheckThreadSafetyRead(ReadPoseDataThreadSafeCounter);
+	CheckThreadSafetyRead(PoseDataThreadSafeCounter);
 
 	if (GVarAnimPoseHistoryDebugDrawTrajectory)
 	{
diff --git a/Engine/Plugins/Animation/PoseSearch/Source/Runtime/Private/PoseSearchInteractionIsland.cpp b/Engine/Plugins/Animation/PoseSearch/Source/Runtime/Private/PoseSearchInteractionIsland.cpp
--- a/Engine/Plugins/Animation/PoseSearch/Source/Runtime/Private/PoseSearchInteractionIsland.cpp
+++ b/Engine/Plugins/Animation/PoseSearch/Source/Runtime/Private/PoseSearchInteractionIsland.cpp
@@ -26,11 +26,52 @@ static FAutoConsoleVariableRef CVarPoseSearchInteractionIslandResultPolicy(TEXT(
 
 // Utils functions
 ///////////////////////////////////////////////////////////
-static FInteractionSearchResult InitSearchResult(const FSearchResult& SearchResult, int32 SearchIndex)
+static FInteractionSearchResult InitSearchResult(const FSearchResult& SearchResult, const FInteractionSearchContext& SearchContext, int32 SearchIndex)
 {
 	FInteractionSearchResult InteractionSearchResult;
 	static_cast<FSearchResult&>(InteractionSearchResult) = SearchResult;
 	InteractionSearchResult.SearchIndex = SearchIndex;
+
+	const int32 AnimContextsNum = SearchContext.AnimContexts.Num();
+	InteractionSearchResult.ActorRootTransforms.SetNum(AnimContextsNum);
+	InteractionSearchResult.ActorRootBoneTransforms.SetNum(AnimContextsNum);
+
+	for (int32 AnimContextIndex = 0; AnimContextIndex < AnimContextsNum; ++AnimContextIndex)
+	{
+		if (const UObject* AnimContext = SearchContext.AnimContexts[AnimContextIndex].Get())
+		{
+			if (const UAnimInstance* AnimInstance = Cast<UAnimInstance>(AnimContext))
+			{
+				if (const FAnimNode_PoseSearchHistoryCollector_Base* HistoryCollector = SearchContext.HistoryCollectors[AnimContextIndex])
+				{
+					const USkeleton* Skeleton = AnimInstance->GetRequiredBonesOnAnyThread().GetSkeletonAsset();
+					check(Skeleton);
+					HistoryCollector->GetPoseHistory().GetTransformAtTime(0.f, InteractionSearchResult.ActorRootBoneTransforms[AnimContextIndex], Skeleton, RootBoneIndexType, ComponentSpaceIndexType);
+
+					// @todo: should we use the pose history ActorRootTransforms?
+					//HistoryCollector->GetPoseHistory().GetTransformAtTime(0.f, InteractionSearchResult.ActorRootTransforms[AnimContextIndex], Skeleton, ComponentSpaceIndexType, WorldSpaceIndexType);
+					InteractionSearchResult.ActorRootTransforms[AnimContextIndex] = AnimInstance->GetSkelMeshComponent()->GetComponentTransform();
+				}
+				else
+				{
+					InteractionSearchResult.ActorRootBoneTransforms[AnimContextIndex] = FTransform::Identity;
+					InteractionSearchResult.ActorRootTransforms[AnimContextIndex] = AnimInstance->GetSkelMeshComponent()->GetComponentTransform();
+				}
+			}
+			else
+			{
+				unimplemented();
+				InteractionSearchResult.ActorRootTransforms[AnimContextIndex] = FTransform::Identity;
+				InteractionSearchResult.ActorRootBoneTransforms[AnimContextIndex] = FTransform::Identity;
+			}
+		}
+		else
+		{
+			InteractionSearchResult.ActorRootTransforms[AnimContextIndex] = FTransform::Identity;
+			InteractionSearchResult.ActorRootBoneTransforms[AnimContextIndex] = FTransform::Identity;
+		}
+	}
+
 	return InteractionSearchResult;
 }
 
@@ -117,7 +158,7 @@ static void InitSearchResults(TArray<FInteractionSearchResult>& SearchResults, T
 			if (PoseSearchResults[BestSearchIndex].IsValid())
 			{
 				SearchResults.SetNum(1, EAllowShrinking::No);
-				SearchResults[0] = InitSearchResult(PoseSearchResults[BestSearchIndex], BestSearchIndex);
+				SearchResults[0] = InitSearchResult(PoseSearchResults[BestSearchIndex], SearchContexts[BestSearchIndex], BestSearchIndex);
 			}
 		}
 		else
@@ -128,7 +169,7 @@ static void InitSearchResults(TArray<FInteractionSearchResult>& SearchResults, T
 			{
 				if (IsPoseSearchResultUsable(SearchIndex, PoseSearchResults, SearchContexts, VisitedAnimContexts))
 				{
-					SearchResults.Add(InitSearchResult(PoseSearchResults[SearchIndex], SearchIndex));
+					SearchResults.Add(InitSearchResult(PoseSearchResults[SearchIndex], SearchContexts[SearchIndex], SearchIndex));
 					
 					for (const TWeakObjectPtr<const UObject>& SearchAnimContext : SearchContexts[SearchIndex].AnimContexts)
 					{
@@ -143,6 +184,42 @@ static void InitSearchResults(TArray<FInteractionSearchResult>& SearchResults, T
 
 // FInteractionSearchContext
 ///////////////////////////////////////////////////////////
+bool FInteractionSearchContext::IsEquivalent(const FInteractionSearchContext& Other) const
+{
+	if (Database != Other.Database)
+	{
+		return false;
+	}
+
+	const int32 Num = AnimContexts.Num();
+	if (Num != Other.AnimContexts.Num())
+	{
+		return false;
+	}
+
+	int32 CommonRoledAnimContexts = 0;
+	for (int32 IndexThis = 0; IndexThis < Num; ++IndexThis)
+	{
+		for (int32 IndexOther = 0; IndexOther < Num; ++IndexOther)
+		{
+			if (AnimContexts[IndexThis] == Other.AnimContexts[IndexOther] &&
+				HistoryCollectors[IndexThis] == Other.HistoryCollectors[IndexOther] &&
+				Roles[IndexThis] == Other.Roles[IndexOther])
+			{
+				++CommonRoledAnimContexts;
+			}
+		}
+	}
+
+	// using >= in case there are duplicated AnimContexts in this or Other (this->IsValid() should be false!)
+	if (CommonRoledAnimContexts >= Num)
+	{
+		return true;
+	}
+
+	return false;
+}
+
 #if DO_CHECK
 bool FInteractionSearchContext::CheckForConsistency() const
 {
@@ -212,84 +289,31 @@ bool FInteractionSearchContext::CheckForConsistency() const
 
 	return true;
 }
-#endif // DO_CHECK
 
-bool FInteractionSearchContext::IsEquivalent(const FInteractionSearchContext& Other) const
+void FInteractionSearchContext::TestHistoryCollectorsThreadingAccess() const
 {
-	if (Database != Other.Database)
-	{
-		return false;
-	}
-
-	const int32 Num = AnimContexts.Num();
-	if (Num != Other.AnimContexts.Num())
-	{
-		return false;
-	}
-
-	int32 CommonRoledAnimContexts = 0;
-	for (int32 IndexThis = 0; IndexThis < Num; ++IndexThis)
+	for (int32 Index = 0; Index < AnimContexts.Num(); ++Index)
 	{
-		for (int32 IndexOther = 0; IndexOther < Num; ++IndexOther)
+		const UAnimInstance* AnimInstance = Cast<UAnimInstance>(AnimContexts[Index].Get());
+		const FAnimNode_PoseSearchHistoryCollector_Base* HistoryCollector = HistoryCollectors[Index];
+		
+		if (AnimInstance && HistoryCollector)
 		{
-			if (AnimContexts[IndexThis] == Other.AnimContexts[IndexOther] &&
-				HistoryCollectors[IndexThis] == Other.HistoryCollectors[IndexOther] &&
-				Roles[IndexThis] == Other.Roles[IndexOther])
-			{
-				++CommonRoledAnimContexts;
-			}
-		}
-	}
-
-	// using >= in case there are duplicated AnimContexts in this or Other (this->IsValid() should be false!)
-	if (CommonRoledAnimContexts >= Num)
-	{
-		return true;
-	}
-
-	return false;
-}
-
-bool FInteractionSearchContext::GetActorTransforms(TArrayView<FTransform> ActorTransforms, const UMultiAnimAsset* MultiAnimAsset) const
-{
-	check(MultiAnimAsset);
-	check(ActorTransforms.Num() == MultiAnimAsset->GetNumRoles());
+			const USkeleton* Skeleton = AnimInstance->GetRequiredBonesOnAnyThread().GetSkeletonAsset();
+			check(Skeleton);
 
-	const FRoleToIndex InteractionSearchContextRoleToIndex = MakeRoleToIndex(Roles);
-
-	bool bAllRolesResolved = true;
-	for (int32 MultiAnimAssetRoleIndex = 0; MultiAnimAssetRoleIndex < MultiAnimAsset->GetNumRoles(); ++MultiAnimAssetRoleIndex)
-	{
-		const FRole& MultiAnimAssetRole = MultiAnimAsset->GetRole(MultiAnimAssetRoleIndex);
-
-		if (const int32* InteractionSearchContextRoleIndex = InteractionSearchContextRoleToIndex.Find(MultiAnimAssetRole))
-		{
-			if (const UAnimInstance* AnimInstance = Cast<UAnimInstance>(AnimContexts[*InteractionSearchContextRoleIndex].Get()))
-			{
-				ActorTransforms[MultiAnimAssetRoleIndex] = AnimInstance->GetSkelMeshComponent()->GetComponentTransform();
-			}
-			else
-			{
-				ActorTransforms[MultiAnimAssetRoleIndex] = FTransform::Identity;
-				bAllRolesResolved = false;
-				break;
-			}
-		}
-		else
-		{
-			ActorTransforms[MultiAnimAssetRoleIndex] = FTransform::Identity;
-			bAllRolesResolved = false;
-			break;
+			FTransform ActorRootBoneTransform;
+			HistoryCollector->GetPoseHistory().GetTransformAtTime(0.f, ActorRootBoneTransform, Skeleton, RootBoneIndexType, ComponentSpaceIndexType);
 		}
 	}
-	return bAllRolesResolved;
 }
+#endif // DO_CHECK
 
 // FInteractionSearchResult
 ///////////////////////////////////////////////////////////
 bool FInteractionSearchResult::operator==(const FInteractionSearchResult& Other) const
 {
-	// not checking for SearchIndex equality
+	// not checking SearchIndex, nor ActorRootTransforms, nor ActorRootBoneTransforms for equality
 	return static_cast<const FSearchResult&>(*this) == static_cast<const FSearchResult&>(Other);
 }
 
@@ -536,8 +560,53 @@ bool FInteractionIsland::IsUninjected() const
 	return PostTickComponents.IsEmpty();
 }
 
+const AActor* FInteractionIsland::GetMainActor() const
+{
+	if (!IsUninjected())
+	{
+		if (const UActorComponent* ActorComponent = PostTickComponents[0].Get())
+		{
+			return ActorComponent->GetOwner();
+		}
+	}
+	return nullptr;
+}
+
+bool FInteractionIsland::CanPerformSearch(const UObject* AnimContext) const
+{
+	check(AnimContext);
+
+	const UAnimInstance* AnimInstance = Cast<UAnimInstance>(AnimContext);
+	if (!AnimInstance)
+	{
+		// @todo: implement AnimNext support (AnimContext is not an UAnimInstance)
+		unimplemented();
+		return false;
+	}
+
+	const AActor* MainActor = GetMainActor();
+	if (!MainActor)
+	{
+		UE_LOG(LogPoseSearch, Error, TEXT("FInteractionIsland::CanPerformSearch invalid MainActor during the search!"));
+		return false;
+	}
+
+	const AActor* Actor = AnimInstance->GetOwningActor();
+	check(Actor);
+
+	if (Actor != MainActor)
+	{
+		UE_LOG(LogPoseSearch, Error, TEXT("FInteractionIsland::CanPerformSearch search NOT performed by the MainActor!"));
+		return false;
+	}
+	
+	return true;
+}
+
 bool FInteractionIsland::DoSearch_AnyThread(UObject* AnimContext, const FPoseSearchContinuingProperties& ContinuingProperties, FPoseSearchInteractionBlueprintResult& Result)
 {
+	check(AnimContext);
+
 	bool bDoPerfomSearch = false;
 	{
 		// thread safety note!
@@ -549,6 +618,24 @@ bool FInteractionIsland::DoSearch_AnyThread(UObject* AnimContext, const FPoseSea
 		//			fly of FInteractionIsland within the same island that requires searches is forbidden
 		FScopeLock Lock(&SearchResultsMutex);
 		bDoPerfomSearch = !bSearchPerfomed;
+
+		if (bDoPerfomSearch)
+		{
+			if (!CanPerformSearch(AnimContext))
+			{
+				// search failed, because AnimContext is not the context from the main actor. continuing the search would lead to threading issues,
+				// since the threading model expect the main actor to perform the search! (this happens when recompiling animation blue prints..)
+				bSearchPerfomed = true;
+				InitSearchResults(SearchResults, TConstArrayView<FSearchResult>(), TConstArrayView<FInteractionSearchContext>());
+
+				Result = FPoseSearchInteractionBlueprintResult();
+				return false;
+			}
+		}
+		else
+		{
+			return GetResult_AnyThread(AnimContext, Result);
+		}
 	}
 
 	if (bDoPerfomSearch)
@@ -641,65 +728,44 @@ bool FInteractionIsland::DoSearch_AnyThread(UObject* AnimContext, const FPoseSea
 				PoseSearchResults[SearchIndex] = PoseSearchResult;
 			}
 		}
+		
+		// locking to update SearchResults and bSearchPerfomed
+		FScopeLock Lock(&SearchResultsMutex);
 
-		if (!PoseSearchResults.IsEmpty())
+		if (!bSearchPerfomed)
 		{
-			// locking to update SearchResults and bSearchPerfomed
-			FScopeLock Lock(&SearchResultsMutex);
-
-			if (!bSearchPerfomed)
-			{
-				InitSearchResults(SearchResults, PoseSearchResults, SearchContexts);
+			InitSearchResults(SearchResults, PoseSearchResults, SearchContexts);
 
-				bSearchPerfomed = true;
-			}
-			else
-			{
-				UE_LOG(LogPoseSearch, Warning, TEXT("FInteractionIsland::DoSearch_AnyThread performance warning: performed duplicated search :("));
+			bSearchPerfomed = true;
+		}
+		else
+		{
+			UE_LOG(LogPoseSearch, Warning, TEXT("FInteractionIsland::DoSearch_AnyThread performance warning: performed duplicated search :("));
 				
 #if DO_CHECK
-				TArray<FInteractionSearchResult> CompareSearchResults;
-				InitSearchResults(CompareSearchResults, PoseSearchResults, SearchContexts);
+			TArray<FInteractionSearchResult> CompareSearchResults;
+			InitSearchResults(CompareSearchResults, PoseSearchResults, SearchContexts);
 
-				if (CompareSearchResults != SearchResults)
-				{
-					UE_LOG(LogPoseSearch, Error, TEXT("FInteractionIsland::DoSearch_AnyThread duplicated search differs from original one. Searches are NOT deterministic!"));
-				}
-#endif // DO_CHECK
-			}
-	
-			// calling this funtion within this scope since we already locked SearchResultsMutex
-			return GetResult_AnyThread(AnimContext, Result);
-		}
-	}
-
-	return GetResult_AnyThread(AnimContext, Result);
-}
-
-FInteractionSearchContext FInteractionIsland::FindSearchContext_AnyThread(UObject* AnimContext)
-{
-	FScopeLock Lock(&SearchResultsMutex);
-
-	if (bSearchPerfomed)
-	{
-		for (const FInteractionSearchResult& SearchResult : SearchResults)
-		{
-			const FInteractionSearchContext& SearchContext = SearchContexts[SearchResult.SearchIndex];
-			for (int32 AnimContextIndex = 0; AnimContextIndex < SearchContext.AnimContexts.Num(); ++AnimContextIndex)
+			if (CompareSearchResults != SearchResults)
 			{
-				if (SearchContext.AnimContexts[AnimContextIndex].Get() == AnimContext)
-				{
-					return SearchContext;
-				}
+				UE_LOG(LogPoseSearch, Error, TEXT("FInteractionIsland::DoSearch_AnyThread duplicated search differs from original one. Searches are NOT deterministic!"));
 			}
+#endif // DO_CHECK
 		}
+	
+		// calling this funtion within this scope since we already locked SearchResultsMutex
+		return GetResult_AnyThread(AnimContext, Result);
 	}
 
-	return FInteractionSearchContext();
+	checkNoEntry();
+	Result = FPoseSearchInteractionBlueprintResult();
+	return false;
 }
 
 bool FInteractionIsland::GetResult_AnyThread(UObject* AnimContext, FPoseSearchInteractionBlueprintResult& Result)
 {
+	check(AnimContext);
+
 	// locking to read SearchResults
 	FScopeLock Lock(&SearchResultsMutex);
 
@@ -750,12 +816,31 @@ bool FInteractionIsland::GetResult_AnyThread(UObject* AnimContext, FPoseSearchIn
 				//	}
 				//}
 
+				if (const UMultiAnimAsset* MultiAnimAsset = Cast<UMultiAnimAsset>(Result.SelectedAnimation))
+				{
+					const int32 NumRoles = MultiAnimAsset->GetNumRoles();
+					Result.ActorRootTransforms.SetNum(NumRoles);
+					Result.ActorRootBoneTransforms.SetNum(NumRoles);
+
+					const FRoleToIndex InteractionSearchContextRoleToIndex = MakeRoleToIndex(SearchContext.Roles);
+
+					for (int32 MultiAnimAssetRoleIndex = 0; MultiAnimAssetRoleIndex < NumRoles; ++MultiAnimAssetRoleIndex)
+					{
+						if (const int32* InteractionSearchContextRoleIndex = InteractionSearchContextRoleToIndex.Find(MultiAnimAsset->GetRole(MultiAnimAssetRoleIndex)))
+						{
+							Result.ActorRootTransforms[MultiAnimAssetRoleIndex] = SearchResult.ActorRootTransforms[*InteractionSearchContextRoleIndex];
+							Result.ActorRootBoneTransforms[MultiAnimAssetRoleIndex] = SearchResult.ActorRootBoneTransforms[*InteractionSearchContextRoleIndex];
+						}
+					}
+				}
+
 				// we found our AnimContext: we can stop searching
 				return true;
 			}
 		}
 	}
 
+	Result = FPoseSearchInteractionBlueprintResult();
 	return false;
 }
 
diff --git a/Engine/Plugins/Animation/PoseSearch/Source/Runtime/Private/PoseSearchInteractionLibrary.cpp b/Engine/Plugins/Animation/PoseSearch/Source/Runtime/Private/PoseSearchInteractionLibrary.cpp
--- a/Engine/Plugins/Animation/PoseSearch/Source/Runtime/Private/PoseSearchInteractionLibrary.cpp
+++ b/Engine/Plugins/Animation/PoseSearch/Source/Runtime/Private/PoseSearchInteractionLibrary.cpp
@@ -30,15 +30,6 @@ FPoseSearchInteractionBlueprintResult UPoseSearchInteractionLibrary::MotionMatch
 	return Result;
 }
 
-UE::PoseSearch::FInteractionSearchContext UPoseSearchInteractionLibrary::FindSearchContext_AnyThread(UObject* AnimContext)
-{
-	if (UPoseSearchInteractionSubsystem* InteractionSubsystem = UPoseSearchInteractionSubsystem::GetSubsystem_AnyThread(AnimContext))
-	{
-		return InteractionSubsystem->FindSearchContext_AnyThread(AnimContext);
-	}
-	return UE::PoseSearch::FInteractionSearchContext();
-}
-
 FPoseSearchContinuingProperties UPoseSearchInteractionLibrary::GetMontageContinuingProperties(UAnimInstance* AnimInstance)
 {
 	FPoseSearchContinuingProperties ContinuingProperties;
diff --git a/Engine/Plugins/Animation/PoseSearch/Source/Runtime/Private/PoseSearchInteractionSubsystem.cpp b/Engine/Plugins/Animation/PoseSearch/Source/Runtime/Private/PoseSearchInteractionSubsystem.cpp
--- a/Engine/Plugins/Animation/PoseSearch/Source/Runtime/Private/PoseSearchInteractionSubsystem.cpp
+++ b/Engine/Plugins/Animation/PoseSearch/Source/Runtime/Private/PoseSearchInteractionSubsystem.cpp
@@ -9,239 +9,241 @@
 
 namespace UE::PoseSearch
 {
-	FString FInteractionAvailabilityEx::GetPoseHistoryName() const
+
+FString FInteractionAvailabilityEx::GetPoseHistoryName() const
+{
+	if (HistoryCollector)
 	{
-		if (HistoryCollector)
-		{
-			return "HistoryProvider";
-		}
-		return PoseHistoryName.ToString();
+		return "HistoryProvider";
 	}
+	return PoseHistoryName.ToString();
+}
 	
-	const FAnimNode_PoseSearchHistoryCollector_Base* FInteractionAvailabilityEx::GetHistoryCollector(const UObject* AnimContext) const
+const FAnimNode_PoseSearchHistoryCollector_Base* FInteractionAvailabilityEx::GetHistoryCollector(const UObject* AnimContext) const
+{
+	if (HistoryCollector)
 	{
-		if (HistoryCollector)
-		{
-			return HistoryCollector;
-		}
+		return HistoryCollector;
+	}
 		
-		if (const UAnimInstance* AnimInstance = Cast<UAnimInstance>(AnimContext))
-		{
-			return UPoseSearchLibrary::FindPoseHistoryNode(PoseHistoryName, AnimInstance);
-		}
-
-		unimplemented();
-		return nullptr;
+	if (const UAnimInstance* AnimInstance = Cast<UAnimInstance>(AnimContext))
+	{
+		return UPoseSearchLibrary::FindPoseHistoryNode(PoseHistoryName, AnimInstance);
 	}
 
-	struct FInteractionSearchContextGroup
+	unimplemented();
+	return nullptr;
+}
+
+struct FInteractionSearchContextGroup
+{
+	bool Contains(const FInteractionSearchContext& SearchContext) const
 	{
-		bool Contains(const FInteractionSearchContext& SearchContext) const
+		for (const TWeakObjectPtr<const UObject>& AnimContext : SearchContext.AnimContexts)
 		{
-			for (const TWeakObjectPtr<const UObject>& AnimContext : SearchContext.AnimContexts)
+			if (AnimContexts.Find(AnimContext.Get()))
 			{
-				if (AnimContexts.Find(AnimContext.Get()))
-				{
-					return true;
-				}
+				return true;
 			}
-			return false;
 		}
+		return false;
+	}
 
-		void Add(const FInteractionSearchContext& SearchContext, int32 SearchContextIndex)
+	void Add(const FInteractionSearchContext& SearchContext, int32 SearchContextIndex)
+	{
+		for (const TWeakObjectPtr<const UObject>& AnimContext : SearchContext.AnimContexts)
 		{
-			for (const TWeakObjectPtr<const UObject>& AnimContext : SearchContext.AnimContexts)
+			if (AnimContext.Get())
 			{
-				if (AnimContext.Get())
-				{
-					AnimContexts.Add(AnimContext.Get());
-				}
+				AnimContexts.Add(AnimContext.Get());
 			}
-
-			SearchContextsIndices.Add(SearchContextIndex);
 		}
 
-		void Merge(const FInteractionSearchContextGroup& SearchContextGroup)
+		SearchContextsIndices.Add(SearchContextIndex);
+	}
+
+	void Merge(const FInteractionSearchContextGroup& SearchContextGroup)
+	{
+		for (const UObject* AnimContext : SearchContextGroup.AnimContexts)
 		{
-			for (const UObject* AnimContext : SearchContextGroup.AnimContexts)
-			{
-				check(AnimContext);
-				AnimContexts.Add(AnimContext);
-			}
+			check(AnimContext);
+			AnimContexts.Add(AnimContext);
+		}
 
-			for (int32 SearchContextsIndex : SearchContextGroup.SearchContextsIndices)
-			{
-				SearchContextsIndices.Add(SearchContextsIndex);
-			}
+		for (int32 SearchContextsIndex : SearchContextGroup.SearchContextsIndices)
+		{
+			SearchContextsIndices.Add(SearchContextsIndex);
 		}
+	}
 
-		TSet<const UObject*, DefaultKeyFuncs<const UObject*>, TInlineSetAllocator<PreallocatedRolesNum>> AnimContexts;
-		TArray<int32, TInlineAllocator<PreallocatedSearchesNum>> SearchContextsIndices;
-	};
+	TSet<const UObject*, DefaultKeyFuncs<const UObject*>, TInlineSetAllocator<PreallocatedRolesNum>> AnimContexts;
+	TArray<int32, TInlineAllocator<PreallocatedSearchesNum>> SearchContextsIndices;
+};
 
-	struct FRoledAnimContext
-	{
-		const UObject* AnimContext = nullptr;
-		FRole Role = DefaultRole;
-		const FAnimNode_PoseSearchHistoryCollector_Base* HistoryCollector = nullptr;
-		float BroadPhaseRadius = 0.f;
-		float MaxCost = MAX_flt;
-	};
-
-	// @todo: inherit from a TSet if it gets too slow 
-	// sorted array of FRoledAnimContext(s). Sorted by UAnimContext to prevent "alternating" tuples between different frames like
-	// ([AnimContextA, RoleA], [AnimContextB, RoleB]) vs ([AnimContextB, RoleB], [AnimContextA, RoleA]) 
-	// that ultimately represent the same search, but endign up with different SearchIDs
-	struct FRoledAnimContexts
+struct FRoledAnimContext
+{
+	const UObject* AnimContext = nullptr;
+	FRole Role = DefaultRole;
+	const FAnimNode_PoseSearchHistoryCollector_Base* HistoryCollector = nullptr;
+	float BroadPhaseRadius = 0.f;
+	float MaxCost = MAX_flt;
+};
+
+// @todo: inherit from a TSet if it gets too slow 
+// sorted array of FRoledAnimContext(s). Sorted by UAnimContext to prevent "alternating" tuples between different frames like
+// ([AnimContextA, RoleA], [AnimContextB, RoleB]) vs ([AnimContextB, RoleB], [AnimContextA, RoleA]) 
+// that ultimately represent the same search, but endign up with different SearchIDs
+struct FRoledAnimContexts
+{
+	void AddRoledAnimContext(const UObject* AnimContext, FRole Role, const FAnimNode_PoseSearchHistoryCollector_Base* HistoryCollector, float BroadPhaseRadius, float MaxCost)
 	{
-		void AddRoledAnimContext(const UObject* AnimContext, FRole Role, const FAnimNode_PoseSearchHistoryCollector_Base* HistoryCollector, float BroadPhaseRadius, float MaxCost)
-		{
-			check(AnimContext && HistoryCollector);
+		check(AnimContext && HistoryCollector);
 
-			for (FRoledAnimContext& RoledAnimContext : RoledAnimContexts)
+		for (FRoledAnimContext& RoledAnimContext : RoledAnimContexts)
+		{
+			if (RoledAnimContext.AnimContext == AnimContext && RoledAnimContext.HistoryCollector == HistoryCollector && RoledAnimContext.Role == Role)
 			{
-				if (RoledAnimContext.AnimContext == AnimContext && RoledAnimContext.HistoryCollector == HistoryCollector && RoledAnimContext.Role == Role)
-				{
-					// we already found the FRoledAnimContexts. Updating the BroadPhaseRadius and MaxCost accordingly
-					RoledAnimContext.BroadPhaseRadius = FMath::Max(RoledAnimContext.BroadPhaseRadius, BroadPhaseRadius);
+				// we already found the FRoledAnimContexts. Updating the BroadPhaseRadius and MaxCost accordingly
+				RoledAnimContext.BroadPhaseRadius = FMath::Max(RoledAnimContext.BroadPhaseRadius, BroadPhaseRadius);
 					
-					// MaxCost is valid only if greater than zero
-					if (MaxCost > 0.f)
-					{
-						RoledAnimContext.MaxCost = FMath::Min(RoledAnimContext.MaxCost, MaxCost);
-					}
-					return;
+				// MaxCost is valid only if greater than zero
+				if (MaxCost > 0.f)
+				{
+					RoledAnimContext.MaxCost = FMath::Min(RoledAnimContext.MaxCost, MaxCost);
 				}
+				return;
 			}
-
-			// we couldn't find the FRoledAnimContexts. Let's add a new one
-			FRoledAnimContext& NewRoledAnimContext = RoledAnimContexts.AddDefaulted_GetRef();
-			NewRoledAnimContext.AnimContext = AnimContext;
-			NewRoledAnimContext.HistoryCollector = HistoryCollector;
-			NewRoledAnimContext.Role = Role;
-			NewRoledAnimContext.BroadPhaseRadius = BroadPhaseRadius;
-			if (MaxCost > 0.f)
-			{
-				NewRoledAnimContext.MaxCost = MaxCost;
-			}
-
-			// @todo: use a better datastructure to keep things FRoledAnimContext(s) sorted? Or use insertion sort?
-			RoledAnimContexts.Sort([](const FRoledAnimContext& A, const FRoledAnimContext& B)
-			{
-				// we're sorting by AnimContext address in memory, since searches neeed to be consistent 
-				// between frames on the same machine, not between replicated machines
-				return B.AnimContext < A.AnimContext;
-			});
 		}
 
-		TArrayView<const FRoledAnimContext> GetDataView() const
+		// we couldn't find the FRoledAnimContexts. Let's add a new one
+		FRoledAnimContext& NewRoledAnimContext = RoledAnimContexts.AddDefaulted_GetRef();
+		NewRoledAnimContext.AnimContext = AnimContext;
+		NewRoledAnimContext.HistoryCollector = HistoryCollector;
+		NewRoledAnimContext.Role = Role;
+		NewRoledAnimContext.BroadPhaseRadius = BroadPhaseRadius;
+		if (MaxCost > 0.f)
 		{
-			return MakeArrayView(RoledAnimContexts.GetData(), RoledAnimContexts.Num());
+			NewRoledAnimContext.MaxCost = MaxCost;
 		}
 
-	private:
-		TArray<FRoledAnimContext, TInlineAllocator<PreallocatedRolesNum, TMemStackAllocator<>>> RoledAnimContexts;
-	};
+		// @todo: use a better datastructure to keep things FRoledAnimContext(s) sorted? Or use insertion sort?
+		RoledAnimContexts.Sort([](const FRoledAnimContext& A, const FRoledAnimContext& B)
+		{
+			// we're sorting by AnimContext address in memory, since searches neeed to be consistent 
+			// between frames on the same machine, not between replicated machines
+			return B.AnimContext < A.AnimContext;
+		});
+	}
 
-	typedef TMap<const UPoseSearchDatabase*, FRoledAnimContexts, TInlineSetAllocator<PreallocatedSearchesNum, TMemStackSetAllocator<>>> FDatabaseToRoledAnimContexts;
-	typedef TPair<const UPoseSearchDatabase*, FRoledAnimContexts> FDatabaseToRoledAnimContextsPair;
+	TArrayView<const FRoledAnimContext> GetDataView() const
+	{
+		return MakeArrayView(RoledAnimContexts.GetData(), RoledAnimContexts.Num());
+	}
 
-	static void ProcessAvailabilityRequestsMap(const FAvailabilityRequestsMap& AvailabilityRequestsMap, UWorld* SubsystemWorld, FDatabaseToRoledAnimContexts& DatabaseToRoledAnimContexts)
+private:
+	TArray<FRoledAnimContext, TInlineAllocator<PreallocatedRolesNum, TMemStackAllocator<>>> RoledAnimContexts;
+};
+
+typedef TMap<const UPoseSearchDatabase*, FRoledAnimContexts, TInlineSetAllocator<PreallocatedSearchesNum, TMemStackSetAllocator<>>> FDatabaseToRoledAnimContexts;
+typedef TPair<const UPoseSearchDatabase*, FRoledAnimContexts> FDatabaseToRoledAnimContextsPair;
+
+static void ProcessAvailabilityRequestsMap(const FAvailabilityRequestsMap& AvailabilityRequestsMap, UWorld* SubsystemWorld, FDatabaseToRoledAnimContexts& DatabaseToRoledAnimContexts)
+{
+	for (FAvailabilityRequestsMap::TConstIterator Iter = AvailabilityRequestsMap.CreateConstIterator(); Iter; ++Iter)
 	{
-		for (FAvailabilityRequestsMap::TConstIterator Iter = AvailabilityRequestsMap.CreateConstIterator(); Iter; ++Iter)
+		// making sure the weak pointer is still valid
+		UObject* AnimContext = Iter->Key.Get();
+		if (!AnimContext)
+		{
+			UE_LOG(LogPoseSearch, Log, TEXT("ProcessAvailabilityRequestsMap null anim context. The associated character got removed from the world"));
+			continue;
+		}
+
+		if (!AnimContext->GetWorld())
 		{
-			// making sure the weak pointer is still valid
-			UObject* AnimContext = Iter->Key.Get();
-			if (!AnimContext)
+			UE_LOG(LogPoseSearch, Error, TEXT("ProcessAvailabilityRequestsMap AnimContext %s is not in a world!"), *AnimContext->GetName());
+			continue;
+		}
+
+		if (SubsystemWorld != AnimContext->GetWorld())
+		{
+			UE_LOG(LogPoseSearch, Error, TEXT("ProcessAvailabilityRequestsMap AnimContext %s is from World %s, and supposed to be in World @s!"), *AnimContext->GetName(), *AnimContext->GetWorld()->GetName(), *SubsystemWorld->GetName());
+			continue;
+		}
+
+		for (const FInteractionAvailabilityEx& AvailabilityRequest : Iter->Value)
+		{
+			if (!AvailabilityRequest.Database)
 			{
-				UE_LOG(LogPoseSearch, Log, TEXT("ProcessAvailabilityRequestsMap null anim context. The associated character got removed from the world"));
+				UE_LOG(LogPoseSearch, Log, TEXT("ProcessAvailabilityRequestsMap null AvailabilityRequest.Database"));
 				continue;
 			}
 
-			if (!AnimContext->GetWorld())
+			if (!AvailabilityRequest.Database->Schema)
 			{
-				UE_LOG(LogPoseSearch, Error, TEXT("ProcessAvailabilityRequestsMap AnimContext %s is not in a world!"), *AnimContext->GetName());
+				UE_LOG(LogPoseSearch, Error, TEXT("ProcessAvailabilityRequestsMap null Schema for Database %s"), *AvailabilityRequest.Database->GetName());
 				continue;
 			}
 
-			if (SubsystemWorld != AnimContext->GetWorld())
+			const FAnimNode_PoseSearchHistoryCollector_Base* HistoryCollector = AvailabilityRequest.GetHistoryCollector(AnimContext);
+			if (!HistoryCollector)
 			{
-				UE_LOG(LogPoseSearch, Error, TEXT("ProcessAvailabilityRequestsMap AnimContext %s is from World %s, and supposed to be in World @s!"), *AnimContext->GetName(), *AnimContext->GetWorld()->GetName(), *SubsystemWorld->GetName());
+				UE_LOG(LogPoseSearch, Error, TEXT("ProcessAvailabilityRequestsMap couldn't find PoseHistory %s for AnimContext %s"), *AvailabilityRequest.GetPoseHistoryName(), *AnimContext->GetName());
 				continue;
 			}
 
-			for (const FInteractionAvailabilityEx& AvailabilityRequest : Iter->Value)
+			FRoledAnimContexts& RoledAnimContexts = DatabaseToRoledAnimContexts.FindOrAdd(AvailabilityRequest.Database);
+			if (AvailabilityRequest.RolesFilter.IsEmpty())
 			{
-				if (!AvailabilityRequest.Database)
+				// adding ALL the possible roles from the database:
+				for (const FPoseSearchRoledSkeleton& RoledSkeleton : AvailabilityRequest.Database->Schema->GetRoledSkeletons())
 				{
-					UE_LOG(LogPoseSearch, Log, TEXT("ProcessAvailabilityRequestsMap null AvailabilityRequest.Database"));
-					continue;
+					RoledAnimContexts.AddRoledAnimContext(AnimContext, RoledSkeleton.Role, HistoryCollector, AvailabilityRequest.BroadPhaseRadius, AvailabilityRequest.MaxCost);
 				}
-
-				if (!AvailabilityRequest.Database->Schema)
-				{
-					UE_LOG(LogPoseSearch, Error, TEXT("ProcessAvailabilityRequestsMap null Schema for Database %s"), *AvailabilityRequest.Database->GetName());
-					continue;
-				}
-
-				const FAnimNode_PoseSearchHistoryCollector_Base* HistoryCollector = AvailabilityRequest.GetHistoryCollector(AnimContext);
-				if (!HistoryCollector)
-				{
-					UE_LOG(LogPoseSearch, Error, TEXT("ProcessAvailabilityRequestsMap couldn't find PoseHistory %s for AnimContext %s"), *AvailabilityRequest.GetPoseHistoryName(), *AnimContext->GetName());
-					continue;
-				}
-
-				FRoledAnimContexts& RoledAnimContexts = DatabaseToRoledAnimContexts.FindOrAdd(AvailabilityRequest.Database);
-				if (AvailabilityRequest.RolesFilter.IsEmpty())
+			}
+			else
+			{
+				for (const FRole& Role : AvailabilityRequest.RolesFilter)
 				{
-					// adding ALL the possible roles from the database:
-					for (const FPoseSearchRoledSkeleton& RoledSkeleton : AvailabilityRequest.Database->Schema->GetRoledSkeletons())
+					if (AvailabilityRequest.Database->Schema->GetRoledSkeleton(Role))
 					{
-						RoledAnimContexts.AddRoledAnimContext(AnimContext, RoledSkeleton.Role, HistoryCollector, AvailabilityRequest.BroadPhaseRadius, AvailabilityRequest.MaxCost);
+						RoledAnimContexts.AddRoledAnimContext(AnimContext, Role, HistoryCollector, AvailabilityRequest.BroadPhaseRadius, AvailabilityRequest.MaxCost);
 					}
-				}
-				else
-				{
-					for (const FRole& Role : AvailabilityRequest.RolesFilter)
+					else
 					{
-						if (AvailabilityRequest.Database->Schema->GetRoledSkeleton(Role))
-						{
-							RoledAnimContexts.AddRoledAnimContext(AnimContext, Role, HistoryCollector, AvailabilityRequest.BroadPhaseRadius, AvailabilityRequest.MaxCost);
-						}
-						else
-						{
-							UE_LOG(LogPoseSearch, Warning, TEXT("ProcessAvailabilityRequestsMap unsupported Role %s for Database %s"), *Role.ToString(), *AvailabilityRequest.Database->GetName());
-						}
+						UE_LOG(LogPoseSearch, Warning, TEXT("ProcessAvailabilityRequestsMap unsupported Role %s for Database %s"), *Role.ToString(), *AvailabilityRequest.Database->GetName());
 					}
 				}
 			}
 		}
 	}
+}
 
-	template <typename DataType, typename EvaluateCombinationType>
-	static void GeneratePermutationsRecursive(const TArrayView<DataType> Data, int32 DataIndex, TArrayView<int32> Combination, int32 CombinationIndex, EvaluateCombinationType EvaluateCombination)
-	{
-		if (CombinationIndex == Combination.Num())
-		{ 
-			EvaluateCombination(Data, Combination);
-		}
-		else if (DataIndex < Data.Num())
-		{
-			Combination[CombinationIndex] = DataIndex;
-			GeneratePermutationsRecursive(Data, DataIndex + 1, Combination, CombinationIndex + 1, EvaluateCombination);
-			GeneratePermutationsRecursive(Data, DataIndex + 1, Combination, CombinationIndex, EvaluateCombination);
-		}
-	} 
-
-	template <typename DataType, typename EvaluateCombinationType>
-	static void GeneratePermutations(const TArrayView<DataType> Data, int32 CombinationCardinality, EvaluateCombinationType EvaluateCombination)
+template <typename DataType, typename EvaluateCombinationType>
+static void GeneratePermutationsRecursive(const TArrayView<DataType> Data, int32 DataIndex, TArrayView<int32> Combination, int32 CombinationIndex, EvaluateCombinationType EvaluateCombination)
+{
+	if (CombinationIndex == Combination.Num())
+	{ 
+		EvaluateCombination(Data, Combination);
+	}
+	else if (DataIndex < Data.Num())
 	{
-		TArray<int32, TInlineAllocator<PreallocatedRolesNum>> Combination;
-		Combination.SetNum(CombinationCardinality);
-		GeneratePermutationsRecursive(Data, 0, Combination, 0, EvaluateCombination);
+		Combination[CombinationIndex] = DataIndex;
+		GeneratePermutationsRecursive(Data, DataIndex + 1, Combination, CombinationIndex + 1, EvaluateCombination);
+		GeneratePermutationsRecursive(Data, DataIndex + 1, Combination, CombinationIndex, EvaluateCombination);
 	}
+} 
+
+template <typename DataType, typename EvaluateCombinationType>
+static void GeneratePermutations(const TArrayView<DataType> Data, int32 CombinationCardinality, EvaluateCombinationType EvaluateCombination)
+{
+	TArray<int32, TInlineAllocator<PreallocatedRolesNum>> Combination;
+	Combination.SetNum(CombinationCardinality);
+	GeneratePermutationsRecursive(Data, 0, Combination, 0, EvaluateCombination);
 }
 
+} // namespace UE::PoseSearch
+
 FCriticalSection UPoseSearchInteractionSubsystem::RetrieveSubsystemMutex;
 
 UE::PoseSearch::FInteractionIsland& UPoseSearchInteractionSubsystem::CreateIsland()
@@ -290,14 +292,20 @@ void UPoseSearchInteractionSubsystem::UninjectAllIslands()
 	}
 }
 
+#if DO_CHECK
 bool UPoseSearchInteractionSubsystem::ValidateAllIslands() const
 {
-#if DO_CHECK
 	using namespace UE::PoseSearch;
 
 	TSet<TWeakObjectPtr<UActorComponent>> PreTickComponents;
 	TSet<TWeakObjectPtr<UActorComponent>> PostTickComponents;
 
+	typedef TSet<const UObject*> FIslandAnimContexts;
+	TArray<FIslandAnimContexts> IslandsAnimContexts;
+
+	const int32 NumIslands = Islands.Num();
+	IslandsAnimContexts.Reserve(NumIslands);
+
 	bool bAlreadyInSet = false;
 	for (const FInteractionIsland* Island : Islands)
 	{
@@ -318,10 +326,41 @@ bool UPoseSearchInteractionSubsystem::ValidateAllIslands() const
 				return false;
 			}
 		}
+
+		FIslandAnimContexts& IslandAnimContexts = IslandsAnimContexts.AddDefaulted_GetRef();
+		for (const FInteractionSearchContext& SearchContext : Island->GetSearchContexts())
+		{
+			for (const TWeakObjectPtr<const UObject>& AnimContextPtr : SearchContext.AnimContexts)
+			{
+				if (const UObject* AnimContext = AnimContextPtr.Get())
+				{
+					IslandAnimContexts.Add(AnimContext);
+				}
+			}
+		}
 	}
-#endif // DO_CHECK
+
+	for (int32 IslandIndex = 0; IslandIndex < NumIslands; ++IslandIndex)
+	{
+		for (const UObject* AnimContext : IslandsAnimContexts[IslandIndex])
+		{
+			for (int32 OtherIslandIndex = 0; OtherIslandIndex < NumIslands; ++OtherIslandIndex)
+			{
+				if (IslandIndex != OtherIslandIndex)
+				{
+					if (IslandsAnimContexts[OtherIslandIndex].Find(AnimContext))
+					{
+						// AnimContext is shared between multiple islands. it'd cause multi threadind issues!
+						return false;
+					}
+				}
+			}
+		}				
+	}
+
 	return true;
 }
+#endif // DO_CHECK
 
 void UPoseSearchInteractionSubsystem::PopulateContinuingProperties(UE::PoseSearch::FInteractionSearchContext& SearchContext, float DeltaSeconds) const
 {
@@ -601,7 +640,7 @@ void UPoseSearchInteractionSubsystem::Tick(float DeltaSeconds)
 
 #if ENABLE_DRAW_DEBUG
 	DebugDraw();
-#endif //ENABLE_DRAW_DEBUG
+#endif // ENABLE_DRAW_DEBUG
 
 	UninjectAllIslands();
 
@@ -681,20 +720,11 @@ void UPoseSearchInteractionSubsystem::Query_AnyThread(const TArrayView<const FPo
 	}
 }
 
-UE::PoseSearch::FInteractionSearchContext UPoseSearchInteractionSubsystem::FindSearchContext_AnyThread(UObject* AnimContext)
-{
-	if (UE::PoseSearch::FInteractionIsland* Island = FindIsland(AnimContext))
-	{
-		return Island->FindSearchContext_AnyThread(AnimContext);
-	}
-	return UE::PoseSearch::FInteractionSearchContext();
-}
-
+#if ENABLE_DRAW_DEBUG
 void UPoseSearchInteractionSubsystem::DebugDraw() const
 {
 	using namespace UE::PoseSearch;
 
-#if ENABLE_DRAW_DEBUG
 	static const FColor Colors[] = { FColor::Red, FColor::Green, FColor::Blue, FColor::Yellow, FColor::Black };
 	static const int32 NumColors = sizeof(Colors) / sizeof(Colors[0]);
 	int32 CurrentColorIndex = 0;
@@ -707,5 +737,5 @@ void UPoseSearchInteractionSubsystem::DebugDraw() const
 			CurrentColorIndex = (CurrentColorIndex + 1) % NumColors;
 		}
 	}
-#endif //ENABLE_DRAW_DEBUG
-}
\ No newline at end of file
+}
+#endif // ENABLE_DRAW_DEBUG
diff --git a/Engine/Plugins/Animation/PoseSearch/Source/Runtime/Private/PoseSearchInteractionValidator.cpp b/Engine/Plugins/Animation/PoseSearch/Source/Runtime/Private/PoseSearchInteractionValidator.cpp
new file mode 100644
--- /dev/null
+++ b/Engine/Plugins/Animation/PoseSearch/Source/Runtime/Private/PoseSearchInteractionValidator.cpp
@@ -0,0 +1,122 @@
+// Copyright Epic Games, Inc. All Rights Reserved.
+
+#include "PoseSearch/PoseSearchInteractionValidator.h"
+#include "PoseSearch/PoseSearchInteractionSubsystem.h"
+
+namespace UE::PoseSearch
+{
+#if ENABLE_ANIM_DEBUG
+
+FInteractionValidator::FInteractionValidator(UObject* AnimContext, EPhase Phase)
+	: ValidatingAnimContext(AnimContext)
+{
+	check(AnimContext);
+	UPoseSearchInteractionSubsystem* InteractionSubsystem = UPoseSearchInteractionSubsystem::GetSubsystem_AnyThread(AnimContext);
+	if (!InteractionSubsystem)
+	{
+		return;
+	}
+
+	ValidatingIsland = InteractionSubsystem->FindIsland(AnimContext);
+	if (!ValidatingIsland)
+	{
+		return;
+	}
+
+	const UAnimInstance* AnimInstance = Cast<UAnimInstance>(AnimContext);
+	if (!AnimInstance)
+	{
+		// @todo: implement AnimNext support (AnimContext is not an UAnimInstance)
+		unimplemented();
+		return;
+	}
+
+	const AActor* Actor = AnimInstance->GetOwningActor();
+	check(Actor);
+
+	const AActor* MainActor = ValidatingIsland->GetMainActor();
+	if (!MainActor)
+	{
+		UE_LOG(LogPoseSearch, Error, TEXT("FInteractionValidator invalid MainActor! How did he died? Did you rebuild the animation blue print while PIE was running?"));
+		return;
+	}
+
+	if (Actor == MainActor)
+	{
+		const int32 PreviousCounter = ValidatingIsland->InteractionIslandThreadSafeCounter.Set(-1);
+		if (PreviousCounter != 0)
+		{
+			UE_LOG(LogPoseSearch, Error, TEXT("Non thread safe call! Why is there any other actor running while we schedule the MainActor? %d"), PreviousCounter);
+		}
+	}
+	else
+	{
+		// not the main actor
+		const int32 PreviousCounter = ValidatingIsland->InteractionIslandThreadSafeCounter.Add(1);
+		if (PreviousCounter < 0)
+		{
+			UE_LOG(LogPoseSearch, Error, TEXT("Non thread safe call! The MainActor is running! Nobody else in the same island should be running! %d"), PreviousCounter);
+		}
+	}
+}
+
+FInteractionValidator::~FInteractionValidator()
+{
+	if (!ValidatingIsland)
+	{
+		return;
+	}
+
+	UPoseSearchInteractionSubsystem* InteractionSubsystem = UPoseSearchInteractionSubsystem::GetSubsystem_AnyThread(ValidatingAnimContext);
+	if (!InteractionSubsystem)
+	{
+		return;
+	}
+
+	FInteractionIsland* Island = InteractionSubsystem->FindIsland(ValidatingAnimContext);
+	if (Island != ValidatingIsland)
+	{
+		UE_LOG(LogPoseSearch, Error, TEXT("FInteractionValidator why did the InteractionIsland changed?"));
+		return;
+	}
+
+	const UAnimInstance* AnimInstance = Cast<UAnimInstance>(ValidatingAnimContext);
+	if (!AnimInstance)
+	{
+		// @todo: implement AnimNext support (AnimContext is not an UAnimInstance)
+		unimplemented();
+		return;
+	}
+
+	const AActor* Actor = AnimInstance->GetOwningActor();
+	check(Actor);
+
+	const AActor* MainActor = Island->GetMainActor();
+	if (!MainActor)
+	{
+		UE_LOG(LogPoseSearch, Error, TEXT("FInteractionValidator invalid MainActor! How did he died? Did you rebuild the animation blue print while PIE was running?"));
+		return;
+	}
+	
+	if (Actor == MainActor)
+	{
+		const int32 PreviousCounter = Island->InteractionIslandThreadSafeCounter.Set(0);
+		if (PreviousCounter >= 0)
+		{
+			UE_LOG(LogPoseSearch, Error, TEXT("Non thread safe call! Why was there some other actor running while we ended the scheduling the MainActor? %d"), PreviousCounter);
+		}
+	}
+	else
+	{
+		// not the main actor
+		const int32 PreviousCounter = Island->InteractionIslandThreadSafeCounter.Add(-1);
+		if (PreviousCounter <= 0)
+		{
+			UE_LOG(LogPoseSearch, Error, TEXT("Non thread safe call! The MainActor is running! Nobody else in the same island should be running! %d"), PreviousCounter);
+		}
+	}
+}
+
+#endif // ENABLE_ANIM_DEBUG
+
+} // namespace UE::PoseSearch
\ No newline at end of file
diff --git a/Engine/Plugins/Animation/PoseSearch/Source/Runtime/Public/PoseSearch/AnimNode_MotionMatchingInteraction.h b/Engine/Plugins/Animation/PoseSearch/Source/Runtime/Public/PoseSearch/AnimNode_MotionMatchingInteraction.h
--- a/Engine/Plugins/Animation/PoseSearch/Source/Runtime/Public/PoseSearch/AnimNode_MotionMatchingInteraction.h
+++ b/Engine/Plugins/Animation/PoseSearch/Source/Runtime/Public/PoseSearch/AnimNode_MotionMatchingInteraction.h
@@ -41,6 +41,10 @@ struct POSESEARCH_API FAnimNode_MotionMatchingInteraction : public FAnimNode_Ble
 	UPROPERTY(EditAnywhere, Category = Warping, meta = (PinHiddenByDefault))
 	bool bUsePoseSearchWarpingNotifyStates = true;
 
+	// enable / disable the usage of UAnimNotifyState_PoseSearchIKWindow(s) to calculate IK constraints and feed them to the related Goals
+	UPROPERTY(EditAnywhere, Category = Warping, meta = (PinHiddenByDefault))
+	bool bUsePoseSearchIKWindowNotifyStates = true;
+
 	// Reset the blend stack if it has become relevant to the graph after not being updated on previous frames.
 	UPROPERTY(EditAnywhere, Category = Settings)
 	bool bResetOnBecomingRelevant = true;
diff --git a/Engine/Plugins/Animation/PoseSearch/Source/Runtime/Public/PoseSearch/PoseSearchAnimNotifies.h b/Engine/Plugins/Animation/PoseSearch/Source/Runtime/Public/PoseSearch/PoseSearchAnimNotifies.h
--- a/Engine/Plugins/Animation/PoseSearch/Source/Runtime/Public/PoseSearch/PoseSearchAnimNotifies.h
+++ b/Engine/Plugins/Animation/PoseSearch/Source/Runtime/Public/PoseSearch/PoseSearchAnimNotifies.h
@@ -118,8 +118,20 @@ public:
 	uint32 GetBranchInId() const;
 };
 
+// multi character interaction notifies
 UCLASS(Experimental, Blueprintable, meta = (DisplayName = "Pose Search: Warping"))
 class POSESEARCH_API UAnimNotifyState_PoseSearchWarping : public UAnimNotifyState_PoseSearchBase
 {
 	GENERATED_BODY()
-};
\ No newline at end of file
+};
+
+UCLASS(Experimental, Blueprintable, meta = (DisplayName = "Pose Search: IKWindow"))
+class POSESEARCH_API UAnimNotifyState_PoseSearchIKWindow : public UAnimNotifyState_PoseSearchBase
+{
+	GENERATED_BODY()
+
+public:
+	// list of IK goals to update with transforms from the fully aligned pose
+	UPROPERTY(EditAnywhere, Category = Config)
+	TArray<FName> GoalNames;
+};
diff --git a/Engine/Plugins/Animation/PoseSearch/Source/Runtime/Public/PoseSearch/PoseSearchAssetSamplerLibrary.h b/Engine/Plugins/Animation/PoseSearch/Source/Runtime/Public/PoseSearch/PoseSearchAssetSamplerLibrary.h
--- a/Engine/Plugins/Animation/PoseSearch/Source/Runtime/Public/PoseSearch/PoseSearchAssetSamplerLibrary.h
+++ b/Engine/Plugins/Animation/PoseSearch/Source/Runtime/Public/PoseSearch/PoseSearchAssetSamplerLibrary.h
@@ -5,6 +5,8 @@
 #include "Animation/AnimCurveTypes.h"
 #include "Kismet/BlueprintFunctionLibrary.h"
 #include "PoseSearch/PoseSearchAssetSampler.h"
+#include "PoseSearch/PoseSearchDefines.h"
+#include "DrawDebugHelpers.h"
 #include "PoseSearchAssetSamplerLibrary.generated.h"
 
 USTRUCT(Experimental, BlueprintType, Category="Animation|Pose Search")
@@ -66,17 +68,80 @@ class POSESEARCH_API UPoseSearchAssetSamplerLibrary : public UBlueprintFunctionL
 {
 	GENERATED_BODY()
 
+public:
 	UFUNCTION(BlueprintPure, Category = "Animation|Pose Search|Experimental", meta = (BlueprintThreadSafe))
 	static FPoseSearchAssetSamplerPose SamplePose(const UAnimInstance* AnimInstance, const FPoseSearchAssetSamplerInput Input);
 
-	UFUNCTION(BlueprintPure, Category = "Animation|Pose Search|Experimental", meta = (BlueprintThreadSafe))
-	static FTransform GetTransform(UPARAM(ref) FPoseSearchAssetSamplerPose& AssetSamplerPose, int32 BoneIndex = -1, EPoseSearchAssetSamplerSpace Space = EPoseSearchAssetSamplerSpace::World);
-
 	UFUNCTION(BlueprintPure, Category = "Animation|Pose Search|Experimental", meta = (BlueprintThreadSafe))
 	static FTransform GetTransformByName(UPARAM(ref) FPoseSearchAssetSamplerPose& AssetSamplerPose, FName BoneName, EPoseSearchAssetSamplerSpace Space = EPoseSearchAssetSamplerSpace::World);
 
+	static FTransform GetTransform(FPoseSearchAssetSamplerPose& AssetSamplerPose, FCompactPoseBoneIndex CompactPoseBoneIndex = FCompactPoseBoneIndex(INDEX_NONE), EPoseSearchAssetSamplerSpace Space = EPoseSearchAssetSamplerSpace::World);
+
 	// @todo: it'd be nice if it was threadsafe...
 	UFUNCTION(BlueprintCallable, Category="Animation|Pose Search|Experimental")
 	static void Draw(const UAnimInstance* AnimInstance, UPARAM(ref) FPoseSearchAssetSamplerPose& AssetSamplerPose);
+
+	template<typename FComponentSpacePoseType>
+	static FTransform GetTransform(FComponentSpacePoseType& ComponentSpacePose, const FTransform& RootTransform, FCompactPoseBoneIndex CompactPoseBoneIndex = FCompactPoseBoneIndex(INDEX_NONE), EPoseSearchAssetSamplerSpace Space = EPoseSearchAssetSamplerSpace::World)
+	{
+		if (!ComponentSpacePose.GetPose().IsValid())
+		{
+			UE_LOG(LogPoseSearch, Error, TEXT("UPoseSearchAssetSamplerLibrary::GetTransform invalid Pose"));
+			return FTransform::Identity;
+		}
+
+		if (CompactPoseBoneIndex.GetInt() == INDEX_NONE)
+		{
+			if (Space != EPoseSearchAssetSamplerSpace::World)
+			{
+				UE_LOG(LogPoseSearch, Error, TEXT("UPoseSearchAssetSamplerLibrary::GetTransform invalid Space %s to get the RootTransform. Expected space: %s"), *UEnum::GetDisplayValueAsText(Space).ToString(), *UEnum::GetDisplayValueAsText(EPoseSearchAssetSamplerSpace::World).ToString());
+			}
+			return RootTransform;
+		}
+	
+		if (!ComponentSpacePose.GetPose().IsValidIndex(CompactPoseBoneIndex))
+		{
+			UE_LOG(LogPoseSearch, Error, TEXT("UPoseSearchAssetSamplerLibrary::GetTransform invalid CompactPoseBoneIndex %d"), CompactPoseBoneIndex.GetInt());
+			return FTransform::Identity;
+		}
+	
+		switch (Space)
+		{
+		case EPoseSearchAssetSamplerSpace::Local:
+			return ComponentSpacePose.GetPose()[CompactPoseBoneIndex];
+
+		case EPoseSearchAssetSamplerSpace::Component:
+			return ComponentSpacePose.GetComponentSpaceTransform(CompactPoseBoneIndex);
+
+		case EPoseSearchAssetSamplerSpace::World:
+			return ComponentSpacePose.GetComponentSpaceTransform(CompactPoseBoneIndex) * RootTransform;
+		}
+
+		checkNoEntry();
+		return FTransform::Identity;
+	}
+
+#if ENABLE_DRAW_DEBUG
+	template<typename FComponentSpacePoseType>
+	static void Draw(const UWorld* World, FComponentSpacePoseType& ComponentSpacePose, const FTransform& RootTransform, float DebugDrawSamplerRootAxisLength = 20.f, float DebugDrawSamplerSize = 6.f)
+	{
+		if (World)
+		{
+			if (DebugDrawSamplerRootAxisLength > 0.f)
+			{
+				const FTransform AxisWorldTransform = GetTransform(ComponentSpacePose, RootTransform);
+				DrawDebugLine(World, AxisWorldTransform.GetTranslation(), AxisWorldTransform.GetTranslation() + AxisWorldTransform.GetScaledAxis(EAxis::X) * DebugDrawSamplerRootAxisLength, FColor::Red, false, 0.f, SDPG_Foreground);
+				DrawDebugLine(World, AxisWorldTransform.GetTranslation(), AxisWorldTransform.GetTranslation() + AxisWorldTransform.GetScaledAxis(EAxis::Y) * DebugDrawSamplerRootAxisLength, FColor::Green, false, 0.f, SDPG_Foreground);
+				DrawDebugLine(World, AxisWorldTransform.GetTranslation(), AxisWorldTransform.GetTranslation() + AxisWorldTransform.GetScaledAxis(EAxis::Z) * DebugDrawSamplerRootAxisLength, FColor::Blue, false, 0.f, SDPG_Foreground);
+			}
+
+			for (int32 BoneIndex = 0; BoneIndex < ComponentSpacePose.GetPose().GetNumBones(); ++BoneIndex)
+			{
+				const FTransform BoneWorldTransform = GetTransform(ComponentSpacePose, RootTransform, FCompactPoseBoneIndex(BoneIndex), EPoseSearchAssetSamplerSpace::World);
+				DrawDebugPoint(World, BoneWorldTransform.GetTranslation(), DebugDrawSamplerSize, FColor::Red, false, 0.f, SDPG_Foreground);
+			}
+		}
+	}
+#endif // ENABLE_DRAW_DEBUG
 };
 
diff --git a/Engine/Plugins/Animation/PoseSearch/Source/Runtime/Public/PoseSearch/PoseSearchHistory.h b/Engine/Plugins/Animation/PoseSearch/Source/Runtime/Public/PoseSearch/PoseSearchHistory.h
--- a/Engine/Plugins/Animation/PoseSearch/Source/Runtime/Public/PoseSearch/PoseSearchHistory.h
+++ b/Engine/Plugins/Animation/PoseSearch/Source/Runtime/Public/PoseSearch/PoseSearchHistory.h
@@ -210,8 +210,7 @@ private:
 #if ENABLE_ANIM_DEBUG
 	
 	// used to analyze thread safety
-	mutable FThreadSafeCounter ReadPoseDataThreadSafeCounter = 0;
-	mutable FThreadSafeCounter WritePoseDataThreadSafeCounter = 0;
+	mutable FThreadSafeCounter PoseDataThreadSafeCounter = 0;
 
 #endif // ENABLE_ANIM_DEBUG
 };
diff --git a/Engine/Plugins/Animation/PoseSearch/Source/Runtime/Public/PoseSearch/PoseSearchInteractionIsland.h b/Engine/Plugins/Animation/PoseSearch/Source/Runtime/Public/PoseSearch/PoseSearchInteractionIsland.h
--- a/Engine/Plugins/Animation/PoseSearch/Source/Runtime/Public/PoseSearch/PoseSearchInteractionIsland.h
+++ b/Engine/Plugins/Animation/PoseSearch/Source/Runtime/Public/PoseSearch/PoseSearchInteractionIsland.h
@@ -2,6 +2,7 @@
 
 #pragma once
 
+#include "PoseSearch/PoseSearchInteractionValidator.h"
 #include "PoseSearch/PoseSearchLibrary.h"
 
 struct FPoseSearchInteractionBlueprintResult;
@@ -34,9 +35,10 @@ struct FInteractionSearchContext
 
 	bool IsValid() const { return !AnimContexts.IsEmpty(); }
 	bool IsEquivalent(const FInteractionSearchContext& Other) const;
-	bool GetActorTransforms(TArrayView<FTransform> ActorTransforms, const UMultiAnimAsset* MultiAnimAsset) const;
 
 #if DO_CHECK
+	// method that queries the HistoryCollectors. the PoseHistory will complain if accessed in a non thread safe manner!
+	void TestHistoryCollectorsThreadingAccess() const;
 	bool CheckForConsistency() const;
 #endif // DO_CHECK
 };
@@ -46,6 +48,14 @@ struct FInteractionSearchResult : public FSearchResult
 {
 	int32 SearchIndex = INDEX_NONE;
 
+	// cached actors root transforms for all the roles in SelectedAnimation (as UMultiAnimAsset),
+	// so we don't have to query the pose history to gather it when it's not thread safe to do so
+	TArray<FTransform, TInlineAllocator<PreallocatedRolesNum>> ActorRootTransforms;
+
+	// cached actors root bone transforms for all the roles in SelectedAnimation (as UMultiAnimAsset),
+	// so we don't have to query the pose history to gather it when it's not thread safe to do so
+	TArray<FTransform, TInlineAllocator<PreallocatedRolesNum>> ActorRootBoneTransforms;
+
 	bool operator==(const FInteractionSearchResult& Other) const;
 };
 
@@ -62,7 +72,6 @@ struct FInteractionIsland
 	~FInteractionIsland();
 
 	bool DoSearch_AnyThread(UObject* AnimContext, const FPoseSearchContinuingProperties& ContinuingProperties, FPoseSearchInteractionBlueprintResult& Result);
-	UE::PoseSearch::FInteractionSearchContext FindSearchContext_AnyThread(UObject* AnimContext);
 
 	const TArray<TWeakObjectPtr<UActorComponent>>& GetPreTickComponents() const { return PreTickComponents; }
 	const TArray<TWeakObjectPtr<UActorComponent>>& GetPostTickComponents() const { return PostTickComponents; }
@@ -83,6 +92,8 @@ struct FInteractionIsland
 
 private:
 	bool GetResult_AnyThread(UObject* AnimContext, FPoseSearchInteractionBlueprintResult& Result);
+	const AActor* GetMainActor() const;
+	bool CanPerformSearch(const UObject* AnimContext) const;
 
 	struct FPreTickFunction : public FTickFunction
 	{
@@ -110,6 +121,12 @@ private:
 	bool bSearchPerfomed = false;
 
 	FCriticalSection SearchResultsMutex;
+
+#if ENABLE_ANIM_DEBUG
+	// used to analyze thread safety
+	friend struct UE::PoseSearch::FInteractionValidator;
+	FThreadSafeCounter InteractionIslandThreadSafeCounter = 0;
+#endif // ENABLE_ANIM_DEBUG
 };
 
 } // namespace UE::PoseSearch
diff --git a/Engine/Plugins/Animation/PoseSearch/Source/Runtime/Public/PoseSearch/PoseSearchInteractionLibrary.h b/Engine/Plugins/Animation/PoseSearch/Source/Runtime/Public/PoseSearch/PoseSearchInteractionLibrary.h
--- a/Engine/Plugins/Animation/PoseSearch/Source/Runtime/Public/PoseSearch/PoseSearchInteractionLibrary.h
+++ b/Engine/Plugins/Animation/PoseSearch/Source/Runtime/Public/PoseSearch/PoseSearchInteractionLibrary.h
@@ -12,6 +12,7 @@
 // with one of the roles in RolesFilter (if empty ANY of the Database roles can be taken) the MotionMatchInteraction_Pure will ultimately setup a motion matching query
 // using looking for the pose history "PoseHistoryName" to gather bone and trajectory positions for this character for an interaction to be valid,
 // the query needs to find all the other interacting characters within BroadPhaseRadius, and reach a maximum cost of MaxCost
+// Experimental, this feature might be removed without warning, not for production use
 USTRUCT(Experimental, BlueprintType, Category = "Animation|Pose Search")
 struct POSESEARCH_API FPoseSearchInteractionAvailability
 {
@@ -34,6 +35,7 @@ struct POSESEARCH_API FPoseSearchInteractionAvailability
 	float MaxCost = 0.f;
 };
 
+// Experimental, this feature might be removed without warning, not for production use
 USTRUCT(Experimental, BlueprintType, Category="Animation|Pose Search")
 struct POSESEARCH_API FPoseSearchInteractionBlueprintResult
 {
@@ -79,8 +81,17 @@ public:
 	// assigned role to this character (AnimInstance)
 	UPROPERTY(Transient, VisibleAnywhere, BlueprintReadOnly, Category=State)
 	FName Role;
+
+	// cached actors root transforms for all the roles in SelectedAnimation (as UMultiAnimAsset)
+	UPROPERTY(Transient, VisibleAnywhere, BlueprintReadOnly, Category=State)
+	TArray<FTransform> ActorRootTransforms;
+
+	// cached actors root bone transforms for all the roles in SelectedAnimation (as UMultiAnimAsset)
+	UPROPERTY(Transient, VisibleAnywhere, BlueprintReadOnly, Category=State)
+	TArray<FTransform> ActorRootBoneTransforms;
 };
 
+// Experimental, this feature might be removed without warning, not for production use
 UCLASS(Experimental)
 class UPoseSearchInteractionLibrary : public UBlueprintFunctionLibrary
 {
@@ -108,10 +119,6 @@ public:
 	// version of MotionMatchInteraction_Pure referencing directly the HistoryCollector rather than looking for it by name
 	static FPoseSearchInteractionBlueprintResult MotionMatchInteraction(const TArrayView<const FPoseSearchInteractionAvailability> Availabilities, UObject* AnimInstance, const FPoseSearchContinuingProperties& ContinuingProperties, const FAnimNode_PoseSearchHistoryCollector_Base* HistoryCollector, bool bValidateResultAgainstAvailabilities = true);
 
-	// returns the associated FInteractionSearchContext to AnimContext. We do a copy to be thread safe.
-	// check FInteractionSearchContext::IsValid for its validity
-	static UE::PoseSearch::FInteractionSearchContext FindSearchContext_AnyThread(UObject* AnimContext);
-
 	UFUNCTION(Experimental, BlueprintPure, Category = "Animation|Pose Search", meta = (BlueprintThreadSafe))
 	static FPoseSearchContinuingProperties GetMontageContinuingProperties(UAnimInstance* AnimInstance);
 };
\ No newline at end of file
diff --git a/Engine/Plugins/Animation/PoseSearch/Source/Runtime/Public/PoseSearch/PoseSearchInteractionSubsystem.h b/Engine/Plugins/Animation/PoseSearch/Source/Runtime/Public/PoseSearch/PoseSearchInteractionSubsystem.h
--- a/Engine/Plugins/Animation/PoseSearch/Source/Runtime/Public/PoseSearch/PoseSearchInteractionSubsystem.h
+++ b/Engine/Plugins/Animation/PoseSearch/Source/Runtime/Public/PoseSearch/PoseSearchInteractionSubsystem.h
@@ -4,6 +4,7 @@
 
 #include "Subsystems/WorldSubsystem.h"
 #include "PoseSearch/PoseSearchInteractionLibrary.h"
+#include "PoseSearch/PoseSearchInteractionValidator.h"
 #include "PoseSearchInteractionSubsystem.generated.h"
 
 struct FPoseSearchContinuingProperties;
@@ -54,7 +55,7 @@ typedef TMap<TWeakObjectPtr<UObject>, TArray<FInteractionAvailabilityEx>> FAvail
 //			{
 //				for each Actor[k][i]
 //				{
-//					Tick CharacterMovementComponent[k][i]
+//					Tick CharacterMovementComponent[k][i] in parallel
 //				}
 // 
 //				Tick Island[k].PreTickFunction (that eventually generates the trajectories with all the updated CMCs)
@@ -66,7 +67,7 @@ typedef TMap<TWeakObjectPtr<UObject>, TArray<FInteractionAvailabilityEx>> FAvail
 //				for each Actor[k][i]
 //				{
 //					if (i != 0)
-//						Tick SkeletalMeshComponent[k][i] (that DoSearch_AnyThread get the cached result calculated by Tick Actor[k][0].SkeletalMeshComponent)
+//						Tick SkeletalMeshComponent[k][i] (that DoSearch_AnyThread get the cached result calculated by Tick Actor[k][0].SkeletalMeshComponent) in parallel
 //				}
 //			}
 // - next frame UPoseSearchInteractionLibrary::MotionMatchInteraction_Pure(TArray<FPoseSearchInteractionAvailability> Availabilities, UObject* AnimInstance), with the context of all the published 
@@ -99,23 +100,27 @@ public:
 		const FPoseSearchContinuingProperties& ContinuingProperties, FPoseSearchInteractionBlueprintResult& Result,
 		FName PoseHistoryName, const FAnimNode_PoseSearchHistoryCollector_Base* HistoryCollector, bool bValidateResultAgainstAvailabilities);
 
-	// returns the associated FInteractionSearchContext to AnimContext. We do a copy to be thread safe.
-	// check FInteractionSearchContext::IsValid for its validity
-	UE::PoseSearch::FInteractionSearchContext FindSearchContext_AnyThread(UObject* AnimContext);
-
 private:
+#if ENABLE_ANIM_DEBUG
+	friend struct UE::PoseSearch::FInteractionValidator;
+#endif // ENABLE_ANIM_DEBUG
+
 	UE::PoseSearch::FInteractionIsland& CreateIsland();
 	UE::PoseSearch::FInteractionIsland& GetAvailableIsland();
 
 	void DestroyIsland(int32 Index);
 	void DestroyAllIslands();
 	void UninjectAllIslands();
+#if DO_CHECK
 	bool ValidateAllIslands() const;
+#endif // DO_CHECK
 
 	void PopulateContinuingProperties(UE::PoseSearch::FInteractionSearchContext& SearchContext, float DeltaSeconds) const;
 	UE::PoseSearch::FInteractionIsland* FindIsland(const UObject* AnimContext);
 	
+#if ENABLE_DRAW_DEBUG
 	void DebugDraw() const;
+#endif // ENABLE_DRAW_DEBUG
 
 	UE::PoseSearch::FAvailabilityRequestsMap AvailabilityRequestsMap;
 	FCriticalSection AvailabilityRequestsMapMutex;
@@ -125,4 +130,4 @@ private:
 
 	// critical section to retrieve the subsystem in a thread safe manner
 	static FCriticalSection RetrieveSubsystemMutex;
-};
\ No newline at end of file
+};
diff --git a/Engine/Plugins/Animation/PoseSearch/Source/Runtime/Public/PoseSearch/PoseSearchInteractionValidator.h b/Engine/Plugins/Animation/PoseSearch/Source/Runtime/Public/PoseSearch/PoseSearchInteractionValidator.h
new file mode 100644
--- /dev/null
+++ b/Engine/Plugins/Animation/PoseSearch/Source/Runtime/Public/PoseSearch/PoseSearchInteractionValidator.h
@@ -0,0 +1,35 @@
+// Copyright Epic Games, Inc. All Rights Reserved.
+
+#pragma once
+
+#include "Animation/AnimTypes.h"
+
+namespace UE::PoseSearch
+{
+#if ENABLE_ANIM_DEBUG
+struct FInteractionIsland;
+
+// Experimental, this feature might be removed without warning, not for production use
+struct FInteractionValidator
+{
+	enum class EPhase
+	{
+		MotionMatchingInteractionUpdate,
+		MotionMatchingInteractionEvaluate,
+		HistoryCollectorUpdate,
+		HistoryCollectorEvaluate
+	};
+
+	FInteractionValidator(UObject* AnimContext, EPhase Phase);
+	~FInteractionValidator();
+
+	const UObject* ValidatingAnimContext = nullptr;
+	FInteractionIsland* ValidatingIsland = nullptr;
+};
+
+#define CheckInteractionThreadSafety(AnimContext, Phase) FInteractionValidator(AnimContext, Phase)
+#else
+#define CheckInteractionThreadSafety(AnimContext, Phase)
+#endif // ENABLE_ANIM_DEBUG
+
+} // namespace UE::PoseSearch
\ No newline at end of file
