posesearch - motion matching interaction broadphase performance and threading improvements

[RN] skip

diff --git a/Engine/Plugins/Animation/PoseSearch/Source/Runtime/Private/PoseSearchInteractionIsland.cpp b/Engine/Plugins/Animation/PoseSearch/Source/Runtime/Private/PoseSearchInteractionIsland.cpp
--- a/Engine/Plugins/Animation/PoseSearch/Source/Runtime/Private/PoseSearchInteractionIsland.cpp
+++ b/Engine/Plugins/Animation/PoseSearch/Source/Runtime/Private/PoseSearchInteractionIsland.cpp
@@ -16,11 +16,6 @@
 namespace UE::PoseSearch
 {
 
-#if ENABLE_DRAW_DEBUG && ENABLE_ANIM_DEBUG
-static bool GVarPoseSearchInteractionShowIslands = false;
-static FAutoConsoleVariableRef CVarPoseSearchInteractionShowIslands(TEXT("a.PoseSearchInteraction.ShowIslands"), GVarPoseSearchInteractionShowIslands, TEXT("Show Pose Search Interaction Islands"));
-#endif // ENABLE_DRAW_DEBUG && ENABLE_ANIM_DEBUG
-
 static int32 GVarPoseSearchInteractionIslandResultPolicy = 1;
 static FAutoConsoleVariableRef CVarPoseSearchInteractionIslandResultPolicy(TEXT("a.PoseSearchInteraction.IslandResultPolicy"), GVarPoseSearchInteractionIslandResultPolicy, TEXT("0: only best result is valid\n1: assign from best to worst result"));
 
@@ -378,37 +373,6 @@ FInteractionIsland::~FInteractionIsland()
 	Uninject();
 }
 
-void FInteractionIsland::DebugDraw(const FColor& Color) const
-{
-	// called only by UPoseSearchInteractionSubsystem::Tick so no need to lock SearchResultsMutex to protect the read of SearchContexts
-#if ENABLE_DRAW_DEBUG && ENABLE_ANIM_DEBUG
-	check(IsInGameThread());
-
-	if (GVarPoseSearchInteractionShowIslands)
-	{
-		for (const FInteractionSearchContext& SearchContext : SearchContexts)
-		{
-			for (int32 Index = 0; Index < SearchContext.AnimContexts.Num(); ++Index)
-			{
-				if (const UObject* AnimContext = SearchContext.AnimContexts[Index].Get())
-				{
-					if (const UAnimInstance* AnimInstance = Cast<UAnimInstance>(AnimContext))
-					{
-						const FVector Position = AnimInstance->GetSkelMeshComponent()->GetComponentLocation();
-						const float BroadPhaseRadius = SearchContext.BroadPhaseRadiuses[Index];
-						DrawDebugCircle(AnimInstance->GetWorld(), Position, BroadPhaseRadius, 40, Color, false, 0.f, SDPG_Foreground, 0.f, FVector::XAxisVector, FVector::YAxisVector, false);
-					}
-					else
-					{
-						unimplemented();
-					}
-				}
-			}
-		}
-	}
-#endif // ENABLE_DRAW_DEBUG && ENABLE_ANIM_DEBUG
-}
-
 UActorComponent* FInteractionIsland::FindPreTickComponent(const AActor* Actor)
 {
 	if (UActorComponent* CharacterMovementComponent = Actor->GetComponentByClass<UCharacterMovementComponent>())
@@ -616,7 +580,9 @@ bool FInteractionIsland::DoSearch_AnyThread(UObject* AnimContext, const FPoseSea
 		// fix:		avoid UPoseSearchLibrary::MotionMatch calls wrapped by any lock, at the cost of eventually (by design should be NEVER) performing the searches twice
 		//			By design we should inject ticks dependencies (added by FInteractionIsland::InjectToActor via AddTickPrerequisiteComponent), so concurrently 
 		//			fly of FInteractionIsland within the same island that requires searches is forbidden
+#if UE_POSE_SEARCH_ENABLE_ISLAND_CRITICAL_SECTION
 		FScopeLock Lock(&SearchResultsMutex);
+#endif // UE_POSE_SEARCH_ENABLE_ISLAND_CRITICAL_SECTION
 		bDoPerfomSearch = !bSearchPerfomed;
 
 		if (bDoPerfomSearch)
@@ -640,6 +606,8 @@ bool FInteractionIsland::DoSearch_AnyThread(UObject* AnimContext, const FPoseSea
 
 	if (bDoPerfomSearch)
 	{
+		QUICK_SCOPE_CYCLE_COUNTER(STAT_UPoseSearchInteractionInteractionIsland_Search);
+
 		FMemMark Mark(FMemStack::Get());
 
 		TArray<const UObject*, TInlineAllocator<PreallocatedRolesNum, TMemStackAllocator<>>> AnimContexts;
@@ -689,7 +657,7 @@ bool FInteractionIsland::DoSearch_AnyThread(UObject* AnimContext, const FPoseSea
 
 			FPoseSearchContinuingProperties ContinuingPropertiesToUse = SearchContext.ContinuingProperties;
 
-			// @todo: implemennt this carefully, since it'll cause thread safety concerns depending on which actors is calling the query first 
+			// @todo: implement this carefully, since it'll cause thread safety concerns depending on which actors is calling the query first 
 			// and if their animations get integrated by different play rates
 			
 			//if (ContinuingProperties.PlayingAsset)
@@ -720,17 +688,14 @@ bool FInteractionIsland::DoSearch_AnyThread(UObject* AnimContext, const FPoseSea
 			//}
 
 			// @todo: we could perform multiple UPoseSearchLibrary::MotionMatch in parallel!
-			const FSearchResult PoseSearchResult = UPoseSearchLibrary::MotionMatch(AnimContexts, SearchContext.Roles,
+			PoseSearchResults[SearchIndex] = UPoseSearchLibrary::MotionMatch(AnimContexts, SearchContext.Roles,
 				PoseHistories, AssetsToSearch, ContinuingPropertiesToUse, PoseSearchFutureProperties);
-
-			if (PoseSearchResult.PoseCost < SearchContext.MaxCost)
-			{
-				PoseSearchResults[SearchIndex] = PoseSearchResult;
-			}
 		}
 		
+#if UE_POSE_SEARCH_ENABLE_ISLAND_CRITICAL_SECTION
 		// locking to update SearchResults and bSearchPerfomed
 		FScopeLock Lock(&SearchResultsMutex);
+#endif // UE_POSE_SEARCH_ENABLE_ISLAND_CRITICAL_SECTION
 
 		if (!bSearchPerfomed)
 		{
@@ -766,8 +731,10 @@ bool FInteractionIsland::GetResult_AnyThread(UObject* AnimContext, FPoseSearchIn
 {
 	check(AnimContext);
 
+#if UE_POSE_SEARCH_ENABLE_ISLAND_CRITICAL_SECTION
 	// locking to read SearchResults
 	FScopeLock Lock(&SearchResultsMutex);
+#endif // UE_POSE_SEARCH_ENABLE_ISLAND_CRITICAL_SECTION
 
 	// looking for AnimContext in SearchResults to fill up Result
 	for (const FInteractionSearchResult& SearchResult : SearchResults)
diff --git a/Engine/Plugins/Animation/PoseSearch/Source/Runtime/Private/PoseSearchInteractionSubsystem.cpp b/Engine/Plugins/Animation/PoseSearch/Source/Runtime/Private/PoseSearchInteractionSubsystem.cpp
--- a/Engine/Plugins/Animation/PoseSearch/Source/Runtime/Private/PoseSearchInteractionSubsystem.cpp
+++ b/Engine/Plugins/Animation/PoseSearch/Source/Runtime/Private/PoseSearchInteractionSubsystem.cpp
@@ -9,6 +9,21 @@
 
 namespace UE::PoseSearch
 {
+#if ENABLE_DRAW_DEBUG && ENABLE_ANIM_DEBUG
+static bool GVarPoseSearchInteractionShowIslands = false;
+static FAutoConsoleVariableRef CVarPoseSearchInteractionShowIslands(TEXT("a.PoseSearchInteraction.ShowIslands"), GVarPoseSearchInteractionShowIslands, TEXT("Show Pose Search Interaction Islands"));
+#endif // ENABLE_DRAW_DEBUG && ENABLE_ANIM_DEBUG
+
+static FVector GetContextLocation(const UObject* AnimContext)
+{
+	if (const UAnimInstance* AnimInstance = Cast<UAnimInstance>(AnimContext))
+	{
+		return AnimInstance->GetSkelMeshComponent()->GetComponentLocation();
+	}
+
+	unimplemented();
+	return FVector::ZeroVector;
+}
 
 FString FInteractionAvailabilityEx::GetPoseHistoryName() const
 {
@@ -76,151 +91,168 @@ struct FInteractionSearchContextGroup
 		}
 	}
 
-	TSet<const UObject*, DefaultKeyFuncs<const UObject*>, TInlineSetAllocator<PreallocatedRolesNum>> AnimContexts;
-	TArray<int32, TInlineAllocator<PreallocatedSearchesNum>> SearchContextsIndices;
+	TSet<const UObject*, DefaultKeyFuncs<const UObject*>, TInlineSetAllocator<16, TMemStackSetAllocator<>>> AnimContexts;
+	TArray<int32, TInlineAllocator<16, TMemStackAllocator<>>> SearchContextsIndices;
 };
 
+typedef TArray<FInteractionSearchContextGroup, TInlineAllocator<PreallocatedSearchesNum, TMemStackAllocator<>>> FInteractionSearchContextGroups;
+
 struct FRoledAnimContext
 {
+	FRoledAnimContext(const UObject* InAnimContext, const FRole InRole, const FAnimNode_PoseSearchHistoryCollector_Base* HistoryCollector)
+		: AnimContext(InAnimContext)
+		, Role(InRole)
+		, HistoryCollector(HistoryCollector)
+	{
+		check(AnimContext && HistoryCollector);
+	}
+
+	bool operator==(const FRoledAnimContext& Other) const = default;
+
 	const UObject* AnimContext = nullptr;
 	FRole Role = DefaultRole;
 	const FAnimNode_PoseSearchHistoryCollector_Base* HistoryCollector = nullptr;
-	float BroadPhaseRadius = 0.f;
-	float MaxCost = MAX_flt;
 };
 
-// @todo: inherit from a TSet if it gets too slow 
-// sorted array of FRoledAnimContext(s). Sorted by UAnimContext to prevent "alternating" tuples between different frames like
-// ([AnimContextA, RoleA], [AnimContextB, RoleB]) vs ([AnimContextB, RoleB], [AnimContextA, RoleA]) 
-// that ultimately represent the same search, but endign up with different SearchIDs
-struct FRoledAnimContexts
+typedef TArray<FRoledAnimContext, TInlineAllocator<PreallocatedRolesNum, TMemStackAllocator<>>> FRoledAnimContexts;
+typedef	TArray<const UObject*, TInlineAllocator<256, TMemStackAllocator<>>> FSortedAnimContexts;
+
+static void GenerateSortedAnimContexts(const FAvailabilityRequestsMap& AvailabilityRequestsMap, const UWorld* SubsystemWorld, FSortedAnimContexts& SortedAnimContexts)
 {
-	void AddRoledAnimContext(const UObject* AnimContext, FRole Role, const FAnimNode_PoseSearchHistoryCollector_Base* HistoryCollector, float BroadPhaseRadius, float MaxCost)
-	{
-		check(AnimContext && HistoryCollector);
+	QUICK_SCOPE_CYCLE_COUNTER(STAT_UPoseSearchInteractionSubsystem_GenerateSortedAnimContexts);
+
+	check(SubsystemWorld);
+
+	SortedAnimContexts.Reset(AvailabilityRequestsMap.Num());
 
-		for (FRoledAnimContext& RoledAnimContext : RoledAnimContexts)
+	for (FAvailabilityRequestsMap::TConstIterator Iter = AvailabilityRequestsMap.CreateConstIterator(); Iter; ++Iter)
+	{
+		// making sure the weak pointer is still valid
+		const UObject* AnimContext = Iter->Key.Get();
+		if (!AnimContext)
 		{
-			if (RoledAnimContext.AnimContext == AnimContext && RoledAnimContext.HistoryCollector == HistoryCollector && RoledAnimContext.Role == Role)
-			{
-				// we already found the FRoledAnimContexts. Updating the BroadPhaseRadius and MaxCost accordingly
-				RoledAnimContext.BroadPhaseRadius = FMath::Max(RoledAnimContext.BroadPhaseRadius, BroadPhaseRadius);
-					
-				// MaxCost is valid only if greater than zero
-				if (MaxCost > 0.f)
-				{
-					RoledAnimContext.MaxCost = FMath::Min(RoledAnimContext.MaxCost, MaxCost);
-				}
-				return;
-			}
+			UE_LOG(LogPoseSearch, Log, TEXT("ProcessAvailabilityRequestsMap null anim context. The associated character got removed from the world"));
 		}
-
-		// we couldn't find the FRoledAnimContexts. Let's add a new one
-		FRoledAnimContext& NewRoledAnimContext = RoledAnimContexts.AddDefaulted_GetRef();
-		NewRoledAnimContext.AnimContext = AnimContext;
-		NewRoledAnimContext.HistoryCollector = HistoryCollector;
-		NewRoledAnimContext.Role = Role;
-		NewRoledAnimContext.BroadPhaseRadius = BroadPhaseRadius;
-		if (MaxCost > 0.f)
+		else if (!AnimContext->GetWorld())
 		{
-			NewRoledAnimContext.MaxCost = MaxCost;
+			UE_LOG(LogPoseSearch, Error, TEXT("ProcessAvailabilityRequestsMap AnimContext %s is not in a world!"), *AnimContext->GetName());
 		}
-
-		// @todo: use a better datastructure to keep things FRoledAnimContext(s) sorted? Or use insertion sort?
-		RoledAnimContexts.Sort([](const FRoledAnimContext& A, const FRoledAnimContext& B)
+		else if (SubsystemWorld != AnimContext->GetWorld())
 		{
-			// we're sorting by AnimContext address in memory, since searches neeed to be consistent 
-			// between frames on the same machine, not between replicated machines
-			return B.AnimContext < A.AnimContext;
-		});
-	}
-
-	TArrayView<const FRoledAnimContext> GetDataView() const
-	{
-		return MakeArrayView(RoledAnimContexts.GetData(), RoledAnimContexts.Num());
+			UE_LOG(LogPoseSearch, Error, TEXT("ProcessAvailabilityRequestsMap AnimContext %s is from World %s, and supposed to be in World @s!"), *AnimContext->GetName(), *AnimContext->GetWorld()->GetName(), *SubsystemWorld->GetName());
+		}
+		else
+		{
+			SortedAnimContexts.Add(AnimContext);
+		}
 	}
+	SortedAnimContexts.Sort([](const UObject& AnimContextA, const UObject& AnimContextB)
+		{
+			// @todo: do we need to sort this reliably for networking?
+			return &AnimContextA < &AnimContextB;
+		});
+}
 
-private:
-	TArray<FRoledAnimContext, TInlineAllocator<PreallocatedRolesNum, TMemStackAllocator<>>> RoledAnimContexts;
+struct FAnimContextInfo
+{
+	FVector Location;
+	float MaxBroadPhaseRadiusSquared;
+	TArray<int32, TInlineAllocator<8, TMemStackAllocator<>>> LinkedAnimContextIndex;
 };
 
-typedef TMap<const UPoseSearchDatabase*, FRoledAnimContexts, TInlineSetAllocator<PreallocatedSearchesNum, TMemStackSetAllocator<>>> FDatabaseToRoledAnimContexts;
-typedef TPair<const UPoseSearchDatabase*, FRoledAnimContexts> FDatabaseToRoledAnimContextsPair;
+typedef TArray<FAnimContextInfo, TMemStackAllocator<>> FAnimContextInfos;
 
-static void ProcessAvailabilityRequestsMap(const FAvailabilityRequestsMap& AvailabilityRequestsMap, UWorld* SubsystemWorld, FDatabaseToRoledAnimContexts& DatabaseToRoledAnimContexts)
+static void GenerateAnimContextInfos(const FAvailabilityRequestsMap& AvailabilityRequestsMap, const FSortedAnimContexts& SortedAnimContexts, FAnimContextInfos& AnimContextInfos)
 {
-	for (FAvailabilityRequestsMap::TConstIterator Iter = AvailabilityRequestsMap.CreateConstIterator(); Iter; ++Iter)
+	QUICK_SCOPE_CYCLE_COUNTER(STAT_UPoseSearchInteractionSubsystem_GenerateAnimContextInfos);
+
+	const int32 NumSortedAnimContexts = SortedAnimContexts.Num();
+	AnimContextInfos.SetNum(NumSortedAnimContexts);
+	for (int32 AnimContextIndex = 0; AnimContextIndex < NumSortedAnimContexts; ++AnimContextIndex)
 	{
-		// making sure the weak pointer is still valid
-		UObject* AnimContext = Iter->Key.Get();
-		if (!AnimContext)
-		{
-			UE_LOG(LogPoseSearch, Log, TEXT("ProcessAvailabilityRequestsMap null anim context. The associated character got removed from the world"));
-			continue;
-		}
+		const UObject* AnimContext = SortedAnimContexts[AnimContextIndex];
+		check(AnimContext);
 
-		if (!AnimContext->GetWorld())
+		float MaxBroadPhaseRadius = 0.f;
+		for (const FInteractionAvailabilityEx& AvailabilityRequest : AvailabilityRequestsMap[AnimContext])
 		{
-			UE_LOG(LogPoseSearch, Error, TEXT("ProcessAvailabilityRequestsMap AnimContext %s is not in a world!"), *AnimContext->GetName());
-			continue;
+			MaxBroadPhaseRadius = FMath::Max(MaxBroadPhaseRadius, AvailabilityRequest.BroadPhaseRadius);
 		}
 
-		if (SubsystemWorld != AnimContext->GetWorld())
-		{
-			UE_LOG(LogPoseSearch, Error, TEXT("ProcessAvailabilityRequestsMap AnimContext %s is from World %s, and supposed to be in World @s!"), *AnimContext->GetName(), *AnimContext->GetWorld()->GetName(), *SubsystemWorld->GetName());
-			continue;
-		}
+		AnimContextInfos[AnimContextIndex].MaxBroadPhaseRadiusSquared = MaxBroadPhaseRadius * MaxBroadPhaseRadius;
+		AnimContextInfos[AnimContextIndex].Location = GetContextLocation(AnimContext);
+	}
 
-		for (const FInteractionAvailabilityEx& AvailabilityRequest : Iter->Value)
+	// solving the broad phase using the AnimContextInfos
+	for (int32 AnimContextIndexA = 0; AnimContextIndexA < NumSortedAnimContexts; ++AnimContextIndexA)
+	{
+		for (int32 AnimContextIndexB = AnimContextIndexA + 1; AnimContextIndexB < NumSortedAnimContexts; ++AnimContextIndexB)
 		{
-			if (!AvailabilityRequest.Database)
+			const FVector DeltaLocation = AnimContextInfos[AnimContextIndexA].Location - AnimContextInfos[AnimContextIndexB].Location;
+			const float MaxBroadPhaseRadiusSquared = FMath::Max(AnimContextInfos[AnimContextIndexA].MaxBroadPhaseRadiusSquared, AnimContextInfos[AnimContextIndexB].MaxBroadPhaseRadiusSquared);
+			if (DeltaLocation.SquaredLength() <= MaxBroadPhaseRadiusSquared)
 			{
-				UE_LOG(LogPoseSearch, Log, TEXT("ProcessAvailabilityRequestsMap null AvailabilityRequest.Database"));
-				continue;
+				// AnimContext can potentially interact with AnimContextB: linking AnimContext to AnimContextB and vice versa
+				AnimContextInfos[AnimContextIndexA].LinkedAnimContextIndex.Add(AnimContextIndexB);
+				AnimContextInfos[AnimContextIndexB].LinkedAnimContextIndex.Add(AnimContextIndexA);
 			}
+		}
+	}
+}
 
-			if (!AvailabilityRequest.Database->Schema)
+struct FAnimContextInfoVisitor
+{
+	FAnimContextInfoVisitor(
+		const FSortedAnimContexts& InSortedAnimContexts,
+		const FAnimContextInfos& InAnimContextInfos,
+		TFunctionRef<void(const UObject*)> OnNewAnimContextFound,
+		TFunctionRef<void()> OnDoneGroupingAnimContexts)
+		: SortedAnimContexts(InSortedAnimContexts)
+		, AnimContextInfos(InAnimContextInfos)
+	{
+		const int32 NumSortedAnimContexts = AnimContextInfos.Num();
+		for (int32 AnimContextIndex = 0; AnimContextIndex < NumSortedAnimContexts; ++AnimContextIndex)
+		{
+			if (!AnimContextInfos[AnimContextIndex].LinkedAnimContextIndex.IsEmpty() && !VisitedIndexes.Find(AnimContextIndex))
 			{
-				UE_LOG(LogPoseSearch, Error, TEXT("ProcessAvailabilityRequestsMap null Schema for Database %s"), *AvailabilityRequest.Database->GetName());
-				continue;
-			}
+				// starting the evaluation of a new set of grouped AnimContext(s)
 
-			const FAnimNode_PoseSearchHistoryCollector_Base* HistoryCollector = AvailabilityRequest.GetHistoryCollector(AnimContext);
-			if (!HistoryCollector)
-			{
-				UE_LOG(LogPoseSearch, Error, TEXT("ProcessAvailabilityRequestsMap couldn't find PoseHistory %s for AnimContext %s"), *AvailabilityRequest.GetPoseHistoryName(), *AnimContext->GetName());
-				continue;
+				// processing the AvailabilityRequestsMap of the current AnimContextArray to fill up a map of Databases pointing to an array of all the AnimContexts with related roles
+				VisitRecursively(AnimContextIndex, OnNewAnimContextFound);
+
+				OnDoneGroupingAnimContexts();
 			}
+		}
+	}
+
+private:
+	void VisitRecursively(int32 AnimContextIndex, TFunctionRef<void(const UObject*)> OnNewAnimContextFound)
+	{
+		check(!AnimContextInfos[AnimContextIndex].LinkedAnimContextIndex.IsEmpty());
 
-			FRoledAnimContexts& RoledAnimContexts = DatabaseToRoledAnimContexts.FindOrAdd(AvailabilityRequest.Database);
-			if (AvailabilityRequest.RolesFilter.IsEmpty())
+		bool bIsAlreadyInSet;
+		VisitedIndexes.FindOrAdd(AnimContextIndex, &bIsAlreadyInSet);
+
+		if (!bIsAlreadyInSet)
+		{
+			OnNewAnimContextFound(SortedAnimContexts[AnimContextIndex]);
+
+			for (int32 LinkedAnimContextIndex : AnimContextInfos[AnimContextIndex].LinkedAnimContextIndex)
 			{
-				// adding ALL the possible roles from the database:
-				for (const FPoseSearchRoledSkeleton& RoledSkeleton : AvailabilityRequest.Database->Schema->GetRoledSkeletons())
-				{
-					RoledAnimContexts.AddRoledAnimContext(AnimContext, RoledSkeleton.Role, HistoryCollector, AvailabilityRequest.BroadPhaseRadius, AvailabilityRequest.MaxCost);
-				}
-			}
-			else
-			{
-				for (const FRole& Role : AvailabilityRequest.RolesFilter)
-				{
-					if (AvailabilityRequest.Database->Schema->GetRoledSkeleton(Role))
-					{
-						RoledAnimContexts.AddRoledAnimContext(AnimContext, Role, HistoryCollector, AvailabilityRequest.BroadPhaseRadius, AvailabilityRequest.MaxCost);
-					}
-					else
-					{
-						UE_LOG(LogPoseSearch, Warning, TEXT("ProcessAvailabilityRequestsMap unsupported Role %s for Database %s"), *Role.ToString(), *AvailabilityRequest.Database->GetName());
-					}
-				}
+				VisitRecursively(LinkedAnimContextIndex, OnNewAnimContextFound);
 			}
 		}
 	}
-}
+
+	const FSortedAnimContexts& SortedAnimContexts;
+	const FAnimContextInfos& AnimContextInfos;
+
+	typedef TSet<int32, DefaultKeyFuncs<int32>, TInlineSetAllocator<32, TMemStackSetAllocator<>>> FVisitedIndexes;
+	FVisitedIndexes VisitedIndexes;
+};
 
 template <typename DataType, typename EvaluateCombinationType>
-static void GeneratePermutationsRecursive(const TArrayView<DataType> Data, int32 DataIndex, TArrayView<int32> Combination, int32 CombinationIndex, EvaluateCombinationType EvaluateCombination)
+static void GeneratePermutationsRecursive(const TConstArrayView<DataType> Data, int32 DataIndex, TArrayView<int32> Combination, int32 CombinationIndex, EvaluateCombinationType EvaluateCombination)
 {
 	if (CombinationIndex == Combination.Num())
 	{ 
@@ -235,16 +267,202 @@ static void GeneratePermutationsRecursive(const TArrayView<DataType> Data, int32
 } 
 
 template <typename DataType, typename EvaluateCombinationType>
-static void GeneratePermutations(const TArrayView<DataType> Data, int32 CombinationCardinality, EvaluateCombinationType EvaluateCombination)
+static void GeneratePermutations(const TConstArrayView<DataType> Data, int32 CombinationCardinality, EvaluateCombinationType EvaluateCombination)
 {
 	TArray<int32, TInlineAllocator<PreallocatedRolesNum>> Combination;
 	Combination.SetNum(CombinationCardinality);
 	GeneratePermutationsRecursive(Data, 0, Combination, 0, EvaluateCombination);
 }
 
+typedef TArray<FInteractionSearchContext, TInlineAllocator<PreallocatedSearchesNum, TMemStackAllocator<>>> FInteractionSearchContexts;
+
+static void GenerateSearchContexts(const FAvailabilityRequestsMap& AvailabilityRequestsMap, const FSortedAnimContexts& SortedAnimContexts, const FAnimContextInfos& AnimContextInfos, FInteractionSearchContexts& SearchContexts)
+{
+	QUICK_SCOPE_CYCLE_COUNTER(STAT_UPoseSearchInteractionSubsystem_GenerateSearchContexts);
+
+	typedef TMap<const UPoseSearchDatabase*, FRoledAnimContexts, TInlineSetAllocator<PreallocatedSearchesNum, TMemStackSetAllocator<>>> FDatabaseToRoledAnimContexts;
+	typedef TPair<const UPoseSearchDatabase*, FRoledAnimContexts> FDatabaseToRoledAnimContextsPair;
+	FDatabaseToRoledAnimContexts DatabaseToRoledAnimContexts;
+
+	// visiting ALL the AnimContexts from SortedAnimContexts with their cached information hold in AnimContextInfos to figure our potential interactions
+	FAnimContextInfoVisitor AnimContextInfoVisitor(SortedAnimContexts, AnimContextInfos,
+		// OnNewAnimContextFound: called when the FAnimContextInfoVisitor finds a new AnimContext that can be grouped in the current group of possibly interacting AnimContext(s)
+		[&AvailabilityRequestsMap, &DatabaseToRoledAnimContexts](const UObject* AnimContext)
+		{
+			// analyzing all the AvailabilityRequest(s) associated with this AnimContext and eventually generate the associated FRoledAnimContexts,
+			// inserted in a per database sorted data structure (DatabaseToRoledAnimContexts)
+			for (const FInteractionAvailabilityEx& AvailabilityRequest : AvailabilityRequestsMap[AnimContext])
+			{
+				if (!AvailabilityRequest.Database)
+				{
+					UE_LOG(LogPoseSearch, Log, TEXT("ProcessAvailabilityRequestsMap null AvailabilityRequest.Database"));
+					continue;
+				}
+
+				if (!AvailabilityRequest.Database->Schema)
+				{
+					UE_LOG(LogPoseSearch, Error, TEXT("ProcessAvailabilityRequestsMap null Schema for Database %s"), *AvailabilityRequest.Database->GetName());
+					continue;
+				}
+
+				const FAnimNode_PoseSearchHistoryCollector_Base* HistoryCollector = AvailabilityRequest.GetHistoryCollector(AnimContext);
+				if (!HistoryCollector)
+				{
+					UE_LOG(LogPoseSearch, Error, TEXT("ProcessAvailabilityRequestsMap couldn't find PoseHistory %s for AnimContext %s"), *AvailabilityRequest.GetPoseHistoryName(), *AnimContext->GetName());
+					continue;
+				}
+
+				FRoledAnimContexts& RoledAnimContexts = DatabaseToRoledAnimContexts.FindOrAdd(AvailabilityRequest.Database);
+				if (AvailabilityRequest.RolesFilter.IsEmpty())
+				{
+					// adding ALL the possible roles from the database:
+					for (const FPoseSearchRoledSkeleton& RoledSkeleton : AvailabilityRequest.Database->Schema->GetRoledSkeletons())
+					{
+						RoledAnimContexts.AddUnique(FRoledAnimContext(AnimContext, RoledSkeleton.Role, HistoryCollector));
+					}
+				}
+				else
+				{
+					for (const FRole& Role : AvailabilityRequest.RolesFilter)
+					{
+						if (AvailabilityRequest.Database->Schema->GetRoledSkeleton(Role))
+						{
+							RoledAnimContexts.AddUnique(FRoledAnimContext(AnimContext, Role, HistoryCollector));
+						}
+						else
+						{
+							UE_LOG(LogPoseSearch, Warning, TEXT("ProcessAvailabilityRequestsMap unsupported Role %s for Database %s"), *Role.ToString(), *AvailabilityRequest.Database->GetName());
+						}
+					}
+				}
+			}
+		},
+		// OnDoneGroupingAnimContexts: called when the FAnimContextInfoVisitor reaches the end of a group of possibly interacting AnimContext(s)
+		[&DatabaseToRoledAnimContexts, &SearchContexts]()
+		{
+			// for each database now we try to create all the possible permutations of the roled anim instances
+			// for example, given a database set up with assets for 2 characters interactions with roles RoleA and RoleB
+			// and 2 anim instances, all of them willing to partecipate in the 2 characters interaction with both roles RoleA and RoleB:
+			// CharA could be taking RoleA and RoleB,
+			// CharB could be taking RoleA and RoleB,
+			// we generate all the permutations from the array of options:
+			// CharA/RoleA, CharA/RoleB, CharB/RoleA, CharB/RoleB
+			//
+			// and we prune the invalid tuples:
+			//
+			// CharA/RoleA - CharA/RoleB -> invalid because of same CharA
+			// CharA/RoleA - CharB/RoleA -> invalid because of same RoleA
+			// CharA/RoleA - CharB/RoleB -> VALID!
+			//
+			// CharA/RoleB - CharB/RoleA -> VALID!
+			// CharA/RoleB - CharB/RoleB -> invalid because of same RoleB
+			//
+			// CharB/RoleA - CharB/RoleB -> invalid because of same CharB
+
+			for (const FDatabaseToRoledAnimContextsPair& DatabaseToRoledAnimContextsPair : DatabaseToRoledAnimContexts)
+			{
+				const UPoseSearchDatabase* Database = DatabaseToRoledAnimContextsPair.Key;
+				const TArray<FPoseSearchRoledSkeleton>& RoledSkeletons = Database->Schema->GetRoledSkeletons();
+				const int32 CombinationCardinality = RoledSkeletons.Num();
+				const TConstArrayView<FRoledAnimContext> RoledAnimContextsView = MakeArrayView(DatabaseToRoledAnimContextsPair.Value.GetData(), DatabaseToRoledAnimContextsPair.Value.Num());
+
+				GeneratePermutations(RoledAnimContextsView, CombinationCardinality,
+					[Database, &RoledSkeletons, &SearchContexts](const TConstArrayView<FRoledAnimContext> RoledAnimContexts, const TConstArrayView<int32> Combination)
+					{
+						// checking if this combination is valid for the Database:
+						TSet<const UObject*, DefaultKeyFuncs<const UObject*>, TInlineSetAllocator<PreallocatedRolesNum, TMemStackSetAllocator<>>> AnimContexts;
+						TSet<FRole, DefaultKeyFuncs<FRole>, TInlineSetAllocator<PreallocatedRolesNum, TMemStackSetAllocator<>>> CoveredRoles;
+
+						for (int32 CombinationIndex = 0; CombinationIndex < Combination.Num(); ++CombinationIndex)
+						{
+							const FRoledAnimContext& RoledAnimContext = RoledAnimContexts[Combination[CombinationIndex]];
+
+							bool bIsAlreadyInSet = false;
+							AnimContexts.Add(RoledAnimContext.AnimContext, &bIsAlreadyInSet);
+
+							// is there any duplicated AnimContext?
+							if (bIsAlreadyInSet)
+							{
+								return false;
+							}
+
+							CoveredRoles.Add(RoledAnimContext.Role);
+						}
+
+						// does it cover all the roles?
+						for (const FPoseSearchRoledSkeleton& RoledSkeleton : RoledSkeletons)
+						{
+							if (!CoveredRoles.Find(RoledSkeleton.Role))
+							{
+								return false;
+							}
+						}
+
+						FInteractionSearchContext& SearchContext = SearchContexts.AddDefaulted_GetRef();
+						SearchContext.Database = Database;
+						SearchContext.AnimContexts.Reserve(Combination.Num());
+						SearchContext.HistoryCollectors.Reserve(Combination.Num());
+						SearchContext.Roles.Reserve(Combination.Num());
+
+						for (int32 CombinationIndex = 0; CombinationIndex < Combination.Num(); ++CombinationIndex)
+						{
+							const FRoledAnimContext& RoledAnimContext = RoledAnimContexts[Combination[CombinationIndex]];
+							SearchContext.AnimContexts.Add(RoledAnimContext.AnimContext);
+							SearchContext.HistoryCollectors.Add(RoledAnimContext.HistoryCollector);
+							SearchContext.Roles.Add(RoledAnimContext.Role);
+						}
+
+						return true;
+					});
+			}
+
+			// done using DatabaseToRoledAnimContexts. clearing up for the next group of AnimContext(s)
+			DatabaseToRoledAnimContexts.Reset();
+		});
+}
+
+static void GenerateInteractionSearchContextGroups(const FInteractionSearchContexts& SearchContexts, FInteractionSearchContextGroups& SearchContextGroups)
+{
+	QUICK_SCOPE_CYCLE_COUNTER(STAT_UPoseSearchInteractionSubsystem_GenerateInteractionSearchContextGroups);
+
+	for (int32 SearchContextIndex = 0; SearchContextIndex < SearchContexts.Num(); ++SearchContextIndex)
+	{
+		const FInteractionSearchContext& SearchContext = SearchContexts[SearchContextIndex];
+
+		int32 MainSearchContextGroupIndex = INDEX_NONE;
+		for (int32 SearchContextGroupIndex = 0; SearchContextGroupIndex < SearchContextGroups.Num();)
+		{
+			if (SearchContextGroups[SearchContextGroupIndex].Contains(SearchContext))
+			{
+				if (MainSearchContextGroupIndex == INDEX_NONE)
+				{
+					MainSearchContextGroupIndex = SearchContextGroupIndex;
+					SearchContextGroups[MainSearchContextGroupIndex].Add(SearchContext, SearchContextIndex);
+					++SearchContextGroupIndex;
+				}
+				else
+				{
+					SearchContextGroups[MainSearchContextGroupIndex].Merge(SearchContextGroups[SearchContextGroupIndex]);
+					SearchContextGroups.RemoveAt(SearchContextGroupIndex);
+				}
+			}
+			else
+			{
+				++SearchContextGroupIndex;
+			}
+		}
+		if (MainSearchContextGroupIndex == INDEX_NONE)
+		{
+			SearchContextGroups.AddDefaulted_GetRef().Add(SearchContext, SearchContextIndex);
+		}
+	}
+}
+
 } // namespace UE::PoseSearch
 
+#if UE_POSE_SEARCH_ENABLE_INTERACTION_SUBSYSTEM_CRITICAL_SECTION
 FCriticalSection UPoseSearchInteractionSubsystem::RetrieveSubsystemMutex;
+#endif // UE_POSE_SEARCH_ENABLE_INTERACTION_SUBSYSTEM_CRITICAL_SECTION
 
 UE::PoseSearch::FInteractionIsland& UPoseSearchInteractionSubsystem::CreateIsland()
 {
@@ -292,6 +510,30 @@ void UPoseSearchInteractionSubsystem::UninjectAllIslands()
 	}
 }
 
+void UPoseSearchInteractionSubsystem::RegenerateAllIslands(TConstArrayView<UE::PoseSearch::FInteractionSearchContextGroup> SearchContextGroups, TConstArrayView<UE::PoseSearch::FInteractionSearchContext> SearchContexts)
+{
+	QUICK_SCOPE_CYCLE_COUNTER(STAT_UPoseSearchInteractionSubsystem_RegenerateAllIslands);
+
+	using namespace UE::PoseSearch;
+
+	UninjectAllIslands();
+
+	for (const FInteractionSearchContextGroup& SearchContextGroup : SearchContextGroups)
+	{
+		FInteractionIsland& Island = GetAvailableIsland();
+
+		for (const UObject* AnimContext : SearchContextGroup.AnimContexts)
+		{
+			Island.InjectToActor(AnimContext);
+		}
+
+		for (int32 SearchContextsIndex : SearchContextGroup.SearchContextsIndices)
+		{
+			Island.AddSearchContext(SearchContexts[SearchContextsIndex]);
+		}
+	}
+}
+
 #if DO_CHECK
 bool UPoseSearchInteractionSubsystem::ValidateAllIslands() const
 {
@@ -362,30 +604,35 @@ bool UPoseSearchInteractionSubsystem::ValidateAllIslands() const
 }
 #endif // DO_CHECK
 
-void UPoseSearchInteractionSubsystem::PopulateContinuingProperties(UE::PoseSearch::FInteractionSearchContext& SearchContext, float DeltaSeconds) const
+void UPoseSearchInteractionSubsystem::PopulateContinuingProperties(float DeltaSeconds, TArrayView<UE::PoseSearch::FInteractionSearchContext> SearchContexts) const
 {
+	QUICK_SCOPE_CYCLE_COUNTER(STAT_UPoseSearchInteractionSubsystem_PopulateContinuingProperties);
+
 	using namespace UE::PoseSearch;
 
 	check(IsInGameThread());
 
-	// searching this SearchContext in all the islands to initialize its continuing pose
-	for (const FInteractionIsland* Island : Islands)
+	for (FInteractionSearchContext& SearchContext : SearchContexts)
 	{
-		if (const FSearchResult* SearchResult = Island->FindSearchResult(SearchContext))
+		// searching this SearchContext in all the islands to initialize its continuing pose
+		for (const FInteractionIsland* Island : Islands)
 		{
-			// is still valid...
-			if (SearchResult->IsValid())
+			if (const FSearchResult* SearchResult = Island->FindSearchResult(SearchContext))
 			{
-				if (const UE::PoseSearch::FSearchIndexAsset* SearchIndexAsset = SearchResult->GetSearchIndexAsset())
+				// is still valid...
+				if (SearchResult->IsValid())
 				{
-					if (const FPoseSearchDatabaseAnimationAssetBase* DatabaseAsset = SearchResult->Database->GetDatabaseAnimationAsset<FPoseSearchDatabaseAnimationAssetBase>(*SearchIndexAsset))
+					if (const UE::PoseSearch::FSearchIndexAsset* SearchIndexAsset = SearchResult->GetSearchIndexAsset())
 					{
-						SearchContext.ContinuingProperties.PlayingAsset = DatabaseAsset->GetAnimationAsset();
-						SearchContext.ContinuingProperties.PlayingAssetAccumulatedTime = SearchResult->AssetTime + DeltaSeconds;
+						if (const FPoseSearchDatabaseAnimationAssetBase* DatabaseAsset = SearchResult->Database->GetDatabaseAnimationAsset<FPoseSearchDatabaseAnimationAssetBase>(*SearchIndexAsset))
+						{
+							SearchContext.ContinuingProperties.PlayingAsset = DatabaseAsset->GetAnimationAsset();
+							SearchContext.ContinuingProperties.PlayingAssetAccumulatedTime = SearchResult->AssetTime + DeltaSeconds;
+						}
 					}
 				}
+				break;
 			}
-			break;
 		}
 	}
 }
@@ -429,9 +676,12 @@ UPoseSearchInteractionSubsystem* UPoseSearchInteractionSubsystem::GetSubsystem_A
 	{
 		if (UWorld* World = AnimContext->GetWorld())
 		{
+#if UE_POSE_SEARCH_ENABLE_INTERACTION_SUBSYSTEM_CRITICAL_SECTION
 			FScopeLock Lock(&RetrieveSubsystemMutex);
+#endif // UE_POSE_SEARCH_ENABLE_INTERACTION_SUBSYSTEM_CRITICAL_SECTION
 
-			// don't create a subsystem from any thread
+			// We expect the subsystem to be already created from the GameThread.
+			// We don't create the subsystem from any thread
 			if (World->HasSubsystem<UPoseSearchInteractionSubsystem>())
 			{
 				return World->GetSubsystem<UPoseSearchInteractionSubsystem>();
@@ -443,6 +693,7 @@ UPoseSearchInteractionSubsystem* UPoseSearchInteractionSubsystem::GetSubsystem_A
 
 void UPoseSearchInteractionSubsystem::Initialize(FSubsystemCollectionBase& Collection)
 {
+	check(IsInGameThread());
 	Super::Initialize(Collection);
 }
 
@@ -453,19 +704,10 @@ void UPoseSearchInteractionSubsystem::Deinitialize()
 	Super::Deinitialize();
 }
 
-static FVector GetContextLocation(const UObject* AnimContext)
-{
-	if (const UAnimInstance* AnimInstance = Cast<UAnimInstance>(AnimContext))
-	{
-		return AnimInstance->GetOwningActor()->GetActorLocation();
-	}
-
-	unimplemented();
-	return FVector::ZeroVector;
-}
-
 void UPoseSearchInteractionSubsystem::Tick(float DeltaSeconds)
 {
+	QUICK_SCOPE_CYCLE_COUNTER(STAT_UPoseSearchInteractionSubsystem_Tick);
+
 	using namespace UE::PoseSearch;
 
 	Super::Tick(DeltaSeconds);
@@ -477,188 +719,37 @@ void UPoseSearchInteractionSubsystem::Tick(float DeltaSeconds)
 	}
 
 	check(IsInGameThread());
-	check(GetWorld());
 
 	// FScopeLock Lock(&AvailabilityRequestsMapMutex); is not necessary since UPoseSearchInteractionSubsystem gets ticked outside the parallel animation jobs
 	FMemMark Mark(FMemStack::Get());
-	
-	// processing the AvailabilityRequestsMap to fill up a map of Databases pointing to an array of all the AnimContexts with related roles
-	FDatabaseToRoledAnimContexts DatabaseToRoledAnimContexts;
-	ProcessAvailabilityRequestsMap(AvailabilityRequestsMap, GetWorld(), DatabaseToRoledAnimContexts);
-
-	// for each database now we try to create all the possible permutations of the roled anim instances
-	// for example, given a database set up with assets for 2 characters interactions with roles RoleA and RoleB
-	// and 2 anim instances, all of them willing to partecipate in the 2 characters interaction with both roles RoleA and RoleB:
-	// CharA could be taking RoleA and RoleB,
-	// CharB could be taking RoleA and RoleB,
-	// we generate all the permutations from the array of options:
-	// CharA/RoleA, CharA/RoleB, CharB/RoleA, CharB/RoleB
-	//
-	// and we prune the invalid tuples:
-	//
-	// CharA/RoleA - CharA/RoleB -> invalid because of same CharA
-	// CharA/RoleA - CharB/RoleA -> invalid because of same RoleA
-	// CharA/RoleA - CharB/RoleB -> VALID!
-	//
-	// CharA/RoleB - CharB/RoleA -> VALID!
-	// CharA/RoleB - CharB/RoleB -> invalid because of same RoleB
-	//
-	// CharB/RoleA - CharB/RoleB -> invalid because of same CharB
-
-	TSet<const UObject*, DefaultKeyFuncs<const UObject*>, TInlineSetAllocator<PreallocatedRolesNum>> AnimContexts;
-	TSet<FRole, DefaultKeyFuncs<FRole>, TInlineSetAllocator<PreallocatedRolesNum>> CoveredRoles;
-	TArray<FInteractionSearchContext, TInlineAllocator<PreallocatedSearchesNum, TMemStackAllocator<>>> SearchContexts;
-	for (const FDatabaseToRoledAnimContextsPair& DatabaseToRoledAnimContextsPair : DatabaseToRoledAnimContexts)
-	{
-		const UPoseSearchDatabase* Database = DatabaseToRoledAnimContextsPair.Key;
-		const TArray<FPoseSearchRoledSkeleton>& RoledSkeletons = Database->Schema->GetRoledSkeletons();
-		const int32 CombinationCardinality = RoledSkeletons.Num();
-		
-		// @todo: naive implementation: optimize this code!
-		GeneratePermutations(DatabaseToRoledAnimContextsPair.Value.GetDataView(), CombinationCardinality, 
-			[&AnimContexts, &CoveredRoles, Database, &RoledSkeletons, &SearchContexts](const TArrayView<const FRoledAnimContext> RoledAnimContexts, const TArrayView<int32> Combination)
-			{
-				// checking if this combination is valid for the Database:
-				AnimContexts.Reset();
-				CoveredRoles.Reset();
 
-				for (int32 CombinationIndex = 0; CombinationIndex < Combination.Num(); ++CombinationIndex)
-				{
-					const FRoledAnimContext& RoledAnimContext = RoledAnimContexts[Combination[CombinationIndex]];
-
-					bool bIsAlreadyInSet = false;
-					AnimContexts.Add(RoledAnimContext.AnimContext, &bIsAlreadyInSet);
-
-					// is there any duplicated AnimContext?
-					if (bIsAlreadyInSet)
-					{
-						return false;
-					}
-
-					CoveredRoles.Add(RoledAnimContext.Role);
-				}
-
-				// does it cover all the roles?
-				for (const FPoseSearchRoledSkeleton& RoledSkeleton : RoledSkeletons)
-				{
-					if (!CoveredRoles.Find(RoledSkeleton.Role))
-					{
-						return false;
-					}
-				}
-
-				// @todo: naive implementation: optimize this code!
-				// checking if all the actors are within BroadPhaseRadius to each other
-				for (int32 CombinationIndexA = 0; CombinationIndexA < Combination.Num(); ++CombinationIndexA)
-				{
-					const FRoledAnimContext& RoledAnimContextA = RoledAnimContexts[Combination[CombinationIndexA]];
-					const FVector ActorLocationA = GetContextLocation(RoledAnimContextA.AnimContext);
+	// collecting and sorting all the AnimContext(s) from the AvailabilityRequestsMap, to have a stable and deterministic island creation
+	FSortedAnimContexts SortedAnimContexts;
+	GenerateSortedAnimContexts(AvailabilityRequestsMap, GetWorld(), SortedAnimContexts);
 
-					for (int32 CombinationIndexB = 0; CombinationIndexB < Combination.Num(); ++CombinationIndexB)
-					{
-						if (CombinationIndexA != CombinationIndexB)
-						{
-							const FRoledAnimContext& RoledAnimContextB = RoledAnimContexts[Combination[CombinationIndexB]];
-							const FVector ActorLocationB = GetContextLocation(RoledAnimContextB.AnimContext);
-
-							const FVector DeltaLocation = ActorLocationA - ActorLocationB;
-							if (DeltaLocation.Length() > RoledAnimContextA.BroadPhaseRadius)
-							{
-								return false;
-							}
-						}
-					}
-				}
+	// caching AnimContext(s) locations, max broad phase radiuses (squared) and collect relations of possible interactions between AnimContext(s)
+	// (stored in FAnimContextInfo::LinkedAnimContextIndex) as fast broad phase evaluation refined later on during FInteractionSearchContexts generation
+	FAnimContextInfos AnimContextInfos;
+	GenerateAnimContextInfos(AvailabilityRequestsMap, SortedAnimContexts, AnimContextInfos);
 
-				FInteractionSearchContext& SearchContext = SearchContexts.AddDefaulted_GetRef();
-				SearchContext.Database = Database;
-				SearchContext.AnimContexts.Reserve(Combination.Num());
-				SearchContext.HistoryCollectors.Reserve(Combination.Num());
-				SearchContext.Roles.Reserve(Combination.Num());
-				
-#if ENABLE_DRAW_DEBUG
-				SearchContext.BroadPhaseRadiuses.Reserve(Combination.Num());
-#endif //ENABLE_DRAW_DEBUG
-
-				for (int32 CombinationIndex = 0; CombinationIndex < Combination.Num(); ++CombinationIndex)
-				{
-					const FRoledAnimContext& RoledAnimContext = RoledAnimContexts[Combination[CombinationIndex]];
-					SearchContext.AnimContexts.Add(RoledAnimContext.AnimContext);
-					SearchContext.HistoryCollectors.Add(RoledAnimContext.HistoryCollector);
-					SearchContext.Roles.Add(RoledAnimContext.Role);
-#if ENABLE_DRAW_DEBUG
-					SearchContext.BroadPhaseRadiuses.Add(RoledAnimContext.BroadPhaseRadius);
-#endif //ENABLE_DRAW_DEBUG
-
-					// search is valid only if cost result is lower than any AvailabilityRequest.MaxCost (copied in RoledAnimContext.MaxCost)
-					SearchContext.MaxCost = FMath::Min(SearchContext.MaxCost, RoledAnimContext.MaxCost);
-				}
+	// generating all the possible interaction tuples of AnimContext(s) with roles and pose histories (defined in FInteractionSearchContext)
+	FInteractionSearchContexts SearchContexts;
+	GenerateSearchContexts(AvailabilityRequestsMap, SortedAnimContexts, AnimContextInfos, SearchContexts);
 
-				return true;
-			});
-	}
+	// populating the continuing pose properties for the SearchContexts from the current Islands
+	PopulateContinuingProperties(DeltaSeconds, SearchContexts);
 	
-	// for each valid SearchContexts we try to figure out the continuing pose properties from the current Islands
-	for (FInteractionSearchContext& SearchContext : SearchContexts)
-	{
-		PopulateContinuingProperties(SearchContext, DeltaSeconds);
-	}
-
-	// grouping SearchContexts in different Islands
-	TArray<FInteractionSearchContextGroup, TInlineAllocator<PreallocatedSearchesNum, TMemStackAllocator<>>> SearchContextGroups;
-	for (int32 SearchContextIndex = 0; SearchContextIndex < SearchContexts.Num(); ++SearchContextIndex)
-	{
-		const FInteractionSearchContext& SearchContext = SearchContexts[SearchContextIndex];
+	// generating groups of AnimContext(s) from the SearchContexts to be able to generate islands later on
+	FInteractionSearchContextGroups SearchContextGroups;
+	GenerateInteractionSearchContextGroups(SearchContexts, SearchContextGroups);
 
-		int32 MainSearchContextGroupIndex = INDEX_NONE;
-		for (int32 SearchContextGroupIndex = 0; SearchContextGroupIndex < SearchContextGroups.Num();)
-		{
-			if (SearchContextGroups[SearchContextGroupIndex].Contains(SearchContext))
-			{
-				if (MainSearchContextGroupIndex == INDEX_NONE)
-				{
-					MainSearchContextGroupIndex = SearchContextGroupIndex;
-					SearchContextGroups[MainSearchContextGroupIndex].Add(SearchContext, SearchContextIndex);
-					++SearchContextGroupIndex;
-				}
-				else
-				{
-					SearchContextGroups[MainSearchContextGroupIndex].Merge(SearchContextGroups[SearchContextGroupIndex]);
-					SearchContextGroups.RemoveAt(SearchContextGroupIndex);
-				}
-			}
-			else
-			{
-				++SearchContextGroupIndex;
-			}
-		}
-		if (MainSearchContextGroupIndex == INDEX_NONE)
-		{
-			SearchContextGroups.AddDefaulted_GetRef().Add(SearchContext, SearchContextIndex);
-		}
-	}
-
-#if ENABLE_DRAW_DEBUG
+#if ENABLE_DRAW_DEBUG && ENABLE_ANIM_DEBUG
 	DebugDraw();
-#endif // ENABLE_DRAW_DEBUG
+#endif // ENABLE_DRAW_DEBUG && ENABLE_ANIM_DEBUG
 
-	UninjectAllIslands();
-
-	for (const FInteractionSearchContextGroup& SearchContextGroup : SearchContextGroups)
-	{
-		FInteractionIsland& Island = GetAvailableIsland();
-
-		for (const UObject* AnimContext : SearchContextGroup.AnimContexts)
-		{
-			Island.InjectToActor(AnimContext);
-		}
-
-		for (int32 SearchContextsIndex : SearchContextGroup.SearchContextsIndices)
-		{
-			Island.AddSearchContext(SearchContexts[SearchContextsIndex]);
-		}
-	}
+	RegenerateAllIslands(SearchContextGroups, SearchContexts);
 
+	// clearing up AvailabilityRequestsMap for the next frame
 	AvailabilityRequestsMap.Reset();
 
 	check(ValidateAllIslands());
@@ -673,12 +764,13 @@ void UPoseSearchInteractionSubsystem::Query_AnyThread(const TArrayView<const FPo
 	const FPoseSearchContinuingProperties& ContinuingProperties, FPoseSearchInteractionBlueprintResult& Result,
 	FName PoseHistoryName, const FAnimNode_PoseSearchHistoryCollector_Base* HistoryCollector, bool bValidateResultAgainstAvailabilities)
 {
+	QUICK_SCOPE_CYCLE_COUNTER(STAT_UPoseSearchInteractionSubsystem_Query_AnyThread);
+
 	using namespace UE::PoseSearch;
 
 	Result = FPoseSearchInteractionBlueprintResult();
 
 	// if we find AnimContext in an island, we perform ALL the Island motion matching searches.
-	// thread safety is ensured by the lock on Island->SearchResultsMutex in DoSearch_AnyThread
 	if (FInteractionIsland* Island = FindIsland(AnimContext))
 	{
 		Island->DoSearch_AnyThread(AnimContext, ContinuingProperties, Result);
@@ -720,12 +812,34 @@ void UPoseSearchInteractionSubsystem::Query_AnyThread(const TArrayView<const FPo
 	}
 }
 
-#if ENABLE_DRAW_DEBUG
+#if ENABLE_DRAW_DEBUG && ENABLE_ANIM_DEBUG
 void UPoseSearchInteractionSubsystem::DebugDraw() const
 {
 	using namespace UE::PoseSearch;
 
-	static const FColor Colors[] = { FColor::Red, FColor::Green, FColor::Blue, FColor::Yellow, FColor::Black };
+	check(IsInGameThread());
+
+	if (!GVarPoseSearchInteractionShowIslands)
+	{
+		return;
+	}
+
+	static const FColor Colors[] =
+	{
+		FColor::White,
+		FColor::Black,
+		FColor::Red,
+		FColor::Green,
+		FColor::Blue,
+		FColor::Yellow,
+		FColor::Cyan,
+		FColor::Magenta,
+		FColor::Orange,
+		FColor::Purple,
+		FColor::Turquoise,
+		FColor::Silver,
+		FColor::Emerald
+	};
 	static const int32 NumColors = sizeof(Colors) / sizeof(Colors[0]);
 	int32 CurrentColorIndex = 0;
 
@@ -733,9 +847,34 @@ void UPoseSearchInteractionSubsystem::DebugDraw() const
 	{
 		if (!Island->IsUninjected())
 		{
-			Island->DebugDraw(Colors[CurrentColorIndex]);
+			const FColor& Color = Colors[CurrentColorIndex];
+
+			for (const FInteractionSearchContext& SearchContext : Island->GetSearchContexts())
+			{
+				for (int32 Index = 0; Index < SearchContext.AnimContexts.Num(); ++Index)
+				{
+					if (const UObject* AnimContext = SearchContext.AnimContexts[Index].Get())
+					{
+						if (const TArray<FInteractionAvailabilityEx>* AvailabilityRequests = AvailabilityRequestsMap.Find(AnimContext))
+						{
+							float MaxBroadPhaseRadius = 0.f;
+							for (const FInteractionAvailabilityEx& AvailabilityRequest : *AvailabilityRequests)
+							{
+								MaxBroadPhaseRadius = FMath::Max(MaxBroadPhaseRadius, AvailabilityRequest.BroadPhaseRadius);
+							}
+
+							if (MaxBroadPhaseRadius > UE_SMALL_NUMBER)
+							{
+								const FVector Position = GetContextLocation(AnimContext);
+								DrawDebugCircle(GetWorld(), Position, MaxBroadPhaseRadius, 40, Color, false, 0.f, SDPG_Foreground, 0.f, FVector::XAxisVector, FVector::YAxisVector, false);
+							}
+						}
+					}
+				}
+			}
+
 			CurrentColorIndex = (CurrentColorIndex + 1) % NumColors;
 		}
 	}
 }
-#endif // ENABLE_DRAW_DEBUG
+#endif // ENABLE_DRAW_DEBUG && ENABLE_ANIM_DEBUG
diff --git a/Engine/Plugins/Animation/PoseSearch/Source/Runtime/Public/PoseSearch/PoseSearchDefines.h b/Engine/Plugins/Animation/PoseSearch/Source/Runtime/Public/PoseSearch/PoseSearchDefines.h
--- a/Engine/Plugins/Animation/PoseSearch/Source/Runtime/Public/PoseSearch/PoseSearchDefines.h
+++ b/Engine/Plugins/Animation/PoseSearch/Source/Runtime/Public/PoseSearch/PoseSearchDefines.h
@@ -7,7 +7,19 @@
 #include "Animation/AnimTrace.h"
 
 // Enable this if object tracing is enabled, mimics animation tracing
+#ifndef UE_POSE_SEARCH_TRACE_ENABLED
 #define UE_POSE_SEARCH_TRACE_ENABLED ANIM_TRACE_ENABLED
+#endif
+
+// NoTe: the UE::PoseSearch::FInteractionIsland mutex is now unnecessary, but I'll leave this macro to reenable it in case I missed something =):>
+#ifndef UE_POSE_SEARCH_ENABLE_ISLAND_CRITICAL_SECTION
+#define UE_POSE_SEARCH_ENABLE_ISLAND_CRITICAL_SECTION 0
+#endif
+
+// NoTe: the UPoseSearchInteractionSubsystem mutex to retrieve the subsystem is now unnecessary, but I'll leave this macro to reenable it in case I missed something =):>
+#ifndef UE_POSE_SEARCH_ENABLE_INTERACTION_SUBSYSTEM_CRITICAL_SECTION
+#define UE_POSE_SEARCH_ENABLE_INTERACTION_SUBSYSTEM_CRITICAL_SECTION 0
+#endif
 
 #ifndef UE_POSE_SEARCH_FORCE_SINGLE_THREAD
 #define UE_POSE_SEARCH_FORCE_SINGLE_THREAD 0
diff --git a/Engine/Plugins/Animation/PoseSearch/Source/Runtime/Public/PoseSearch/PoseSearchInteractionIsland.h b/Engine/Plugins/Animation/PoseSearch/Source/Runtime/Public/PoseSearch/PoseSearchInteractionIsland.h
--- a/Engine/Plugins/Animation/PoseSearch/Source/Runtime/Public/PoseSearch/PoseSearchInteractionIsland.h
+++ b/Engine/Plugins/Animation/PoseSearch/Source/Runtime/Public/PoseSearch/PoseSearchInteractionIsland.h
@@ -24,12 +24,6 @@ struct FInteractionSearchContext
 	TArray<const FAnimNode_PoseSearchHistoryCollector_Base*, TInlineAllocator<PreallocatedRolesNum>> HistoryCollectors;
 	TArray<UE::PoseSearch::FRole, TInlineAllocator<PreallocatedRolesNum>> Roles;
 
-#if ENABLE_DRAW_DEBUG
-	TArray<float> BroadPhaseRadiuses;
-#endif // ENABLE_DRAW_DEBUG
-
-	float MaxCost = MAX_flt;
-
 	TWeakObjectPtr<const UPoseSearchDatabase> Database;
 	FPoseSearchContinuingProperties ContinuingProperties;
 
@@ -88,8 +82,6 @@ struct FInteractionIsland
 	const FInteractionSearchResult* FindSearchResult(const FInteractionSearchContext& SearchContext) const;
 	void AddSearchContext(const FInteractionSearchContext& SearchContext);
 
-	void DebugDraw(const FColor& Color = FColor::Red) const;
-
 private:
 	bool GetResult_AnyThread(UObject* AnimContext, FPoseSearchInteractionBlueprintResult& Result);
 	const AActor* GetMainActor() const;
@@ -110,17 +102,23 @@ private:
 	};
 	FPostTickFunction PostTickFunction;
 
+	// islands don't get deallocated, so this array once warmed up will not hit the allocator and waste cycles
 	TArray<TWeakObjectPtr<UActorComponent>> PreTickComponents;
+	// islands don't get deallocated, so this array once warmed up will not hit the allocator and waste cycles
 	TArray<TWeakObjectPtr<UActorComponent>> PostTickComponents;
 	
 	// there's one FSearchContext for each search we need to perform (including all the possible roles permutations). Added by UPoseSearchInteractionSubsystem::Tick
+	// islands don't get deallocated, so this array once warmed up will not hit the allocator and waste cycles
 	TArray<FInteractionSearchContext> SearchContexts;
 
 	// SearchResults contains only the best results, and it has not necessarly the same cardinality as SearchContexts. usually SearchResults.Num() < SearchContexts.Num()
+	// islands don't get deallocated, so this array once warmed up will not hit the allocator and waste cycles
 	TArray<FInteractionSearchResult> SearchResults;
 	bool bSearchPerfomed = false;
 
+#if UE_POSE_SEARCH_ENABLE_ISLAND_CRITICAL_SECTION
 	FCriticalSection SearchResultsMutex;
+#endif // UE_POSE_SEARCH_ENABLE_ISLAND_CRITICAL_SECTION
 
 #if ENABLE_ANIM_DEBUG
 	// used to analyze thread safety
diff --git a/Engine/Plugins/Animation/PoseSearch/Source/Runtime/Public/PoseSearch/PoseSearchInteractionLibrary.h b/Engine/Plugins/Animation/PoseSearch/Source/Runtime/Public/PoseSearch/PoseSearchInteractionLibrary.h
--- a/Engine/Plugins/Animation/PoseSearch/Source/Runtime/Public/PoseSearch/PoseSearchInteractionLibrary.h
+++ b/Engine/Plugins/Animation/PoseSearch/Source/Runtime/Public/PoseSearch/PoseSearchInteractionLibrary.h
@@ -29,10 +29,6 @@ struct POSESEARCH_API FPoseSearchInteractionAvailability
 	// the associated character to this FPoseSearchInteractionAvailability will partecipate in an interaction only if all the necessary roles gest assigned to character within BroadPhaseRadius centimeters
 	UPROPERTY(EditAnywhere, BlueprintReadWrite, Category=Settings)
 	float BroadPhaseRadius = 500.f;
-
-	// if MaxCost if greater than zero, the associated character to this FPoseSearchInteractionAvailability will not partecipate in an interaction if the motion matching search cost result is higher than MaxCost
-	UPROPERTY(EditAnywhere, BlueprintReadWrite, Category=Settings)
-	float MaxCost = 0.f;
 };
 
 // Experimental, this feature might be removed without warning, not for production use
diff --git a/Engine/Plugins/Animation/PoseSearch/Source/Runtime/Public/PoseSearch/PoseSearchInteractionSubsystem.h b/Engine/Plugins/Animation/PoseSearch/Source/Runtime/Public/PoseSearch/PoseSearchInteractionSubsystem.h
--- a/Engine/Plugins/Animation/PoseSearch/Source/Runtime/Public/PoseSearch/PoseSearchInteractionSubsystem.h
+++ b/Engine/Plugins/Animation/PoseSearch/Source/Runtime/Public/PoseSearch/PoseSearchInteractionSubsystem.h
@@ -34,6 +34,7 @@ private:
 
 typedef TMap<TWeakObjectPtr<UObject>, TArray<FInteractionAvailabilityEx>> FAvailabilityRequestsMap;
 
+struct FInteractionSearchContextGroup;
 } // namespace UE::PoseSearch
 
 // World subsystem accepting the publication of characters (via their AnimInstance(s)) FPoseSearchInteractionAvailability, representing the characters willingness to partecipate in an 
@@ -111,23 +112,28 @@ private:
 	void DestroyIsland(int32 Index);
 	void DestroyAllIslands();
 	void UninjectAllIslands();
+	void RegenerateAllIslands(TConstArrayView<UE::PoseSearch::FInteractionSearchContextGroup> SearchContextGroups, TConstArrayView<UE::PoseSearch::FInteractionSearchContext> SearchContexts);
+
 #if DO_CHECK
 	bool ValidateAllIslands() const;
 #endif // DO_CHECK
 
-	void PopulateContinuingProperties(UE::PoseSearch::FInteractionSearchContext& SearchContext, float DeltaSeconds) const;
+	void PopulateContinuingProperties(float DeltaSeconds, TArrayView<UE::PoseSearch::FInteractionSearchContext> SearchContexts) const;
 	UE::PoseSearch::FInteractionIsland* FindIsland(const UObject* AnimContext);
 	
-#if ENABLE_DRAW_DEBUG
+#if ENABLE_DRAW_DEBUG && ENABLE_ANIM_DEBUG
 	void DebugDraw() const;
-#endif // ENABLE_DRAW_DEBUG
+#endif // ENABLE_DRAW_DEBUG && ENABLE_ANIM_DEBUG
 
+	// @todo: use a thread local storage data structure to make publishing availabilities lockless to avoid the need of AvailabilityRequestsMapMutex
 	UE::PoseSearch::FAvailabilityRequestsMap AvailabilityRequestsMap;
 	FCriticalSection AvailabilityRequestsMapMutex;
 
 	// array of groups of characters that needs to be anaylzed together for possible interactions
 	TArray<UE::PoseSearch::FInteractionIsland*> Islands;
 
+#if UE_POSE_SEARCH_ENABLE_INTERACTION_SUBSYSTEM_CRITICAL_SECTION
 	// critical section to retrieve the subsystem in a thread safe manner
 	static FCriticalSection RetrieveSubsystemMutex;
+#endif // UE_POSE_SEARCH_ENABLE_INTERACTION_SUBSYSTEM_CRITICAL_SECTION
 };
