posesearch - motion matching interaction narrowphase fix and broadphase performance improvements

[RN] skip

diff --git a/Engine/Plugins/Animation/PoseSearch/Source/Runtime/Private/PoseSearchInteractionIsland.cpp b/Engine/Plugins/Animation/PoseSearch/Source/Runtime/Private/PoseSearchInteractionIsland.cpp
--- a/Engine/Plugins/Animation/PoseSearch/Source/Runtime/Private/PoseSearchInteractionIsland.cpp
+++ b/Engine/Plugins/Animation/PoseSearch/Source/Runtime/Private/PoseSearchInteractionIsland.cpp
@@ -461,15 +461,8 @@ void FInteractionIsland::InjectToActor(const UObject* AnimContext)
 
 void FInteractionIsland::AddSearchContext(const FInteractionSearchContext& SearchContext)
 {
-#if DO_CHECK
 	check(SearchContext.CheckForConsistency());
 	check(IsInGameThread());
-
-	for (const FInteractionSearchContext& ContainedSearchContext : SearchContexts)
-	{
-		check(!ContainedSearchContext.IsEquivalent(SearchContext));
-	}
-#endif // DO_CHECK
 	SearchContexts.Add(SearchContext);
 }
 
@@ -480,43 +473,50 @@ void FInteractionIsland::Uninject()
 
 	check(PreTickComponents.Num() == PostTickComponents.Num());
 
-	for (int32 ActorIndex = 0; ActorIndex < PreTickComponents.Num(); ++ActorIndex)
+	if (IsUninjected())
 	{
-		UActorComponent* PreTickComponent = PreTickComponents[ActorIndex].Get();
-		UActorComponent* PostTickComponent = PostTickComponents[ActorIndex].Get();
-
-		if (PreTickComponent)
+		check(SearchContexts.IsEmpty() && SearchResults.IsEmpty() && !bSearchPerfomed);
+	}
+	else
+	{
+		for (int32 ActorIndex = 0; ActorIndex < PreTickComponents.Num(); ++ActorIndex)
 		{
-			check(PostTickComponent);
+			UActorComponent* PreTickComponent = PreTickComponents[ActorIndex].Get();
+			UActorComponent* PostTickComponent = PostTickComponents[ActorIndex].Get();
 
-			AActor* Actor = PreTickComponent->GetOwner();
-			check(Actor);
+			if (PreTickComponent)
+			{
+				check(PostTickComponent);
 
-			PreTickFunction.RemovePrerequisite(Actor, PreTickComponent->PrimaryComponentTick);
-			PostTickComponent->PrimaryComponentTick.RemovePrerequisite(Actor, PreTickFunction);
+				AActor* Actor = PreTickComponent->GetOwner();
+				check(Actor);
 
-			const bool bIsFirstInjectedActor = ActorIndex == 0;
-			if (bIsFirstInjectedActor)
-			{
-				PostTickFunction.RemovePrerequisite(Actor, PostTickComponent->PrimaryComponentTick);
+				PreTickFunction.RemovePrerequisite(Actor, PreTickComponent->PrimaryComponentTick);
+				PostTickComponent->PrimaryComponentTick.RemovePrerequisite(Actor, PreTickFunction);
+
+				const bool bIsFirstInjectedActor = ActorIndex == 0;
+				if (bIsFirstInjectedActor)
+				{
+					PostTickFunction.RemovePrerequisite(Actor, PostTickComponent->PrimaryComponentTick);
+				}
+				else
+				{
+					PostTickComponent->PrimaryComponentTick.RemovePrerequisite(Actor, PostTickFunction);
+				}
 			}
 			else
 			{
-				PostTickComponent->PrimaryComponentTick.RemovePrerequisite(Actor, PostTickFunction);
+				check(!PostTickComponent);
 			}
 		}
-		else
-		{
-			check(!PostTickComponent);
-		}
-	}
 
-	PreTickComponents.Reset();
-	PostTickComponents.Reset();
+		PreTickComponents.Reset();
+		PostTickComponents.Reset();
 
-	SearchContexts.Reset();
-	SearchResults.Reset();
-	bSearchPerfomed = false;
+		SearchContexts.Reset();
+		SearchResults.Reset();
+		bSearchPerfomed = false;
+	}
 }
 
 bool FInteractionIsland::IsUninjected() const
@@ -567,7 +567,7 @@ bool FInteractionIsland::CanPerformSearch(const UObject* AnimContext) const
 	return true;
 }
 
-bool FInteractionIsland::DoSearch_AnyThread(UObject* AnimContext, const FPoseSearchContinuingProperties& ContinuingProperties, FPoseSearchInteractionBlueprintResult& Result)
+bool FInteractionIsland::DoSearch_AnyThread(const UObject* AnimContext, const FPoseSearchContinuingProperties& ContinuingProperties, FPoseSearchInteractionBlueprintResult& Result)
 {
 	check(AnimContext);
 
@@ -727,7 +727,7 @@ bool FInteractionIsland::DoSearch_AnyThread(UObject* AnimContext, const FPoseSea
 	return false;
 }
 
-bool FInteractionIsland::GetResult_AnyThread(UObject* AnimContext, FPoseSearchInteractionBlueprintResult& Result)
+bool FInteractionIsland::GetResult_AnyThread(const UObject* AnimContext, FPoseSearchInteractionBlueprintResult& Result)
 {
 	check(AnimContext);
 
diff --git a/Engine/Plugins/Animation/PoseSearch/Source/Runtime/Private/PoseSearchInteractionLibrary.cpp b/Engine/Plugins/Animation/PoseSearch/Source/Runtime/Private/PoseSearchInteractionLibrary.cpp
--- a/Engine/Plugins/Animation/PoseSearch/Source/Runtime/Private/PoseSearchInteractionLibrary.cpp
+++ b/Engine/Plugins/Animation/PoseSearch/Source/Runtime/Private/PoseSearchInteractionLibrary.cpp
@@ -5,7 +5,7 @@
 #include "PoseSearch/PoseSearchInteractionSubsystem.h"
 #include "Engine/World.h"
 
-FPoseSearchInteractionBlueprintResult UPoseSearchInteractionLibrary::MotionMatchInteraction_Pure(TArray<FPoseSearchInteractionAvailability> Availabilities, UObject* AnimContext, FPoseSearchContinuingProperties ContinuingProperties, FName PoseHistoryName, bool bValidateResultAgainstAvailabilities)
+FPoseSearchInteractionBlueprintResult UPoseSearchInteractionLibrary::MotionMatchInteraction_Pure(TArray<FPoseSearchInteractionAvailability> Availabilities, const UObject* AnimContext, FPoseSearchContinuingProperties ContinuingProperties, FName PoseHistoryName, bool bValidateResultAgainstAvailabilities)
 {
 	FPoseSearchInteractionBlueprintResult Result;
 	if (UPoseSearchInteractionSubsystem* InteractionSubsystem = UPoseSearchInteractionSubsystem::GetSubsystem_AnyThread(AnimContext))
@@ -15,12 +15,12 @@ FPoseSearchInteractionBlueprintResult UPoseSearchInteractionLibrary::MotionMatch
 	return Result;
 }
 
-FPoseSearchInteractionBlueprintResult UPoseSearchInteractionLibrary::MotionMatchInteraction(TArray<FPoseSearchInteractionAvailability> Availabilities, UObject* AnimContext, FPoseSearchContinuingProperties ContinuingProperties, FName PoseHistoryName, bool bValidateResultAgainstAvailabilities)
+FPoseSearchInteractionBlueprintResult UPoseSearchInteractionLibrary::MotionMatchInteraction(TArray<FPoseSearchInteractionAvailability> Availabilities, const UObject* AnimContext, FPoseSearchContinuingProperties ContinuingProperties, FName PoseHistoryName, bool bValidateResultAgainstAvailabilities)
 {
 	return MotionMatchInteraction_Pure(Availabilities, AnimContext, ContinuingProperties, PoseHistoryName, bValidateResultAgainstAvailabilities);
 }
 
-FPoseSearchInteractionBlueprintResult UPoseSearchInteractionLibrary::MotionMatchInteraction(const TArrayView<const FPoseSearchInteractionAvailability> Availabilities, UObject* AnimContext, const FPoseSearchContinuingProperties& ContinuingProperties, const FAnimNode_PoseSearchHistoryCollector_Base* HistoryCollector, bool bValidateResultAgainstAvailabilities)
+FPoseSearchInteractionBlueprintResult UPoseSearchInteractionLibrary::MotionMatchInteraction(const TArrayView<const FPoseSearchInteractionAvailability> Availabilities, const UObject* AnimContext, const FPoseSearchContinuingProperties& ContinuingProperties, const FAnimNode_PoseSearchHistoryCollector_Base* HistoryCollector, bool bValidateResultAgainstAvailabilities)
 {
 	FPoseSearchInteractionBlueprintResult Result;
 	if (UPoseSearchInteractionSubsystem* InteractionSubsystem = UPoseSearchInteractionSubsystem::GetSubsystem_AnyThread(AnimContext))
diff --git a/Engine/Plugins/Animation/PoseSearch/Source/Runtime/Private/PoseSearchInteractionSubsystem.cpp b/Engine/Plugins/Animation/PoseSearch/Source/Runtime/Private/PoseSearchInteractionSubsystem.cpp
--- a/Engine/Plugins/Animation/PoseSearch/Source/Runtime/Private/PoseSearchInteractionSubsystem.cpp
+++ b/Engine/Plugins/Animation/PoseSearch/Source/Runtime/Private/PoseSearchInteractionSubsystem.cpp
@@ -14,6 +14,11 @@ static bool GVarPoseSearchInteractionShowIslands = false;
 static FAutoConsoleVariableRef CVarPoseSearchInteractionShowIslands(TEXT("a.PoseSearchInteraction.ShowIslands"), GVarPoseSearchInteractionShowIslands, TEXT("Show Pose Search Interaction Islands"));
 #endif // ENABLE_DRAW_DEBUG && ENABLE_ANIM_DEBUG
 
+#if !NO_CVARS
+static bool GVarPoseSearchInteractionEnabled = true;
+static FAutoConsoleVariableRef CVarPoseSearchInteractionEnabled(TEXT("a.PoseSearchInteraction.Enabled"), GVarPoseSearchInteractionEnabled, TEXT("Enable/Disable Pose Search Interaction"));
+#endif // !NO_CVARS
+
 static FVector GetContextLocation(const UObject* AnimContext)
 {
 	if (const UAnimInstance* AnimInstance = Cast<UAnimInstance>(AnimContext))
@@ -25,7 +30,7 @@ static FVector GetContextLocation(const UObject* AnimContext)
 	return FVector::ZeroVector;
 }
 
-FString FInteractionAvailabilityEx::GetPoseHistoryName() const
+FString FInteractionAvailability::GetPoseHistoryName() const
 {
 	if (HistoryCollector)
 	{
@@ -34,7 +39,7 @@ FString FInteractionAvailabilityEx::GetPoseHistoryName() const
 	return PoseHistoryName.ToString();
 }
 	
-const FAnimNode_PoseSearchHistoryCollector_Base* FInteractionAvailabilityEx::GetHistoryCollector(const UObject* AnimContext) const
+const FAnimNode_PoseSearchHistoryCollector_Base* FInteractionAvailability::GetHistoryCollector(const UObject* AnimContext) const
 {
 	if (HistoryCollector)
 	{
@@ -94,332 +99,79 @@ struct FInteractionSearchContextGroup
 	TSet<const UObject*, DefaultKeyFuncs<const UObject*>, TInlineSetAllocator<16, TMemStackSetAllocator<>>> AnimContexts;
 	TArray<int32, TInlineAllocator<16, TMemStackAllocator<>>> SearchContextsIndices;
 };
-
 typedef TArray<FInteractionSearchContextGroup, TInlineAllocator<PreallocatedSearchesNum, TMemStackAllocator<>>> FInteractionSearchContextGroups;
 
-struct FRoledAnimContext
+struct FAnimContextInfo;
+struct FNearbyAnimContextInfo
 {
-	FRoledAnimContext(const UObject* InAnimContext, const FRole InRole, const FAnimNode_PoseSearchHistoryCollector_Base* HistoryCollector)
-		: AnimContext(InAnimContext)
-		, Role(InRole)
-		, HistoryCollector(HistoryCollector)
+	FNearbyAnimContextInfo(const FAnimContextInfo& InAnimContextInfo, float InDistanceSquared)
+		: AnimContextInfo(&InAnimContextInfo)
+		, DistanceSquared(InDistanceSquared)
 	{
-		check(AnimContext && HistoryCollector);
 	}
 
-	bool operator==(const FRoledAnimContext& Other) const = default;
-
-	const UObject* AnimContext = nullptr;
-	FRole Role = DefaultRole;
-	const FAnimNode_PoseSearchHistoryCollector_Base* HistoryCollector = nullptr;
+	const FAnimContextInfo* AnimContextInfo = nullptr;
+	float DistanceSquared = 0.f;
 };
-
-typedef TArray<FRoledAnimContext, TInlineAllocator<PreallocatedRolesNum, TMemStackAllocator<>>> FRoledAnimContexts;
-typedef	TArray<const UObject*, TInlineAllocator<256, TMemStackAllocator<>>> FSortedAnimContexts;
-
-static void GenerateSortedAnimContexts(const FAvailabilityRequestsMap& AvailabilityRequestsMap, const UWorld* SubsystemWorld, FSortedAnimContexts& SortedAnimContexts)
-{
-	QUICK_SCOPE_CYCLE_COUNTER(STAT_UPoseSearchInteractionSubsystem_GenerateSortedAnimContexts);
-
-	check(SubsystemWorld);
-
-	SortedAnimContexts.Reset(AvailabilityRequestsMap.Num());
-
-	for (FAvailabilityRequestsMap::TConstIterator Iter = AvailabilityRequestsMap.CreateConstIterator(); Iter; ++Iter)
-	{
-		// making sure the weak pointer is still valid
-		const UObject* AnimContext = Iter->Key.Get();
-		if (!AnimContext)
-		{
-			UE_LOG(LogPoseSearch, Log, TEXT("ProcessAvailabilityRequestsMap null anim context. The associated character got removed from the world"));
-		}
-		else if (!AnimContext->GetWorld())
-		{
-			UE_LOG(LogPoseSearch, Error, TEXT("ProcessAvailabilityRequestsMap AnimContext %s is not in a world!"), *AnimContext->GetName());
-		}
-		else if (SubsystemWorld != AnimContext->GetWorld())
-		{
-			UE_LOG(LogPoseSearch, Error, TEXT("ProcessAvailabilityRequestsMap AnimContext %s is from World %s, and supposed to be in World @s!"), *AnimContext->GetName(), *AnimContext->GetWorld()->GetName(), *SubsystemWorld->GetName());
-		}
-		else
-		{
-			SortedAnimContexts.Add(AnimContext);
-		}
-	}
-	SortedAnimContexts.Sort([](const UObject& AnimContextA, const UObject& AnimContextB)
-		{
-			// @todo: do we need to sort this reliably for networking?
-			return &AnimContextA < &AnimContextB;
-		});
-}
-
 struct FAnimContextInfo
 {
-	FVector Location;
-	float MaxBroadPhaseRadiusSquared;
-	TArray<int32, TInlineAllocator<8, TMemStackAllocator<>>> LinkedAnimContextIndex;
-};
-
-typedef TArray<FAnimContextInfo, TMemStackAllocator<>> FAnimContextInfos;
-
-static void GenerateAnimContextInfos(const FAvailabilityRequestsMap& AvailabilityRequestsMap, const FSortedAnimContexts& SortedAnimContexts, FAnimContextInfos& AnimContextInfos)
-{
-	QUICK_SCOPE_CYCLE_COUNTER(STAT_UPoseSearchInteractionSubsystem_GenerateAnimContextInfos);
-
-	const int32 NumSortedAnimContexts = SortedAnimContexts.Num();
-	AnimContextInfos.SetNum(NumSortedAnimContexts);
-	for (int32 AnimContextIndex = 0; AnimContextIndex < NumSortedAnimContexts; ++AnimContextIndex)
+	float GetCachedDistanceSquaredFrom(const FAnimContextInfo& OtherAnimContextInfo) const
 	{
-		const UObject* AnimContext = SortedAnimContexts[AnimContextIndex];
-		check(AnimContext);
-
-		float MaxBroadPhaseRadius = 0.f;
-		for (const FInteractionAvailabilityEx& AvailabilityRequest : AvailabilityRequestsMap[AnimContext])
-		{
-			MaxBroadPhaseRadius = FMath::Max(MaxBroadPhaseRadius, AvailabilityRequest.BroadPhaseRadius);
-		}
-
-		AnimContextInfos[AnimContextIndex].MaxBroadPhaseRadiusSquared = MaxBroadPhaseRadius * MaxBroadPhaseRadius;
-		AnimContextInfos[AnimContextIndex].Location = GetContextLocation(AnimContext);
-	}
+		check(this != &OtherAnimContextInfo);
 
-	// solving the broad phase using the AnimContextInfos
-	for (int32 AnimContextIndexA = 0; AnimContextIndexA < NumSortedAnimContexts; ++AnimContextIndexA)
-	{
-		for (int32 AnimContextIndexB = AnimContextIndexA + 1; AnimContextIndexB < NumSortedAnimContexts; ++AnimContextIndexB)
+		// since there should be only few NearbyAnimContextInfos we can efficiently search linearly rather than calculating and
+		// return (Location - OtherAnimContextInfo.Location).SquaredLength()...
+		for (const FNearbyAnimContextInfo& NearbyAnimContextInfo : NearbyAnimContextInfos)
 		{
-			const FVector DeltaLocation = AnimContextInfos[AnimContextIndexA].Location - AnimContextInfos[AnimContextIndexB].Location;
-			const float MaxBroadPhaseRadiusSquared = FMath::Max(AnimContextInfos[AnimContextIndexA].MaxBroadPhaseRadiusSquared, AnimContextInfos[AnimContextIndexB].MaxBroadPhaseRadiusSquared);
-			if (DeltaLocation.SquaredLength() <= MaxBroadPhaseRadiusSquared)
+			if (NearbyAnimContextInfo.AnimContextInfo == &OtherAnimContextInfo)
 			{
-				// AnimContext can potentially interact with AnimContextB: linking AnimContext to AnimContextB and vice versa
-				AnimContextInfos[AnimContextIndexA].LinkedAnimContextIndex.Add(AnimContextIndexB);
-				AnimContextInfos[AnimContextIndexB].LinkedAnimContextIndex.Add(AnimContextIndexA);
+				return NearbyAnimContextInfo.DistanceSquared;
 			}
 		}
-	}
-}
 
-struct FAnimContextInfoVisitor
-{
-	FAnimContextInfoVisitor(
-		const FSortedAnimContexts& InSortedAnimContexts,
-		const FAnimContextInfos& InAnimContextInfos,
-		TFunctionRef<void(const UObject*)> OnNewAnimContextFound,
-		TFunctionRef<void()> OnDoneGroupingAnimContexts)
-		: SortedAnimContexts(InSortedAnimContexts)
-		, AnimContextInfos(InAnimContextInfos)
-	{
-		const int32 NumSortedAnimContexts = AnimContextInfos.Num();
-		for (int32 AnimContextIndex = 0; AnimContextIndex < NumSortedAnimContexts; ++AnimContextIndex)
-		{
-			if (!AnimContextInfos[AnimContextIndex].LinkedAnimContextIndex.IsEmpty() && !VisitedIndexes.Find(AnimContextIndex))
-			{
-				// starting the evaluation of a new set of grouped AnimContext(s)
-
-				// processing the AvailabilityRequestsMap of the current AnimContextArray to fill up a map of Databases pointing to an array of all the AnimContexts with related roles
-				VisitRecursively(AnimContextIndex, OnNewAnimContextFound);
-
-				OnDoneGroupingAnimContexts();
-			}
-		}
-	}
-
-private:
-	void VisitRecursively(int32 AnimContextIndex, TFunctionRef<void(const UObject*)> OnNewAnimContextFound)
-	{
-		check(!AnimContextInfos[AnimContextIndex].LinkedAnimContextIndex.IsEmpty());
-
-		bool bIsAlreadyInSet;
-		VisitedIndexes.FindOrAdd(AnimContextIndex, &bIsAlreadyInSet);
-
-		if (!bIsAlreadyInSet)
-		{
-			OnNewAnimContextFound(SortedAnimContexts[AnimContextIndex]);
+#if DO_CHECK
+		// sanity check: if we couldn't find OtherAnimContextInfo in the NearbyAnimContextInfos, it means it didn't pass the broad phase,
+		// that means its location is further away than MaxDistanceSquared
+		const FVector DeltaLocation = Location - OtherAnimContextInfo.Location;
+		const float DistanceSquared = DeltaLocation.SquaredLength();
+		check(DistanceSquared > MaxDistanceSquared);
+#endif // DO_CHECK
 
-			for (int32 LinkedAnimContextIndex : AnimContextInfos[AnimContextIndex].LinkedAnimContextIndex)
-			{
-				VisitRecursively(LinkedAnimContextIndex, OnNewAnimContextFound);
-			}
-		}
+		// we could pass any number higher than MaxDistanceSquared. I like UE_BIG_NUMBER more :)
+		return UE_BIG_NUMBER;
 	}
 
-	const FSortedAnimContexts& SortedAnimContexts;
-	const FAnimContextInfos& AnimContextInfos;
-
-	typedef TSet<int32, DefaultKeyFuncs<int32>, TInlineSetAllocator<32, TMemStackSetAllocator<>>> FVisitedIndexes;
-	FVisitedIndexes VisitedIndexes;
+	const UObject* AnimContext = nullptr;
+	const FInteractionAvailabilities* Availabilities = nullptr;
+	FVector Location = FVector::ZeroVector;
+	float MaxDistanceSquared = 0.f;
+	TArray<FNearbyAnimContextInfo, TInlineAllocator<8, TMemStackAllocator<>>> NearbyAnimContextInfos;
 };
+struct FAnimContextInfos : TArray<FAnimContextInfo, TMemStackAllocator<>> {};
 
-template <typename DataType, typename EvaluateCombinationType>
-static void GeneratePermutationsRecursive(const TConstArrayView<DataType> Data, int32 DataIndex, TArrayView<int32> Combination, int32 CombinationIndex, EvaluateCombinationType EvaluateCombination)
+template <typename EvaluateCombinationType>
+static void GeneratePermutationsRecursive(int32 DataCardinality, int32 DataIndex, TArrayView<int32> Combination, int32 CombinationIndex, EvaluateCombinationType EvaluateCombination)
 {
 	if (CombinationIndex == Combination.Num())
 	{ 
-		EvaluateCombination(Data, Combination);
+		EvaluateCombination(Combination);
 	}
-	else if (DataIndex < Data.Num())
+	else if (DataIndex < DataCardinality)
 	{
 		Combination[CombinationIndex] = DataIndex;
-		GeneratePermutationsRecursive(Data, DataIndex + 1, Combination, CombinationIndex + 1, EvaluateCombination);
-		GeneratePermutationsRecursive(Data, DataIndex + 1, Combination, CombinationIndex, EvaluateCombination);
+		GeneratePermutationsRecursive(DataCardinality, DataIndex + 1, Combination, CombinationIndex + 1, EvaluateCombination);
+		GeneratePermutationsRecursive(DataCardinality, DataIndex + 1, Combination, CombinationIndex, EvaluateCombination);
 	}
 } 
 
-template <typename DataType, typename EvaluateCombinationType>
-static void GeneratePermutations(const TConstArrayView<DataType> Data, int32 CombinationCardinality, EvaluateCombinationType EvaluateCombination)
+template <typename EvaluateCombinationType>
+static void GeneratePermutations(int32 DataCardinality, int32 CombinationCardinality, EvaluateCombinationType EvaluateCombination)
 {
 	TArray<int32, TInlineAllocator<PreallocatedRolesNum>> Combination;
 	Combination.SetNum(CombinationCardinality);
-	GeneratePermutationsRecursive(Data, 0, Combination, 0, EvaluateCombination);
+	GeneratePermutationsRecursive(DataCardinality, 0, Combination, 0, EvaluateCombination);
 }
 
-typedef TArray<FInteractionSearchContext, TInlineAllocator<PreallocatedSearchesNum, TMemStackAllocator<>>> FInteractionSearchContexts;
-
-static void GenerateSearchContexts(const FAvailabilityRequestsMap& AvailabilityRequestsMap, const FSortedAnimContexts& SortedAnimContexts, const FAnimContextInfos& AnimContextInfos, FInteractionSearchContexts& SearchContexts)
-{
-	QUICK_SCOPE_CYCLE_COUNTER(STAT_UPoseSearchInteractionSubsystem_GenerateSearchContexts);
-
-	typedef TMap<const UPoseSearchDatabase*, FRoledAnimContexts, TInlineSetAllocator<PreallocatedSearchesNum, TMemStackSetAllocator<>>> FDatabaseToRoledAnimContexts;
-	typedef TPair<const UPoseSearchDatabase*, FRoledAnimContexts> FDatabaseToRoledAnimContextsPair;
-	FDatabaseToRoledAnimContexts DatabaseToRoledAnimContexts;
-
-	// visiting ALL the AnimContexts from SortedAnimContexts with their cached information hold in AnimContextInfos to figure our potential interactions
-	FAnimContextInfoVisitor AnimContextInfoVisitor(SortedAnimContexts, AnimContextInfos,
-		// OnNewAnimContextFound: called when the FAnimContextInfoVisitor finds a new AnimContext that can be grouped in the current group of possibly interacting AnimContext(s)
-		[&AvailabilityRequestsMap, &DatabaseToRoledAnimContexts](const UObject* AnimContext)
-		{
-			// analyzing all the AvailabilityRequest(s) associated with this AnimContext and eventually generate the associated FRoledAnimContexts,
-			// inserted in a per database sorted data structure (DatabaseToRoledAnimContexts)
-			for (const FInteractionAvailabilityEx& AvailabilityRequest : AvailabilityRequestsMap[AnimContext])
-			{
-				if (!AvailabilityRequest.Database)
-				{
-					UE_LOG(LogPoseSearch, Log, TEXT("ProcessAvailabilityRequestsMap null AvailabilityRequest.Database"));
-					continue;
-				}
-
-				if (!AvailabilityRequest.Database->Schema)
-				{
-					UE_LOG(LogPoseSearch, Error, TEXT("ProcessAvailabilityRequestsMap null Schema for Database %s"), *AvailabilityRequest.Database->GetName());
-					continue;
-				}
-
-				const FAnimNode_PoseSearchHistoryCollector_Base* HistoryCollector = AvailabilityRequest.GetHistoryCollector(AnimContext);
-				if (!HistoryCollector)
-				{
-					UE_LOG(LogPoseSearch, Error, TEXT("ProcessAvailabilityRequestsMap couldn't find PoseHistory %s for AnimContext %s"), *AvailabilityRequest.GetPoseHistoryName(), *AnimContext->GetName());
-					continue;
-				}
-
-				FRoledAnimContexts& RoledAnimContexts = DatabaseToRoledAnimContexts.FindOrAdd(AvailabilityRequest.Database);
-				if (AvailabilityRequest.RolesFilter.IsEmpty())
-				{
-					// adding ALL the possible roles from the database:
-					for (const FPoseSearchRoledSkeleton& RoledSkeleton : AvailabilityRequest.Database->Schema->GetRoledSkeletons())
-					{
-						RoledAnimContexts.AddUnique(FRoledAnimContext(AnimContext, RoledSkeleton.Role, HistoryCollector));
-					}
-				}
-				else
-				{
-					for (const FRole& Role : AvailabilityRequest.RolesFilter)
-					{
-						if (AvailabilityRequest.Database->Schema->GetRoledSkeleton(Role))
-						{
-							RoledAnimContexts.AddUnique(FRoledAnimContext(AnimContext, Role, HistoryCollector));
-						}
-						else
-						{
-							UE_LOG(LogPoseSearch, Warning, TEXT("ProcessAvailabilityRequestsMap unsupported Role %s for Database %s"), *Role.ToString(), *AvailabilityRequest.Database->GetName());
-						}
-					}
-				}
-			}
-		},
-		// OnDoneGroupingAnimContexts: called when the FAnimContextInfoVisitor reaches the end of a group of possibly interacting AnimContext(s)
-		[&DatabaseToRoledAnimContexts, &SearchContexts]()
-		{
-			// for each database now we try to create all the possible permutations of the roled anim instances
-			// for example, given a database set up with assets for 2 characters interactions with roles RoleA and RoleB
-			// and 2 anim instances, all of them willing to partecipate in the 2 characters interaction with both roles RoleA and RoleB:
-			// CharA could be taking RoleA and RoleB,
-			// CharB could be taking RoleA and RoleB,
-			// we generate all the permutations from the array of options:
-			// CharA/RoleA, CharA/RoleB, CharB/RoleA, CharB/RoleB
-			//
-			// and we prune the invalid tuples:
-			//
-			// CharA/RoleA - CharA/RoleB -> invalid because of same CharA
-			// CharA/RoleA - CharB/RoleA -> invalid because of same RoleA
-			// CharA/RoleA - CharB/RoleB -> VALID!
-			//
-			// CharA/RoleB - CharB/RoleA -> VALID!
-			// CharA/RoleB - CharB/RoleB -> invalid because of same RoleB
-			//
-			// CharB/RoleA - CharB/RoleB -> invalid because of same CharB
-
-			for (const FDatabaseToRoledAnimContextsPair& DatabaseToRoledAnimContextsPair : DatabaseToRoledAnimContexts)
-			{
-				const UPoseSearchDatabase* Database = DatabaseToRoledAnimContextsPair.Key;
-				const TArray<FPoseSearchRoledSkeleton>& RoledSkeletons = Database->Schema->GetRoledSkeletons();
-				const int32 CombinationCardinality = RoledSkeletons.Num();
-				const TConstArrayView<FRoledAnimContext> RoledAnimContextsView = MakeArrayView(DatabaseToRoledAnimContextsPair.Value.GetData(), DatabaseToRoledAnimContextsPair.Value.Num());
-
-				GeneratePermutations(RoledAnimContextsView, CombinationCardinality,
-					[Database, &RoledSkeletons, &SearchContexts](const TConstArrayView<FRoledAnimContext> RoledAnimContexts, const TConstArrayView<int32> Combination)
-					{
-						// checking if this combination is valid for the Database:
-						TSet<const UObject*, DefaultKeyFuncs<const UObject*>, TInlineSetAllocator<PreallocatedRolesNum, TMemStackSetAllocator<>>> AnimContexts;
-						TSet<FRole, DefaultKeyFuncs<FRole>, TInlineSetAllocator<PreallocatedRolesNum, TMemStackSetAllocator<>>> CoveredRoles;
-
-						for (int32 CombinationIndex = 0; CombinationIndex < Combination.Num(); ++CombinationIndex)
-						{
-							const FRoledAnimContext& RoledAnimContext = RoledAnimContexts[Combination[CombinationIndex]];
-
-							bool bIsAlreadyInSet = false;
-							AnimContexts.Add(RoledAnimContext.AnimContext, &bIsAlreadyInSet);
-
-							// is there any duplicated AnimContext?
-							if (bIsAlreadyInSet)
-							{
-								return false;
-							}
-
-							CoveredRoles.Add(RoledAnimContext.Role);
-						}
-
-						// does it cover all the roles?
-						for (const FPoseSearchRoledSkeleton& RoledSkeleton : RoledSkeletons)
-						{
-							if (!CoveredRoles.Find(RoledSkeleton.Role))
-							{
-								return false;
-							}
-						}
-
-						FInteractionSearchContext& SearchContext = SearchContexts.AddDefaulted_GetRef();
-						SearchContext.Database = Database;
-						SearchContext.AnimContexts.Reserve(Combination.Num());
-						SearchContext.HistoryCollectors.Reserve(Combination.Num());
-						SearchContext.Roles.Reserve(Combination.Num());
-
-						for (int32 CombinationIndex = 0; CombinationIndex < Combination.Num(); ++CombinationIndex)
-						{
-							const FRoledAnimContext& RoledAnimContext = RoledAnimContexts[Combination[CombinationIndex]];
-							SearchContext.AnimContexts.Add(RoledAnimContext.AnimContext);
-							SearchContext.HistoryCollectors.Add(RoledAnimContext.HistoryCollector);
-							SearchContext.Roles.Add(RoledAnimContext.Role);
-						}
-
-						return true;
-					});
-			}
-
-			// done using DatabaseToRoledAnimContexts. clearing up for the next group of AnimContext(s)
-			DatabaseToRoledAnimContexts.Reset();
-		});
-}
 
 static void GenerateInteractionSearchContextGroups(const FInteractionSearchContexts& SearchContexts, FInteractionSearchContextGroups& SearchContextGroups)
 {
@@ -704,6 +456,360 @@ void UPoseSearchInteractionSubsystem::Deinitialize()
 	Super::Deinitialize();
 }
 
+bool UPoseSearchInteractionSubsystem::IsThereAnyAvailability() const
+{
+	check(IsInGameThread());
+	return !AvailabilitiesMap.IsEmpty();
+}
+
+void UPoseSearchInteractionSubsystem::AddAvailabilities(const TArrayView<const FPoseSearchInteractionAvailability> Availabilities, const UObject* AnimContext, FName PoseHistoryName, const FAnimNode_PoseSearchHistoryCollector_Base* HistoryCollector)
+{
+	FScopeLock Lock(&AvailabilitiesMapMutex);
+	AvailabilitiesMap.AddAvailabilities(Availabilities, AnimContext, PoseHistoryName, HistoryCollector);
+}
+
+#if ENABLE_DRAW_DEBUG && ENABLE_ANIM_DEBUG
+const UE::PoseSearch::FInteractionAvailabilities* UPoseSearchInteractionSubsystem::FindAvailabilities(const UObject* AnimContext) const
+{
+	check(IsInGameThread());
+	return AvailabilitiesMap.Find(AnimContext);
+}
+#endif // ENABLE_DRAW_DEBUG && ENABLE_ANIM_DEBUG
+
+void UPoseSearchInteractionSubsystem::ResetAvailabilities()
+{
+	check(IsInGameThread());
+	AvailabilitiesMap.Reset();
+}
+
+void UPoseSearchInteractionSubsystem::GenerateAnimContextInfos(UE::PoseSearch::FAnimContextInfos& AnimContextInfos) const
+{
+	QUICK_SCOPE_CYCLE_COUNTER(STAT_UPoseSearchInteractionSubsystem_GenerateAnimContextInfos);
+
+	using namespace UE::PoseSearch;
+
+	const UWorld* SubsystemWorld = GetWorld();
+	check(SubsystemWorld);
+
+	// collecting and sorting all the AnimContext(s) from the AvailabilitiesMap in AnimContextsWithAvailabilities, to have a stable and deterministic island creation
+	typedef TPair<const UObject*, const FInteractionAvailabilities*> FAnimContextWithAvailabilities;
+	TArray<FAnimContextWithAvailabilities, TInlineAllocator<512, TMemStackAllocator<>>> AnimContextsWithAvailabilities;
+	for (FAvailabilitiesMap::TConstIterator Iter = AvailabilitiesMap.CreateConstIterator(); Iter; ++Iter)
+	{
+		// making sure the weak pointer is still valid
+		const UObject* AnimContext = Iter->Key.Get();
+		if (!AnimContext)
+		{
+			UE_LOG(LogPoseSearch, Log, TEXT("UPoseSearchInteractionSubsystem::GenerateAnimContextInfos null anim context. The associated character got removed from the world"));
+		}
+		else if (!AnimContext->GetWorld())
+		{
+			UE_LOG(LogPoseSearch, Error, TEXT("UPoseSearchInteractionSubsystem::GenerateAnimContextInfos AnimContext %s is not in a world!"), *AnimContext->GetName());
+		}
+		else if (SubsystemWorld != AnimContext->GetWorld())
+		{
+			UE_LOG(LogPoseSearch, Error, TEXT("UPoseSearchInteractionSubsystem::GenerateAnimContextInfos AnimContext %s is from World %s, and supposed to be in World @s!"), *AnimContext->GetName(), *AnimContext->GetWorld()->GetName(), *SubsystemWorld->GetName());
+		}
+		else
+		{
+			AnimContextsWithAvailabilities.Add({ AnimContext, &Iter->Value });
+		}
+	}
+
+	AnimContextsWithAvailabilities.Sort([](const FAnimContextWithAvailabilities& A, const FAnimContextWithAvailabilities& B)
+		{
+			// @todo: do we need to sort this reliably for networking?
+			return A.Key < B.Key;
+		});
+
+	const int32 NumAnimContextInfos = AnimContextsWithAvailabilities.Num();
+	AnimContextInfos.SetNum(NumAnimContextInfos);
+	for (int32 AnimContextInfoIndex = 0; AnimContextInfoIndex < NumAnimContextInfos; ++AnimContextInfoIndex)
+	{
+		FAnimContextInfo& AnimContextInfo = AnimContextInfos[AnimContextInfoIndex];
+
+		AnimContextInfo.AnimContext = AnimContextsWithAvailabilities[AnimContextInfoIndex].Key;
+		AnimContextInfo.Availabilities = AnimContextsWithAvailabilities[AnimContextInfoIndex].Value;
+		check(AnimContextInfo.AnimContext && AnimContextInfo.Availabilities);
+
+		AnimContextInfo.Location = GetContextLocation(AnimContextInfo.AnimContext);
+
+		for (const FInteractionAvailability& Availability : *AnimContextInfo.Availabilities)
+		{
+			AnimContextInfo.MaxDistanceSquared = FMath::Max(AnimContextInfo.MaxDistanceSquared, Availability.BroadPhaseRadius);
+		}
+		AnimContextInfo.MaxDistanceSquared *= AnimContextInfo.MaxDistanceSquared;
+	}
+
+	// solving the broad phase using the AnimContextInfos
+	for (int32 AnimContextInfoIndexA = 0; AnimContextInfoIndexA < NumAnimContextInfos; ++AnimContextInfoIndexA)
+	{
+		for (int32 AnimContextInfoIndexB = AnimContextInfoIndexA + 1; AnimContextInfoIndexB < NumAnimContextInfos; ++AnimContextInfoIndexB)
+		{
+			const FVector DeltaLocation = AnimContextInfos[AnimContextInfoIndexA].Location - AnimContextInfos[AnimContextInfoIndexB].Location;
+			const float DistanceSquared = DeltaLocation.SquaredLength();
+			const float MaxDistanceSquared = FMath::Max(AnimContextInfos[AnimContextInfoIndexA].MaxDistanceSquared, AnimContextInfos[AnimContextInfoIndexB].MaxDistanceSquared);
+			if (DistanceSquared <= MaxDistanceSquared)
+			{
+				// the AnimContext of AnimContextInfos[AnimContextInfoIndexA] can potentially interact with the one from AnimContextInfos[AnimContextInfoIndexB]:
+				// linking AnimContextInfos[AnimContextInfoIndexA] to AnimContextInfos[AnimContextInfoIndexA] and vice versa to keep track of this when evaluating the broad phase.
+				// Since AnimContextInfos does't reallocate anymore, it's safe to store pointers to internal elements of the array!
+				AnimContextInfos[AnimContextInfoIndexA].NearbyAnimContextInfos.Emplace(AnimContextInfos[AnimContextInfoIndexB], DistanceSquared);
+				AnimContextInfos[AnimContextInfoIndexB].NearbyAnimContextInfos.Emplace(AnimContextInfos[AnimContextInfoIndexA], DistanceSquared);
+			}
+		}
+	}
+}
+
+void UPoseSearchInteractionSubsystem::GenerateSearchContexts(UE::PoseSearch::FInteractionSearchContexts& SearchContexts) const
+{
+	QUICK_SCOPE_CYCLE_COUNTER(STAT_UPoseSearchInteractionSubsystem_GenerateSearchContexts);
+
+	using namespace UE::PoseSearch;
+
+	struct FRoledAnimContextInfo
+	{
+		FRoledAnimContextInfo(const FAnimContextInfo* InAnimContextInfo, const FRole InRole, const FAnimNode_PoseSearchHistoryCollector_Base* InHistoryCollector, const UPoseSearchDatabase* InDatabase, float InMaxDistanceSquared = 0.f)
+			: AnimContextInfo(InAnimContextInfo)
+			, Role(InRole)
+			, HistoryCollector(InHistoryCollector)
+			, Database(InDatabase)
+			, MaxDistanceSquared(InMaxDistanceSquared)
+		{
+			check(InAnimContextInfo && InHistoryCollector && InDatabase && MaxDistanceSquared > UE_SMALL_NUMBER);
+		}
+
+		bool operator==(const FRoledAnimContextInfo& Other) const = default;
+
+		const FAnimContextInfo* AnimContextInfo = nullptr;
+		FRole Role = DefaultRole;
+		const FAnimNode_PoseSearchHistoryCollector_Base* HistoryCollector = nullptr;
+		const UPoseSearchDatabase* Database;
+		float MaxDistanceSquared = 0.f;
+	};
+	typedef TArray<FRoledAnimContextInfo, TInlineAllocator<PreallocatedRolesNum, TMemStackAllocator<>>> FRoledAnimContextInfos;
+
+	// visits the FAnimContextInfos recursively to identify groups of nearby AnimContextInfo(s), relying on FAnimContextInfo::NearbyAnimContextInfos information.
+	// it calls OnNewAnimAnimContextInfoFound on every new FAnimContextInfo found/visited, and OnDoneGroupingAnimContexts once reaches the end of the current group AnimContext(s).
+	// it's then restart calling OnNewAnimAnimContextInfoFound in case there are still unvisited FAnimContextInfo(s), untill it visited ALL the FAnimContextInfo(s) in the input FAnimContextInfos
+	struct FAnimContextInfoVisitor
+	{
+		FAnimContextInfoVisitor(
+			const FAnimContextInfos& AnimContextInfos,
+			TFunctionRef<void(const FAnimContextInfo&)> OnNewAnimAnimContextInfoFound,
+			TFunctionRef<void()> OnDoneGroupingAnimContexts)
+		{
+			for (const FAnimContextInfo& AnimContextInfo : AnimContextInfos)
+			{
+				if (!AnimContextInfo.NearbyAnimContextInfos.IsEmpty() && !VisitedAnimContextInfos.Find(&AnimContextInfo))
+				{
+					// starting the evaluation of a new set of grouped AnimContext(s)
+
+					// processing the AvailabilitiesMap of the current AnimContextArray to fill up a map of Databases pointing to an array of all the AnimContexts with related roles
+					VisitRecursively(AnimContextInfo, OnNewAnimAnimContextInfoFound);
+
+					OnDoneGroupingAnimContexts();
+				}
+			}
+		}
+
+	private:
+		void VisitRecursively(const FAnimContextInfo& AnimContextInfoToVisit, TFunctionRef<void(const FAnimContextInfo&)> OnNewAnimAnimContextInfoFound)
+		{
+			check(!AnimContextInfoToVisit.NearbyAnimContextInfos.IsEmpty());
+
+			bool bIsAlreadyInSet;
+			VisitedAnimContextInfos.FindOrAdd(&AnimContextInfoToVisit, &bIsAlreadyInSet);
+
+			if (!bIsAlreadyInSet)
+			{
+				OnNewAnimAnimContextInfoFound(AnimContextInfoToVisit);
+
+				for (const FNearbyAnimContextInfo& NearbyAnimContextInfo : AnimContextInfoToVisit.NearbyAnimContextInfos)
+				{
+					check(NearbyAnimContextInfo.AnimContextInfo);
+					VisitRecursively(*NearbyAnimContextInfo.AnimContextInfo, OnNewAnimAnimContextInfoFound);
+				}
+			}
+		}
+
+		TSet<const FAnimContextInfo*, DefaultKeyFuncs<const FAnimContextInfo*>, TInlineSetAllocator<32, TMemStackSetAllocator<>>> VisitedAnimContextInfos;
+	};
+
+	// caching AnimContext(s) locations, max broad phase radiuses (squared) and collect relations of possible interactions between AnimContext(s)
+	// (stored in FAnimContextInfo::NearbyAnimContextInfos::Index) as fast broad phase evaluation refined later on during FInteractionSearchContexts generation
+	FAnimContextInfos AnimContextInfos;
+	GenerateAnimContextInfos(AnimContextInfos);
+
+	TMap<const UPoseSearchDatabase*, FRoledAnimContextInfos, TInlineSetAllocator<PreallocatedSearchesNum, TMemStackSetAllocator<>>> DatabaseToRoledAnimContextInfos;
+	// visiting ALL the AnimContexts in AnimContextInfos and relying on the cached information to refine potential interactions
+	FAnimContextInfoVisitor AnimContextInfoVisitor(AnimContextInfos,
+		// OnNewAnimContextFound: called when the FAnimContextInfoVisitor finds a new AnimContext that can be grouped in the current group of possibly interacting AnimContext(s)
+		[&DatabaseToRoledAnimContextInfos](const FAnimContextInfo& AnimContextInfo)
+		{
+			// analyzing all the Availability(s) associated with this AnimContext and eventually generate the associated FRoledAnimContextInfos,
+			// inserted in a per database sorted data structure (DatabaseToRoledAnimContextInfos)
+			check(AnimContextInfo.Availabilities);
+			for (const FInteractionAvailability& Availability : *AnimContextInfo.Availabilities)
+			{
+				if (Availability.BroadPhaseRadius <= UE_SMALL_NUMBER)
+				{
+					UE_LOG(LogPoseSearch, Log, TEXT("UPoseSearchInteractionSubsystem::GenerateSearchContexts too small Availability.BroadPhaseRadius. Character will not interact with anything within such small radius!"));
+					continue;
+				}
+
+				if (!Availability.Database)
+				{
+					UE_LOG(LogPoseSearch, Log, TEXT("UPoseSearchInteractionSubsystem::GenerateSearchContexts null Availability.Database"));
+					continue;
+				}
+
+				if (!Availability.Database->Schema)
+				{
+					UE_LOG(LogPoseSearch, Error, TEXT("UPoseSearchInteractionSubsystem::GenerateSearchContexts null Schema for Database %s"), *Availability.Database->GetName());
+					continue;
+				}
+
+				const FAnimNode_PoseSearchHistoryCollector_Base* HistoryCollector = Availability.GetHistoryCollector(AnimContextInfo.AnimContext);
+				if (!HistoryCollector)
+				{
+					UE_LOG(LogPoseSearch, Error, TEXT("UPoseSearchInteractionSubsystem::GenerateSearchContexts couldn't find PoseHistory %s for AnimContext %s"), *Availability.GetPoseHistoryName(), *AnimContextInfo.AnimContext->GetName());
+					continue;
+				}
+
+				FRoledAnimContextInfos& RoledAnimContextInfos = DatabaseToRoledAnimContextInfos.FindOrAdd(Availability.Database);
+				const float MaxDistanceSquared = Availability.BroadPhaseRadius * Availability.BroadPhaseRadius;
+				if (Availability.RolesFilter.IsEmpty())
+				{
+					// adding ALL the possible roles from the database:
+					for (const FPoseSearchRoledSkeleton& RoledSkeleton : Availability.Database->Schema->GetRoledSkeletons())
+					{
+						RoledAnimContextInfos.AddUnique(FRoledAnimContextInfo(&AnimContextInfo, RoledSkeleton.Role, HistoryCollector, Availability.Database, MaxDistanceSquared));
+					}
+				}
+				else
+				{
+					for (const FRole& Role : Availability.RolesFilter)
+					{
+						if (Availability.Database->Schema->GetRoledSkeleton(Role))
+						{
+							RoledAnimContextInfos.AddUnique(FRoledAnimContextInfo(&AnimContextInfo, Role, HistoryCollector, Availability.Database, MaxDistanceSquared));
+						}
+						else
+						{
+							UE_LOG(LogPoseSearch, Warning, TEXT("UPoseSearchInteractionSubsystem::GenerateSearchContexts unsupported Role %s for Database %s"), *Role.ToString(), *Availability.Database->GetName());
+						}
+					}
+				}
+			}
+		},
+		// OnDoneGroupingAnimContexts: called when the FAnimContextInfoVisitor reaches the end of a group of possibly interacting AnimContext(s)
+		[&DatabaseToRoledAnimContextInfos, &SearchContexts]()
+		{
+			// for each database now we try to create all the possible permutations of the roled anim instances
+			// for example, given a database set up with assets for 2 characters interactions with roles RoleA and RoleB
+			// and 2 anim instances, all of them willing to partecipate in the 2 characters interaction with both roles RoleA and RoleB:
+			// CharA could be taking RoleA and RoleB,
+			// CharB could be taking RoleA and RoleB,
+			// we generate all the permutations from the array of options:
+			// CharA/RoleA, CharA/RoleB, CharB/RoleA, CharB/RoleB
+			//
+			// and we prune the invalid tuples:
+			//
+			// CharA/RoleA - CharA/RoleB -> invalid because of same CharA
+			// CharA/RoleA - CharB/RoleA -> invalid because of same RoleA
+			// CharA/RoleA - CharB/RoleB -> VALID!
+			//
+			// CharA/RoleB - CharB/RoleA -> VALID!
+			// CharA/RoleB - CharB/RoleB -> invalid because of same RoleB
+			//
+			// CharB/RoleA - CharB/RoleB -> invalid because of same CharB
+
+			for (const TPair<const UPoseSearchDatabase*, FRoledAnimContextInfos>& DatabaseToRoledAnimContextInfosPair : DatabaseToRoledAnimContextInfos)
+			{
+				const UPoseSearchDatabase* Database = DatabaseToRoledAnimContextInfosPair.Key;
+				const TArray<FPoseSearchRoledSkeleton>& RoledSkeletons = Database->Schema->GetRoledSkeletons();
+				const int32 CombinationCardinality = RoledSkeletons.Num();
+				const FRoledAnimContextInfos& RoledAnimContextInfos = DatabaseToRoledAnimContextInfosPair.Value;
+
+				GeneratePermutations(RoledAnimContextInfos.Num(), CombinationCardinality,
+					// Combination is an array of indexes in RoledAnimContextInfos: 0 <= Combination[i] < RoledAnimContextInfos.Num()
+					[Database, &RoledSkeletons, &RoledAnimContextInfos, &SearchContexts](const TConstArrayView<int32> Combination)
+					{
+						// CombinationCardinality represents the number of roles as well as the number interacting AnimContext(s) (ultimately number of Characters involved in the interaction)
+						const int32 CombinationCardinality = Combination.Num();
+						TSet<const UObject*, DefaultKeyFuncs<const UObject*>, TInlineSetAllocator<PreallocatedRolesNum, TMemStackSetAllocator<>>> UniqueAnimContexts;
+
+						for (int32 CombinationIndex = 0; CombinationIndex < CombinationCardinality; ++CombinationIndex)
+						{
+							const int32 RoledAnimContextIndex = Combination[CombinationIndex];
+							const FRoledAnimContextInfo& RoledAnimContextInfo = RoledAnimContextInfos[RoledAnimContextIndex];
+							check(RoledAnimContextInfo.AnimContextInfo);
+							bool bIsAlreadyInSet;
+							UniqueAnimContexts.Add(RoledAnimContextInfo.AnimContextInfo->AnimContext, &bIsAlreadyInSet);
+							if (bIsAlreadyInSet)
+							{
+								// we have a duplicate AnimContext. this permutation is NOT valid
+								return false;
+							}
+						}
+
+						FInteractionSearchContext SearchContext;
+						SearchContext.Database = Database;
+						SearchContext.AnimContexts.Reserve(CombinationCardinality);
+						SearchContext.HistoryCollectors.Reserve(CombinationCardinality);
+						SearchContext.Roles.Reserve(CombinationCardinality);
+
+						// checking if this combination is valid for the Database:
+						for (int32 CombinationIndex = 0; CombinationIndex < CombinationCardinality; ++CombinationIndex)
+						{
+							const int32 RoledAnimContextIndex = Combination[CombinationIndex];
+							const FRoledAnimContextInfo& RoledAnimContextInfo = RoledAnimContextInfos[RoledAnimContextIndex];
+
+							// checking the narrow phase!
+							for (int32 OtherCombinationIndex = CombinationIndex + 1; OtherCombinationIndex < CombinationCardinality; ++OtherCombinationIndex)
+							{
+								const int32 OtherRoledAnimContextIndex = Combination[OtherCombinationIndex];
+								const FRoledAnimContextInfo& OtherRoledAnimContextInfo = RoledAnimContextInfos[OtherRoledAnimContextIndex];
+
+								const float DistanceSquared = RoledAnimContextInfo.AnimContextInfo->GetCachedDistanceSquaredFrom(*OtherRoledAnimContextInfo.AnimContextInfo);
+								const float MinDistanceSquared = FMath::Min(RoledAnimContextInfo.MaxDistanceSquared, RoledAnimContextInfo.MaxDistanceSquared);
+								if (DistanceSquared > MinDistanceSquared)
+								{
+									return false;
+								}
+							}
+
+							SearchContext.AnimContexts.Add(RoledAnimContextInfo.AnimContextInfo->AnimContext);
+							SearchContext.Roles.Add(RoledAnimContextInfo.Role);
+							SearchContext.HistoryCollectors.Add(RoledAnimContextInfo.HistoryCollector);
+						}
+
+						// does SearchContext cover all the roles required by this interaction?
+						for (const FPoseSearchRoledSkeleton& RoledSkeleton : RoledSkeletons)
+						{
+							// CombinationCardinality is usually 2-3, so we can search the SearchContext.Roles array for duplicates without requiring a faster container like TSet
+							if (!SearchContext.Roles.Contains(RoledSkeleton.Role))
+							{
+								return false;
+							}
+						}
+#if DO_CHECK
+						for (const FInteractionSearchContext& ContainedSearchContext : SearchContexts)
+						{
+							check(!ContainedSearchContext.IsEquivalent(SearchContext));
+						}
+#endif // DO_CHECK
+						SearchContexts.Emplace(SearchContext);
+						return true;
+					});
+			}
+
+			// done using DatabaseToRoledAnimContextInfos. clearing up for the next group of AnimContext(s)
+			DatabaseToRoledAnimContextInfos.Reset();
+		});
+}
+
 void UPoseSearchInteractionSubsystem::Tick(float DeltaSeconds)
 {
 	QUICK_SCOPE_CYCLE_COUNTER(STAT_UPoseSearchInteractionSubsystem_Tick);
@@ -712,29 +818,32 @@ void UPoseSearchInteractionSubsystem::Tick(float DeltaSeconds)
 
 	Super::Tick(DeltaSeconds);
 
-	if (AvailabilityRequestsMap.IsEmpty() && Islands.IsEmpty())
+	if (!IsThereAnyAvailability())
 	{
-		// nothing to do. early out
-		return;
+		bool bAllUninjected = true;
+		for (FInteractionIsland* Island : Islands)
+		{
+			if (!Island->IsUninjected())
+			{
+				bAllUninjected = false;
+			}
+		}
+
+		if (bAllUninjected)
+		{
+			// nothing to do. early out
+			return;
+		}
 	}
 
 	check(IsInGameThread());
 
-	// FScopeLock Lock(&AvailabilityRequestsMapMutex); is not necessary since UPoseSearchInteractionSubsystem gets ticked outside the parallel animation jobs
+	// FScopeLock Lock(&AvailabilitiesMapMutex); is not necessary since UPoseSearchInteractionSubsystem gets ticked outside the parallel animation jobs
 	FMemMark Mark(FMemStack::Get());
 
-	// collecting and sorting all the AnimContext(s) from the AvailabilityRequestsMap, to have a stable and deterministic island creation
-	FSortedAnimContexts SortedAnimContexts;
-	GenerateSortedAnimContexts(AvailabilityRequestsMap, GetWorld(), SortedAnimContexts);
-
-	// caching AnimContext(s) locations, max broad phase radiuses (squared) and collect relations of possible interactions between AnimContext(s)
-	// (stored in FAnimContextInfo::LinkedAnimContextIndex) as fast broad phase evaluation refined later on during FInteractionSearchContexts generation
-	FAnimContextInfos AnimContextInfos;
-	GenerateAnimContextInfos(AvailabilityRequestsMap, SortedAnimContexts, AnimContextInfos);
-
 	// generating all the possible interaction tuples of AnimContext(s) with roles and pose histories (defined in FInteractionSearchContext)
 	FInteractionSearchContexts SearchContexts;
-	GenerateSearchContexts(AvailabilityRequestsMap, SortedAnimContexts, AnimContextInfos, SearchContexts);
+	GenerateSearchContexts(SearchContexts);
 
 	// populating the continuing pose properties for the SearchContexts from the current Islands
 	PopulateContinuingProperties(DeltaSeconds, SearchContexts);
@@ -749,8 +858,8 @@ void UPoseSearchInteractionSubsystem::Tick(float DeltaSeconds)
 
 	RegenerateAllIslands(SearchContextGroups, SearchContexts);
 
-	// clearing up AvailabilityRequestsMap for the next frame
-	AvailabilityRequestsMap.Reset();
+	// clearing up AvailabilitiesMap for the next frame
+	ResetAvailabilities();
 
 	check(ValidateAllIslands());
 }
@@ -760,7 +869,7 @@ TStatId UPoseSearchInteractionSubsystem::GetStatId() const
 	RETURN_QUICK_DECLARE_CYCLE_STAT(UPoseSearchInteractionSubsystem, STATGROUP_Tickables);
 }
 
-void UPoseSearchInteractionSubsystem::Query_AnyThread(const TArrayView<const FPoseSearchInteractionAvailability> Availabilities, UObject* AnimContext, 
+void UPoseSearchInteractionSubsystem::Query_AnyThread(const TArrayView<const FPoseSearchInteractionAvailability> Availabilities, const UObject* AnimContext, 
 	const FPoseSearchContinuingProperties& ContinuingProperties, FPoseSearchInteractionBlueprintResult& Result,
 	FName PoseHistoryName, const FAnimNode_PoseSearchHistoryCollector_Base* HistoryCollector, bool bValidateResultAgainstAvailabilities)
 {
@@ -770,6 +879,13 @@ void UPoseSearchInteractionSubsystem::Query_AnyThread(const TArrayView<const FPo
 
 	Result = FPoseSearchInteractionBlueprintResult();
 
+#if !NO_CVARS
+	if (!GVarPoseSearchInteractionEnabled)
+	{
+		return;
+	}
+#endif // !NO_CVARS
+
 	// if we find AnimContext in an island, we perform ALL the Island motion matching searches.
 	if (FInteractionIsland* Island = FindIsland(AnimContext))
 	{
@@ -799,16 +915,7 @@ void UPoseSearchInteractionSubsystem::Query_AnyThread(const TArrayView<const FPo
 	// queuing the availabilities for the next frame Query_AnyThread
 	if (!Availabilities.IsEmpty())
 	{
-		FScopeLock Lock(&AvailabilityRequestsMapMutex);
-
-		TArray<FInteractionAvailabilityEx>& AvailabilityRequests = AvailabilityRequestsMap.FindOrAdd(AnimContext);
-		for (const FPoseSearchInteractionAvailability& Availability : Availabilities)
-		{
-			if (Availability.Database)
-			{
-				AvailabilityRequests.Emplace(Availability, PoseHistoryName, HistoryCollector);
-			}
-		}
+		AddAvailabilities(Availabilities, AnimContext, PoseHistoryName, HistoryCollector);
 	}
 }
 
@@ -855,12 +962,12 @@ void UPoseSearchInteractionSubsystem::DebugDraw() const
 				{
 					if (const UObject* AnimContext = SearchContext.AnimContexts[Index].Get())
 					{
-						if (const TArray<FInteractionAvailabilityEx>* AvailabilityRequests = AvailabilityRequestsMap.Find(AnimContext))
+						if (const FInteractionAvailabilities* Availabilities = FindAvailabilities(AnimContext))
 						{
 							float MaxBroadPhaseRadius = 0.f;
-							for (const FInteractionAvailabilityEx& AvailabilityRequest : *AvailabilityRequests)
+							for (const FInteractionAvailability& Availability : *Availabilities)
 							{
-								MaxBroadPhaseRadius = FMath::Max(MaxBroadPhaseRadius, AvailabilityRequest.BroadPhaseRadius);
+								MaxBroadPhaseRadius = FMath::Max(MaxBroadPhaseRadius, Availability.BroadPhaseRadius);
 							}
 
 							if (MaxBroadPhaseRadius > UE_SMALL_NUMBER)
diff --git a/Engine/Plugins/Animation/PoseSearch/Source/Runtime/Public/PoseSearch/PoseSearchInteractionIsland.h b/Engine/Plugins/Animation/PoseSearch/Source/Runtime/Public/PoseSearch/PoseSearchInteractionIsland.h
--- a/Engine/Plugins/Animation/PoseSearch/Source/Runtime/Public/PoseSearch/PoseSearchInteractionIsland.h
+++ b/Engine/Plugins/Animation/PoseSearch/Source/Runtime/Public/PoseSearch/PoseSearchInteractionIsland.h
@@ -36,6 +36,7 @@ struct FInteractionSearchContext
 	bool CheckForConsistency() const;
 #endif // DO_CHECK
 };
+typedef TArray<FInteractionSearchContext, TInlineAllocator<PreallocatedSearchesNum, TMemStackAllocator<>>> FInteractionSearchContexts;
 
 // Experimental, this feature might be removed without warning, not for production use
 struct FInteractionSearchResult : public FSearchResult
@@ -65,7 +66,7 @@ struct FInteractionIsland
 	FInteractionIsland(ULevel* Level);
 	~FInteractionIsland();
 
-	bool DoSearch_AnyThread(UObject* AnimContext, const FPoseSearchContinuingProperties& ContinuingProperties, FPoseSearchInteractionBlueprintResult& Result);
+	bool DoSearch_AnyThread(const UObject* AnimContext, const FPoseSearchContinuingProperties& ContinuingProperties, FPoseSearchInteractionBlueprintResult& Result);
 
 	const TArray<TWeakObjectPtr<UActorComponent>>& GetPreTickComponents() const { return PreTickComponents; }
 	const TArray<TWeakObjectPtr<UActorComponent>>& GetPostTickComponents() const { return PostTickComponents; }
@@ -83,7 +84,7 @@ struct FInteractionIsland
 	void AddSearchContext(const FInteractionSearchContext& SearchContext);
 
 private:
-	bool GetResult_AnyThread(UObject* AnimContext, FPoseSearchInteractionBlueprintResult& Result);
+	bool GetResult_AnyThread(const UObject* AnimContext, FPoseSearchInteractionBlueprintResult& Result);
 	const AActor* GetMainActor() const;
 	bool CanPerformSearch(const UObject* AnimContext) const;
 
diff --git a/Engine/Plugins/Animation/PoseSearch/Source/Runtime/Public/PoseSearch/PoseSearchInteractionLibrary.h b/Engine/Plugins/Animation/PoseSearch/Source/Runtime/Public/PoseSearch/PoseSearchInteractionLibrary.h
--- a/Engine/Plugins/Animation/PoseSearch/Source/Runtime/Public/PoseSearch/PoseSearchInteractionLibrary.h
+++ b/Engine/Plugins/Animation/PoseSearch/Source/Runtime/Public/PoseSearch/PoseSearchInteractionLibrary.h
@@ -106,14 +106,14 @@ public:
 	// PoseHistoryName is the name of the pose history node used for the associated motion matching search
 	// if bValidateResultAgainstAvailabilities is true, the result will be invalidated if doesn't respect the new availabilities
 	UFUNCTION(Experimental, BlueprintPure, Category = "Animation|Pose Search", meta = (BlueprintThreadSafe))
-	static FPoseSearchInteractionBlueprintResult MotionMatchInteraction_Pure(TArray<FPoseSearchInteractionAvailability> Availabilities, UObject* AnimInstance, FPoseSearchContinuingProperties ContinuingProperties, FName PoseHistoryName, bool bValidateResultAgainstAvailabilities = true);
+	static FPoseSearchInteractionBlueprintResult MotionMatchInteraction_Pure(TArray<FPoseSearchInteractionAvailability> Availabilities, const UObject* AnimContext, FPoseSearchContinuingProperties ContinuingProperties, FName PoseHistoryName, bool bValidateResultAgainstAvailabilities = true);
 
 	// BlueprintCallable version of MotionMatchInteraction_Pure
 	UFUNCTION(Experimental, BlueprintCallable, Category = "Animation|Pose Search", meta = (BlueprintThreadSafe))
-	static FPoseSearchInteractionBlueprintResult MotionMatchInteraction(TArray<FPoseSearchInteractionAvailability> Availabilities, UObject* AnimInstance, FPoseSearchContinuingProperties ContinuingProperties, FName PoseHistoryName, bool bValidateResultAgainstAvailabilities = true);
+	static FPoseSearchInteractionBlueprintResult MotionMatchInteraction(TArray<FPoseSearchInteractionAvailability> Availabilities, const UObject* AnimContext, FPoseSearchContinuingProperties ContinuingProperties, FName PoseHistoryName, bool bValidateResultAgainstAvailabilities = true);
 
 	// version of MotionMatchInteraction_Pure referencing directly the HistoryCollector rather than looking for it by name
-	static FPoseSearchInteractionBlueprintResult MotionMatchInteraction(const TArrayView<const FPoseSearchInteractionAvailability> Availabilities, UObject* AnimInstance, const FPoseSearchContinuingProperties& ContinuingProperties, const FAnimNode_PoseSearchHistoryCollector_Base* HistoryCollector, bool bValidateResultAgainstAvailabilities = true);
+	static FPoseSearchInteractionBlueprintResult MotionMatchInteraction(const TArrayView<const FPoseSearchInteractionAvailability> Availabilities, const UObject* AnimContext, const FPoseSearchContinuingProperties& ContinuingProperties, const FAnimNode_PoseSearchHistoryCollector_Base* HistoryCollector, bool bValidateResultAgainstAvailabilities = true);
 
 	UFUNCTION(Experimental, BlueprintPure, Category = "Animation|Pose Search", meta = (BlueprintThreadSafe))
 	static FPoseSearchContinuingProperties GetMontageContinuingProperties(UAnimInstance* AnimInstance);
diff --git a/Engine/Plugins/Animation/PoseSearch/Source/Runtime/Public/PoseSearch/PoseSearchInteractionSubsystem.h b/Engine/Plugins/Animation/PoseSearch/Source/Runtime/Public/PoseSearch/PoseSearchInteractionSubsystem.h
--- a/Engine/Plugins/Animation/PoseSearch/Source/Runtime/Public/PoseSearch/PoseSearchInteractionSubsystem.h
+++ b/Engine/Plugins/Animation/PoseSearch/Source/Runtime/Public/PoseSearch/PoseSearchInteractionSubsystem.h
@@ -15,9 +15,9 @@ namespace UE::PoseSearch
 {
 
 // Experimental, this feature might be removed without warning, not for production use
-struct FInteractionAvailabilityEx : public FPoseSearchInteractionAvailability
+struct FInteractionAvailability : public FPoseSearchInteractionAvailability
 {
-	FInteractionAvailabilityEx(const FPoseSearchInteractionAvailability& InAvailability, FName InPoseHistoryName, const FAnimNode_PoseSearchHistoryCollector_Base* InHistoryCollector)
+	FInteractionAvailability(const FPoseSearchInteractionAvailability& InAvailability, FName InPoseHistoryName, const FAnimNode_PoseSearchHistoryCollector_Base* InHistoryCollector)
 		: FPoseSearchInteractionAvailability(InAvailability)
 		, PoseHistoryName(InPoseHistoryName)
 		, HistoryCollector(InHistoryCollector)
@@ -32,9 +32,33 @@ private:
 	const FAnimNode_PoseSearchHistoryCollector_Base* HistoryCollector = nullptr;
 };
 
-typedef TMap<TWeakObjectPtr<UObject>, TArray<FInteractionAvailabilityEx>> FAvailabilityRequestsMap;
+typedef TArray<FInteractionAvailability> FInteractionAvailabilities;
+struct FAvailabilitiesMap : protected TMap<TWeakObjectPtr<const UObject>, FInteractionAvailabilities>
+{
+	using Super = TMap<TWeakObjectPtr<const UObject>, FInteractionAvailabilities>;
+	using Super::TConstIterator;
+
+	using Super::CreateConstIterator;
+	using Super::IsEmpty;
+	using Super::Find;
+	using Super::Reset;
+	
+	void AddAvailabilities(const TArrayView<const FPoseSearchInteractionAvailability> Availabilities, const UObject* AnimContext, FName PoseHistoryName, const FAnimNode_PoseSearchHistoryCollector_Base* HistoryCollector)
+	{
+		FInteractionAvailabilities& AnimContextAvailabilities = FindOrAdd(AnimContext);
+		for (const FPoseSearchInteractionAvailability& Availability : Availabilities)
+		{
+			if (Availability.Database)
+			{
+				AnimContextAvailabilities.Emplace(Availability, PoseHistoryName, HistoryCollector);
+			}
+		}
+	}
+};
 
+struct FAnimContextInfos;
 struct FInteractionSearchContextGroup;
+
 } // namespace UE::PoseSearch
 
 // World subsystem accepting the publication of characters (via their AnimInstance(s)) FPoseSearchInteractionAvailability, representing the characters willingness to partecipate in an 
@@ -94,10 +118,13 @@ public:
 	
 	virtual TStatId GetStatId() const override;
 
-	// publishing FPoseSearchInteractionAvailability(s) for the requesting character (via AnimContext as UAnimInstance or UAnimNextComponent) and getting the FPoseSearchInteractionBlueprintResult from the PREVIOUS Tick update
-	// containing the animation to play at what time, and the assigned role to partecipate in the selected interaction
+	// publishing FPoseSearchInteractionAvailability(s) for the requesting character (via AnimContext as UAnimInstance or UAnimNextComponent) 
+	// and getting the FPoseSearchInteractionBlueprintResult from the PREVIOUS Tick update
+	// containing the animation to play at what time, and the assigned role to partecipate in the selected interaction.
 	// Either a PoseHistoryName or a HistoryCollector are required to perform the associated motion matching searches
-	void Query_AnyThread(const TArrayView<const FPoseSearchInteractionAvailability> Availabilities, UObject* AnimContext, 
+	// if bValidateResultAgainstAvailabilities is true, the result will be validated against the currently published availabilities,
+	// that could differ from the previous frame one for this AnimContext, and discarded if not valid within the context of the new availabilities
+	void Query_AnyThread(const TArrayView<const FPoseSearchInteractionAvailability> Availabilities, const UObject* AnimContext, 
 		const FPoseSearchContinuingProperties& ContinuingProperties, FPoseSearchInteractionBlueprintResult& Result,
 		FName PoseHistoryName, const FAnimNode_PoseSearchHistoryCollector_Base* HistoryCollector, bool bValidateResultAgainstAvailabilities);
 
@@ -123,11 +150,18 @@ private:
 	
 #if ENABLE_DRAW_DEBUG && ENABLE_ANIM_DEBUG
 	void DebugDraw() const;
+	const UE::PoseSearch::FInteractionAvailabilities* FindAvailabilities(const UObject* AnimContext) const;
 #endif // ENABLE_DRAW_DEBUG && ENABLE_ANIM_DEBUG
 
-	// @todo: use a thread local storage data structure to make publishing availabilities lockless to avoid the need of AvailabilityRequestsMapMutex
-	UE::PoseSearch::FAvailabilityRequestsMap AvailabilityRequestsMap;
-	FCriticalSection AvailabilityRequestsMapMutex;
+	bool IsThereAnyAvailability() const;
+	void AddAvailabilities(const TArrayView<const FPoseSearchInteractionAvailability> Availabilities, const UObject* AnimContext, FName PoseHistoryName, const FAnimNode_PoseSearchHistoryCollector_Base* HistoryCollector);
+	void ResetAvailabilities();
+	void GenerateAnimContextInfos(UE::PoseSearch::FAnimContextInfos& AnimContextInfos) const;
+	void GenerateSearchContexts(UE::PoseSearch::FInteractionSearchContexts& SearchContexts) const;
+
+	// @todo: use a thread local storage data structure to make publishing availabilities lockless to avoid the need of AvailabilitiesMapMutex
+	UE::PoseSearch::FAvailabilitiesMap AvailabilitiesMap;
+	FCriticalSection AvailabilitiesMapMutex;
 
 	// array of groups of characters that needs to be anaylzed together for possible interactions
 	TArray<UE::PoseSearch::FInteractionIsland*> Islands;
