posesearch - MM interaction threading fix and diagnostic improvements

[RN] skip

diff --git a/Engine/Plugins/Animation/PoseSearch/Source/Runtime/Private/AnimNode_MotionMatchingInteraction.cpp b/Engine/Plugins/Animation/PoseSearch/Source/Runtime/Private/AnimNode_MotionMatchingInteraction.cpp
--- a/Engine/Plugins/Animation/PoseSearch/Source/Runtime/Private/AnimNode_MotionMatchingInteraction.cpp
+++ b/Engine/Plugins/Animation/PoseSearch/Source/Runtime/Private/AnimNode_MotionMatchingInteraction.cpp
@@ -55,7 +55,8 @@ void FAnimNode_MotionMatchingInteraction::UpdateAssetPlayer(const FAnimationUpda
 {
 	using namespace UE::PoseSearch;
 
-	CheckInteractionThreadSafety(Context.AnimInstanceProxy->GetAnimInstanceObject(), FInteractionValidator::EPhase::MotionMatchingInteractionUpdate);
+	check(Context.AnimInstanceProxy);
+	CheckInteractionThreadSafety(Context.AnimInstanceProxy->GetAnimInstanceObject());
 
 	if (NeedsReset(Context))
 	{
@@ -76,7 +77,6 @@ void FAnimNode_MotionMatchingInteraction::UpdateAssetPlayer(const FAnimationUpda
 		ContinuingProperties.PlayingAsset = CurrentResult.SelectedAnimation; // GetAnimAsset();
 		ContinuingProperties.PlayingAssetAccumulatedTime = GetAccumulatedTime();
 
-		check(Context.AnimInstanceProxy);
 		CurrentResult = UPoseSearchInteractionLibrary::MotionMatchInteraction(Availabilities, Context.AnimInstanceProxy->GetAnimInstanceObject(), ContinuingProperties, PoseHistoryProvider->GetHistoryCollector(), bValidateResultAgainstAvailabilities);
 		check(CurrentResult.ActorRootTransforms.Num() == CurrentResult.ActorRootBoneTransforms.Num());
 
@@ -132,8 +132,7 @@ void FAnimNode_MotionMatchingInteraction::Evaluate_AnyThread(FPoseContext& Outpu
 	using namespace UE::PoseSearch;
 
 	check(Output.AnimInstanceProxy);
-
-	CheckInteractionThreadSafety(Output.AnimInstanceProxy->GetAnimInstanceObject(), FInteractionValidator::EPhase::MotionMatchingInteractionEvaluate);
+	CheckInteractionThreadSafety(Output.AnimInstanceProxy->GetAnimInstanceObject());
 
 	Super::Evaluate_AnyThread(Output);
 
diff --git a/Engine/Plugins/Animation/PoseSearch/Source/Runtime/Private/AnimNode_PoseSearchHistoryCollector.cpp b/Engine/Plugins/Animation/PoseSearch/Source/Runtime/Private/AnimNode_PoseSearchHistoryCollector.cpp
--- a/Engine/Plugins/Animation/PoseSearch/Source/Runtime/Private/AnimNode_PoseSearchHistoryCollector.cpp
+++ b/Engine/Plugins/Animation/PoseSearch/Source/Runtime/Private/AnimNode_PoseSearchHistoryCollector.cpp
@@ -115,12 +115,12 @@ void FAnimNode_PoseSearchHistoryCollector::Evaluate_AnyThread(FPoseContext& Outp
 
 	using namespace UE::PoseSearch;
 
-	CheckInteractionThreadSafety(Output.AnimInstanceProxy->GetAnimInstanceObject(), FInteractionValidator::EPhase::HistoryCollectorEvaluate);
-	check(Output.AnimInstanceProxy);
-
 	Super::Evaluate_AnyThread(Output);
 	Source.Evaluate(Output);
 
+	check(Output.AnimInstanceProxy);
+	CheckInteractionThreadSafety(Output.AnimInstanceProxy->GetAnimInstanceObject());
+
 	const bool bNeedsReset = bResetOnBecomingRelevant && UpdateCounter.HasEverBeenUpdated() && !UpdateCounter.WasSynchronizedCounter(Output.AnimInstanceProxy->GetUpdateCounter());
 
 	FCSPose<FCompactPose> ComponentSpacePose;
@@ -156,16 +156,21 @@ void FAnimNode_PoseSearchHistoryCollector::Update_AnyThread(const FAnimationUpda
 	DECLARE_SCOPE_HIERARCHICAL_COUNTER_ANIMNODE(Update_AnyThread);
 
 	using namespace UE::PoseSearch;
+	
+	// scoping CheckInteractionThreadSafety
+	{
+		check(Context.AnimInstanceProxy);
+		CheckInteractionThreadSafety(Context.AnimInstanceProxy->GetAnimInstanceObject());
 
-	CheckInteractionThreadSafety(Context.AnimInstanceProxy->GetAnimInstanceObject(), FInteractionValidator::EPhase::HistoryCollectorUpdate);
+		if (bGenerateTrajectory)
+		{
+			GenerateTrajectory(Cast<UAnimInstance>(Context.AnimInstanceProxy->GetAnimInstanceObject()));
+		}
+		bIsTrajectoryGeneratedBeforePreUpdate = false;
 
-	if (bGenerateTrajectory)
-	{
-		GenerateTrajectory(Cast<UAnimInstance>(Context.AnimInstanceProxy->GetAnimInstanceObject()));
+		Super::Update_AnyThread(Context);
 	}
-	bIsTrajectoryGeneratedBeforePreUpdate = false;
-
-	Super::Update_AnyThread(Context);
+	
 	UE::Anim::TScopedGraphMessage<UE::PoseSearch::FPoseHistoryProvider> ScopedMessage(Context, this);
 	Source.Update(Context);
 }
@@ -200,13 +205,12 @@ void FAnimNode_PoseSearchComponentSpaceHistoryCollector::EvaluateComponentSpace_
 
 	using namespace UE::PoseSearch;
 
-	CheckInteractionThreadSafety(Output.AnimInstanceProxy->GetAnimInstanceObject(), FInteractionValidator::EPhase::HistoryCollectorEvaluate);
-
-	check(Output.AnimInstanceProxy);
-
 	Super::EvaluateComponentSpace_AnyThread(Output);
 	Source.EvaluateComponentSpace(Output);
 
+	check(Output.AnimInstanceProxy);
+	CheckInteractionThreadSafety(Output.AnimInstanceProxy->GetAnimInstanceObject());
+
 	const bool bNeedsReset = bResetOnBecomingRelevant && UpdateCounter.HasEverBeenUpdated() && !UpdateCounter.WasSynchronizedCounter(Output.AnimInstanceProxy->GetUpdateCounter());
 
 	TArray<FBoneIndexType> RequiredBones;
@@ -240,15 +244,21 @@ void FAnimNode_PoseSearchComponentSpaceHistoryCollector::Update_AnyThread(const
 
 	using namespace UE::PoseSearch;
 
-	CheckInteractionThreadSafety(Context.AnimInstanceProxy->GetAnimInstanceObject(), FInteractionValidator::EPhase::HistoryCollectorUpdate);
 
-	if (bGenerateTrajectory)
+	// scoping CheckInteractionThreadSafety
 	{
-		GenerateTrajectory(Cast<UAnimInstance>(Context.AnimInstanceProxy->GetAnimInstanceObject()));
+		check(Context.AnimInstanceProxy);
+		CheckInteractionThreadSafety(Context.AnimInstanceProxy->GetAnimInstanceObject());
+
+		if (bGenerateTrajectory)
+		{
+			GenerateTrajectory(Cast<UAnimInstance>(Context.AnimInstanceProxy->GetAnimInstanceObject()));
+		}
+		bIsTrajectoryGeneratedBeforePreUpdate = false;
+
+		Super::Update_AnyThread(Context);
 	}
-	bIsTrajectoryGeneratedBeforePreUpdate = false;
 
-	Super::Update_AnyThread(Context);
 	UE::Anim::TScopedGraphMessage<UE::PoseSearch::FPoseHistoryProvider> ScopedMessage(Context, this);
 	Source.Update(Context);
 }
diff --git a/Engine/Plugins/Animation/PoseSearch/Source/Runtime/Private/PoseSearchInteractionIsland.cpp b/Engine/Plugins/Animation/PoseSearch/Source/Runtime/Private/PoseSearchInteractionIsland.cpp
--- a/Engine/Plugins/Animation/PoseSearch/Source/Runtime/Private/PoseSearchInteractionIsland.cpp
+++ b/Engine/Plugins/Animation/PoseSearch/Source/Runtime/Private/PoseSearchInteractionIsland.cpp
@@ -9,6 +9,8 @@
 #include "GameFramework/CharacterMovementComponent.h"
 #include "PoseSearch/AnimNode_PoseSearchHistoryCollector.h"
 #include "PoseSearch/MultiAnimAsset.h"
+#include "PoseSearch/PoseSearchInteractionSubsystem.h"
+#include "PoseSearch/PoseSearchInteractionValidator.h"
 #include "PoseSearch/PoseSearchLibrary.h"
 #include "PoseSearch/PoseSearchDatabase.h"
 #include "PoseSearch/PoseSearchSchema.h"
@@ -318,25 +320,51 @@ void FInteractionIsland::FPreTickFunction::ExecuteTick(float DeltaTime, enum ELe
 {
 	// Called before any skeletal mesh component tick, when there aren't animation jobs flying. No need to FScopeLock Lock(&Mutex);
 	// generating trajectories before running any of the skeletal mesh component ticks
-	for (const FInteractionSearchContext& SearchContext : Island->SearchContexts)
+	check(Island);
+
+	if (!Island->IsUninjected())
 	{
-		for (int32 Index = 0; Index < SearchContext.AnimContexts.Num(); ++Index)
+		CheckInteractionThreadSafety(Island);
+
+		for (const FInteractionSearchContext& SearchContext : Island->SearchContexts)
 		{
-			if (const UObject* AnimContext = SearchContext.AnimContexts[Index].Get())
+			for (int32 Index = 0; Index < SearchContext.AnimContexts.Num(); ++Index)
 			{
-				if (const UAnimInstance* AnimInstance = Cast<UAnimInstance>(AnimContext))
-				{
-					// since FInteractionIsland has a tick dependency with the USkeletalMeshComponent it's safe modify the FAnimNode_PoseSearchHistoryCollector_Base
-					FAnimNode_PoseSearchHistoryCollector_Base* HistoryCollector = const_cast<FAnimNode_PoseSearchHistoryCollector_Base*>(SearchContext.HistoryCollectors[Index]);
-					check(HistoryCollector);
-					HistoryCollector->GenerateTrajectory(AnimInstance);
-				}
-				else
+				if (const UObject* AnimContext = SearchContext.AnimContexts[Index].Get())
 				{
-					unimplemented();
+					if (const UAnimInstance* AnimInstance = Cast<UAnimInstance>(AnimContext))
+					{
+						// since FInteractionIsland has a tick dependency with the USkeletalMeshComponent it's safe modify the FAnimNode_PoseSearchHistoryCollector_Base
+						FAnimNode_PoseSearchHistoryCollector_Base* HistoryCollector = const_cast<FAnimNode_PoseSearchHistoryCollector_Base*>(SearchContext.HistoryCollectors[Index]);
+						check(HistoryCollector);
+						HistoryCollector->GenerateTrajectory(AnimInstance);
+					}
+					else
+					{
+						unimplemented();
+					}
 				}
 			}
 		}
+
+#if ENABLE_ANIM_DEBUG
+		if (Island->bPreTickFunctionExecuted)
+		{
+			// @todo: need to figure out why when creating a new island FPreTickFunction gets called twice (it's not a real issue rather than a performance hit)
+			//		  use GVarPoseSearchInteractionCacheIslands = false to debug the issue (it destroys the islands every frame)
+			UE_LOG(LogPoseSearch, Warning, TEXT("FInteractionIsland::FPreTickFunction::ExecuteTick, called twice before UPoseSearchInteractionSubsystem::Tick!"));
+		}
+		else
+		{
+			if (Island->bPostTickFunctionExecuted)
+			{
+				UE_LOG(LogPoseSearch, Error, TEXT("FInteractionIsland::FPreTickFunction::ExecuteTick, FPostTickFunction::ExecuteTick alreay run?!"));
+				Island->LogTickDependencies();
+			}
+
+			Island->bPreTickFunctionExecuted = true;
+		}
+#endif // ENABLE_ANIM_DEBUG
 	}
 }
 
@@ -344,33 +372,58 @@ void FInteractionIsland::FPreTickFunction::ExecuteTick(float DeltaTime, enum ELe
 ///////////////////////////////////////////////////////////
 void FInteractionIsland::FPostTickFunction::ExecuteTick(float DeltaTime, enum ELevelTick TickType, ENamedThreads::Type CurrentThread, const FGraphEventRef& MyCompletionGraphEvent)
 {
-	// do nothing
+	check(Island);
+
+#if ENABLE_ANIM_DEBUG
+	if (!Island->IsUninjected())
+	{
+		CheckInteractionThreadSafety(Island);
+
+		if (Island->bPostTickFunctionExecuted)
+		{
+			// @todo: need to figure out why when creating a new island FPostTickFunction gets called twice (it's not a real issue rather than a performance hit)
+			//		  use GVarPoseSearchInteractionCacheIslands = false to debug the issue (it destroys the islands every frame)
+			UE_LOG(LogPoseSearch, Warning, TEXT("FInteractionIsland::FPostTickFunction::ExecuteTick, called twice before UPoseSearchInteractionSubsystem::Tick!"));
+		}
+		else
+		{
+			if (!Island->bPreTickFunctionExecuted)
+			{
+				UE_LOG(LogPoseSearch, Error, TEXT("FInteractionIsland::FPostTickFunction::ExecuteTick, FPreTickFunction::ExecuteTick didn't run!"));
+				Island->LogTickDependencies();
+			}
+
+			Island->bPostTickFunctionExecuted = true;
+		}
+	}
+#endif // ENABLE_ANIM_DEBUG
 }
 
 // FInteractionIsland
 ///////////////////////////////////////////////////////////
-FInteractionIsland::FInteractionIsland(ULevel* Level)
+FInteractionIsland::FInteractionIsland(ULevel* Level, UPoseSearchInteractionSubsystem* Subsystem)
 {
-	PreTickFunction.bCanEverTick = true;
-	PreTickFunction.bStartWithTickEnabled = true;
-	PreTickFunction.TickGroup = ETickingGroup::TG_PrePhysics;
+	PreTickFunction.bAllowTickBatching = true;
+	PreTickFunction.bRunOnAnyThread = true;
 	PreTickFunction.Island = this;
-	PreTickFunction.SetTickFunctionEnable(true);
 	PreTickFunction.RegisterTickFunction(Level);
 
-	PostTickFunction.bCanEverTick = true;
-	PostTickFunction.bStartWithTickEnabled = true;
-	PostTickFunction.TickGroup = ETickingGroup::TG_PrePhysics;
-	PostTickFunction.SetTickFunctionEnable(true);
+	PostTickFunction.bAllowTickBatching = true;
+	PostTickFunction.bRunOnAnyThread = true;
+	PostTickFunction.Island = this;
 	PostTickFunction.RegisterTickFunction(Level);
+
+	InteractionSubsystem = Subsystem;
 }
 
 FInteractionIsland::~FInteractionIsland()
 {
+	Uninject(false);
+
 	PreTickFunction.UnRegisterTickFunction();
 	PostTickFunction.UnRegisterTickFunction();
 
-	Uninject();
+	InteractionSubsystem = nullptr;
 }
 
 UActorComponent* FInteractionIsland::FindPreTickComponent(const AActor* Actor)
@@ -415,6 +468,22 @@ void FInteractionIsland::InjectToActor(const UObject* AnimContext)
 	// Called by UPoseSearchInteractionSubsystem::Tick when there aren't animation jobs flying. No need to FScopeLock Lock(&Mutex);
 	if (AnimContext)
 	{
+#if ENABLE_ANIM_DEBUG
+		if (bPreTickFunctionExecuted || bPostTickFunctionExecuted)
+		{
+			if (bPreTickFunctionExecuted)
+			{
+				UE_LOG(LogPoseSearch, Error, TEXT("FInteractionIsland::InjectToActor, unexpected FPreTickFunction::ExecuteTick run!"));
+			}
+			else
+			{
+				UE_LOG(LogPoseSearch, Error, TEXT("FInteractionIsland::InjectToActor, unexpected FPostTickFunction::ExecuteTick run!"));
+			}
+				
+			LogTickDependencies();
+		}
+#endif // ENABLE_ANIM_DEBUG
+
 		if (const UAnimInstance* AnimInstance = Cast<UAnimInstance>(AnimContext))
 		{
 			AActor* Actor = AnimInstance->GetOwningActor();
@@ -438,16 +507,18 @@ void FInteractionIsland::InjectToActor(const UObject* AnimContext)
 					// (so we can remove them later on in a consistent fashion)
 					check(PreTickComponents.Num() == PostTickComponents.Num());
 
-					PreTickFunction.AddPrerequisite(Actor, PreTickComponent->PrimaryComponentTick);
-					PostTickComponent->PrimaryComponentTick.AddPrerequisite(Actor, PreTickFunction);
+					PreTickFunction.AddPrerequisite(PreTickComponent, PreTickComponent->PrimaryComponentTick);
 
 					if (bIsFirstInjectedActor)
 					{
-						PostTickFunction.AddPrerequisite(Actor, PostTickComponent->PrimaryComponentTick);
+						// it should be this island, but it's not a UObject, so we use the InteractionSubsystem
+						PostTickComponent->PrimaryComponentTick.AddPrerequisite(InteractionSubsystem, PreTickFunction);
+						PostTickFunction.AddPrerequisite(PostTickComponent, PostTickComponent->PrimaryComponentTick);
 					}
 					else
 					{
-						PostTickComponent->PrimaryComponentTick.AddPrerequisite(Actor, PostTickFunction);
+						// it should be this island, but it's not a UObject, so we use the InteractionSubsystem
+						PostTickComponent->PrimaryComponentTick.AddPrerequisite(InteractionSubsystem, PostTickFunction);
 					}
 				}
 			}
@@ -466,7 +537,7 @@ void FInteractionIsland::AddSearchContext(const FInteractionSearchContext& Searc
 	SearchContexts.Add(SearchContext);
 }
 
-void FInteractionIsland::Uninject()
+void FInteractionIsland::Uninject(bool bValidateTickDependencies)
 {
 	// Called by UPoseSearchInteractionSubsystem::Tick when there aren't animation jobs flying. No need to FScopeLock Lock(&Mutex);
 	check(IsInGameThread());
@@ -476,40 +547,70 @@ void FInteractionIsland::Uninject()
 	if (IsUninjected())
 	{
 		check(SearchContexts.IsEmpty() && SearchResults.IsEmpty() && !bSearchPerfomed);
+
+#if ENABLE_ANIM_DEBUG
+		if (bValidateTickDependencies && (bPreTickFunctionExecuted || bPostTickFunctionExecuted))
+		{
+			if (bPreTickFunctionExecuted)
+			{
+				UE_LOG(LogPoseSearch, Error, TEXT("FInteractionIsland::Uninject, unexpected FPreTickFunction::ExecuteTick run!"));
+			}
+			else
+			{
+				UE_LOG(LogPoseSearch, Error, TEXT("FInteractionIsland::Uninject, unexpected FPostTickFunction::ExecuteTick run!"));
+			}
+				
+			LogTickDependencies();
+		}
+#endif // ENABLE_ANIM_DEBUG
 	}
 	else
 	{
-		for (int32 ActorIndex = 0; ActorIndex < PreTickComponents.Num(); ++ActorIndex)
+#if ENABLE_ANIM_DEBUG
+		if (bValidateTickDependencies && (!bPreTickFunctionExecuted || !bPostTickFunctionExecuted))
 		{
-			UActorComponent* PreTickComponent = PreTickComponents[ActorIndex].Get();
-			UActorComponent* PostTickComponent = PostTickComponents[ActorIndex].Get();
-
-			if (PreTickComponent)
+			if (!bPreTickFunctionExecuted)
 			{
-				check(PostTickComponent);
-
-				AActor* Actor = PreTickComponent->GetOwner();
-				check(Actor);
+				UE_LOG(LogPoseSearch, Error, TEXT("FInteractionIsland::Uninject, expected FPreTickFunction::ExecuteTick didn't run!"));
+			}
+			else
+			{
+				UE_LOG(LogPoseSearch, Error, TEXT("FInteractionIsland::Uninject, expected FPostTickFunction::ExecuteTick didn't run!"));
+			}
+			
+			LogTickDependencies();
+		}
+#endif // ENABLE_ANIM_DEBUG
 
-				PreTickFunction.RemovePrerequisite(Actor, PreTickComponent->PrimaryComponentTick);
-				PostTickComponent->PrimaryComponentTick.RemovePrerequisite(Actor, PreTickFunction);
+		for (TWeakObjectPtr<UActorComponent>& PreTickComponentPtr : PreTickComponents)
+		{
+			if (UActorComponent* PreTickComponent = PreTickComponentPtr.Get())
+			{
+				PreTickFunction.RemovePrerequisite(PreTickComponent, PreTickComponent->PrimaryComponentTick);
+			}
+		}
+		check(PreTickFunction.GetPrerequisites().IsEmpty());
 
-				const bool bIsFirstInjectedActor = ActorIndex == 0;
-				if (bIsFirstInjectedActor)
+		bool bMainActor = true;
+		for (TWeakObjectPtr<UActorComponent>& PostTickComponentPtr : PostTickComponents)
+		{
+			if (UActorComponent* PostTickComponent = PostTickComponentPtr.Get())
+			{
+				if (bMainActor)
 				{
-					PostTickFunction.RemovePrerequisite(Actor, PostTickComponent->PrimaryComponentTick);
+					PostTickComponent->PrimaryComponentTick.RemovePrerequisite(InteractionSubsystem, PreTickFunction);
+					PostTickFunction.RemovePrerequisite(PostTickComponent, PostTickComponent->PrimaryComponentTick);
 				}
 				else
 				{
-					PostTickComponent->PrimaryComponentTick.RemovePrerequisite(Actor, PostTickFunction);
+					PostTickComponent->PrimaryComponentTick.RemovePrerequisite(InteractionSubsystem, PostTickFunction);
 				}
 			}
-			else
-			{
-				check(!PostTickComponent);
-			}
+			bMainActor = false;
 		}
 
+		check(PostTickFunction.GetPrerequisites().IsEmpty());
+
 		PreTickComponents.Reset();
 		PostTickComponents.Reset();
 
@@ -517,6 +618,11 @@ void FInteractionIsland::Uninject()
 		SearchResults.Reset();
 		bSearchPerfomed = false;
 	}
+
+#if ENABLE_ANIM_DEBUG
+	bPreTickFunctionExecuted = false;
+	bPostTickFunctionExecuted = false;
+#endif // ENABLE_ANIM_DEBUG
 }
 
 bool FInteractionIsland::IsUninjected() const
@@ -530,7 +636,9 @@ const AActor* FInteractionIsland::GetMainActor() const
 	{
 		if (const UActorComponent* ActorComponent = PostTickComponents[0].Get())
 		{
-			return ActorComponent->GetOwner();
+			const AActor* MainActor = ActorComponent->GetOwner();
+			check(MainActor);
+			return MainActor;
 		}
 	}
 	return nullptr;
@@ -552,6 +660,11 @@ bool FInteractionIsland::CanPerformSearch(const UObject* AnimContext) const
 	if (!MainActor)
 	{
 		UE_LOG(LogPoseSearch, Error, TEXT("FInteractionIsland::CanPerformSearch invalid MainActor during the search!"));
+
+#if ENABLE_ANIM_DEBUG
+		LogTickDependencies();
+#endif // ENABLE_ANIM_DEBUG
+
 		return false;
 	}
 
@@ -561,12 +674,24 @@ bool FInteractionIsland::CanPerformSearch(const UObject* AnimContext) const
 	if (Actor != MainActor)
 	{
 		UE_LOG(LogPoseSearch, Error, TEXT("FInteractionIsland::CanPerformSearch search NOT performed by the MainActor!"));
+
+#if ENABLE_ANIM_DEBUG
+		LogTickDependencies();
+#endif // ENABLE_ANIM_DEBUG
 		return false;
 	}
 	
 	return true;
 }
 
+#if ENABLE_ANIM_DEBUG
+void FInteractionIsland::LogTickDependencies() const
+{
+	const ENamedThreads::Type Type = IsInGameThread() ? ENamedThreads::GameThread : ENamedThreads::AnyThread;
+	const_cast<FPostTickFunction&>(PostTickFunction).LogTickFunction(Type, true, 1);
+}
+#endif // ENABLE_ANIM_DEBUG
+
 bool FInteractionIsland::DoSearch_AnyThread(const UObject* AnimContext, const FPoseSearchContinuingProperties& ContinuingProperties, FPoseSearchInteractionBlueprintResult& Result)
 {
 	check(AnimContext);
diff --git a/Engine/Plugins/Animation/PoseSearch/Source/Runtime/Private/PoseSearchInteractionSubsystem.cpp b/Engine/Plugins/Animation/PoseSearch/Source/Runtime/Private/PoseSearchInteractionSubsystem.cpp
--- a/Engine/Plugins/Animation/PoseSearch/Source/Runtime/Private/PoseSearchInteractionSubsystem.cpp
+++ b/Engine/Plugins/Animation/PoseSearch/Source/Runtime/Private/PoseSearchInteractionSubsystem.cpp
@@ -17,6 +17,9 @@ static FAutoConsoleVariableRef CVarPoseSearchInteractionShowIslands(TEXT("a.Pose
 #if !NO_CVARS
 static bool GVarPoseSearchInteractionEnabled = true;
 static FAutoConsoleVariableRef CVarPoseSearchInteractionEnabled(TEXT("a.PoseSearchInteraction.Enabled"), GVarPoseSearchInteractionEnabled, TEXT("Enable/Disable Pose Search Interaction"));
+
+static bool GVarPoseSearchInteractionCacheIslands = true;
+static FAutoConsoleVariableRef CVarPoseSearchInteractionCacheIslands(TEXT("a.PoseSearchInteraction.CacheIslands"), GVarPoseSearchInteractionCacheIslands, TEXT("Cache Pose Search Interaction Islands for future reuse instead of destrying them"));
 #endif // !NO_CVARS
 
 static FVector GetContextLocation(const UObject* AnimContext)
@@ -218,7 +221,7 @@ FCriticalSection UPoseSearchInteractionSubsystem::RetrieveSubsystemMutex;
 
 UE::PoseSearch::FInteractionIsland& UPoseSearchInteractionSubsystem::CreateIsland()
 {
-	return *Islands.Add_GetRef(new UE::PoseSearch::FInteractionIsland(ToRawPtr(GetWorld()->PersistentLevel)));
+	return *Islands.Add_GetRef(new UE::PoseSearch::FInteractionIsland(ToRawPtr(GetWorld()->PersistentLevel), this));
 }
 
 void UPoseSearchInteractionSubsystem::DestroyIsland(int32 Index)
@@ -250,29 +253,34 @@ void UPoseSearchInteractionSubsystem::DestroyAllIslands()
 	}
 }
 
-void UPoseSearchInteractionSubsystem::UninjectAllIslands()
+void UPoseSearchInteractionSubsystem::RegenerateAllIslands(TConstArrayView<UE::PoseSearch::FInteractionSearchContextGroup> SearchContextGroups, TConstArrayView<UE::PoseSearch::FInteractionSearchContext> SearchContexts)
 {
+	QUICK_SCOPE_CYCLE_COUNTER(STAT_UPoseSearchInteractionSubsystem_RegenerateAllIslands);
+
 	using namespace UE::PoseSearch;
 
 	check(IsInGameThread());
 
-	for (FInteractionIsland* Island : Islands)
+#if !NO_CVARS
+	if (!GVarPoseSearchInteractionCacheIslands)
 	{
-		Island->Uninject();
+		// not caching the islands. Destroy them all!
+		DestroyAllIslands();
+	}
+	else
+#endif // !NO_CVARS
+	{
+		for (FInteractionIsland* Island : Islands)
+		{
+			CheckInteractionThreadSafety(Island);
+			Island->Uninject(true);
+		}
 	}
-}
-
-void UPoseSearchInteractionSubsystem::RegenerateAllIslands(TConstArrayView<UE::PoseSearch::FInteractionSearchContextGroup> SearchContextGroups, TConstArrayView<UE::PoseSearch::FInteractionSearchContext> SearchContexts)
-{
-	QUICK_SCOPE_CYCLE_COUNTER(STAT_UPoseSearchInteractionSubsystem_RegenerateAllIslands);
-
-	using namespace UE::PoseSearch;
-
-	UninjectAllIslands();
 
 	for (const FInteractionSearchContextGroup& SearchContextGroup : SearchContextGroups)
 	{
 		FInteractionIsland& Island = GetAvailableIsland();
+		CheckInteractionThreadSafety(&Island);
 
 		for (const UObject* AnimContext : SearchContextGroup.AnimContexts)
 		{
diff --git a/Engine/Plugins/Animation/PoseSearch/Source/Runtime/Private/PoseSearchInteractionValidator.cpp b/Engine/Plugins/Animation/PoseSearch/Source/Runtime/Private/PoseSearchInteractionValidator.cpp
--- a/Engine/Plugins/Animation/PoseSearch/Source/Runtime/Private/PoseSearchInteractionValidator.cpp
+++ b/Engine/Plugins/Animation/PoseSearch/Source/Runtime/Private/PoseSearchInteractionValidator.cpp
@@ -7,8 +7,9 @@ namespace UE::PoseSearch
 {
 #if ENABLE_ANIM_DEBUG
 
-FInteractionValidator::FInteractionValidator(UObject* AnimContext, EPhase Phase)
+FInteractionValidator::FInteractionValidator(UObject* AnimContext)
 	: ValidatingAnimContext(AnimContext)
+	, ValidatingIsland(nullptr)
 {
 	check(AnimContext);
 	UPoseSearchInteractionSubsystem* InteractionSubsystem = UPoseSearchInteractionSubsystem::GetSubsystem_AnyThread(AnimContext);
@@ -38,6 +39,7 @@ FInteractionValidator::FInteractionValidator(UObject* AnimContext, EPhase Phase)
 	if (!MainActor)
 	{
 		UE_LOG(LogPoseSearch, Error, TEXT("FInteractionValidator invalid MainActor! How did he died? Did you rebuild the animation blue print while PIE was running?"));
+		ValidatingIsland->LogTickDependencies();
 		return;
 	}
 
@@ -47,6 +49,7 @@ FInteractionValidator::FInteractionValidator(UObject* AnimContext, EPhase Phase)
 		if (PreviousCounter != 0)
 		{
 			UE_LOG(LogPoseSearch, Error, TEXT("Non thread safe call! Why is there any other actor running while we schedule the MainActor? %d"), PreviousCounter);
+			ValidatingIsland->LogTickDependencies();
 		}
 	}
 	else
@@ -56,12 +59,36 @@ FInteractionValidator::FInteractionValidator(UObject* AnimContext, EPhase Phase)
 		if (PreviousCounter < 0)
 		{
 			UE_LOG(LogPoseSearch, Error, TEXT("Non thread safe call! The MainActor is running! Nobody else in the same island should be running! %d"), PreviousCounter);
+			ValidatingIsland->LogTickDependencies();
 		}
 	}
 }
 
+FInteractionValidator::FInteractionValidator(FInteractionIsland* InValidatingIsland)
+	: ValidatingAnimContext(nullptr)
+	, ValidatingIsland(InValidatingIsland)
+{
+	const int32 PreviousCounter = ValidatingIsland->TickFunctionsThreadSafeCounter.Set(-1);
+	if (PreviousCounter != 0)
+	{
+		UE_LOG(LogPoseSearch, Error, TEXT("Non thread safe call! TickFunctions running concurrently? %d"), PreviousCounter);
+		ValidatingIsland->LogTickDependencies();
+	}
+}
+
 FInteractionValidator::~FInteractionValidator()
 {
+	if (!ValidatingAnimContext)
+	{
+		check(ValidatingIsland);
+		const int32 PreviousCounter = ValidatingIsland->TickFunctionsThreadSafeCounter.Set(0);
+		if (PreviousCounter != -1)
+		{
+			UE_LOG(LogPoseSearch, Error, TEXT("Non thread safe call! TickFunctions running concurrently? %d"), PreviousCounter);
+		}
+		return;
+	}
+
 	if (!ValidatingIsland)
 	{
 		return;
diff --git a/Engine/Plugins/Animation/PoseSearch/Source/Runtime/Private/PoseSearchSchema.cpp b/Engine/Plugins/Animation/PoseSearch/Source/Runtime/Private/PoseSearchSchema.cpp
--- a/Engine/Plugins/Animation/PoseSearch/Source/Runtime/Private/PoseSearchSchema.cpp
+++ b/Engine/Plugins/Animation/PoseSearch/Source/Runtime/Private/PoseSearchSchema.cpp
@@ -353,7 +353,7 @@ void UPoseSearchSchema::Finalize()
 				}
 				else
 				{
-					UE_LOG(LogPoseSearch, Warning, TEXT("UPoseSearchSchema::Finalize: couldn't Finalize '%s' because bone index doest not exist in mirror table or mirrot table is empty."), *GetNameSafe(this));
+					UE_LOG(LogPoseSearch, Warning, TEXT("UPoseSearchSchema::Finalize: couldn't Finalize '%s' because bone index doest not exist in mirror table or mirror table is empty."), *GetNameSafe(this));
 				}
 			}
 		}
diff --git a/Engine/Plugins/Animation/PoseSearch/Source/Runtime/Public/PoseSearch/PoseSearchInteractionIsland.h b/Engine/Plugins/Animation/PoseSearch/Source/Runtime/Public/PoseSearch/PoseSearchInteractionIsland.h
--- a/Engine/Plugins/Animation/PoseSearch/Source/Runtime/Public/PoseSearch/PoseSearchInteractionIsland.h
+++ b/Engine/Plugins/Animation/PoseSearch/Source/Runtime/Public/PoseSearch/PoseSearchInteractionIsland.h
@@ -9,6 +9,7 @@ struct FPoseSearchInteractionBlueprintResult;
 struct FPoseSearchContinuingProperties;
 class UCharacterMovementComponent;
 class UMultiAnimAsset;
+class UPoseSearchInteractionSubsystem;
 
 namespace UE::PoseSearch
 {
@@ -63,7 +64,7 @@ struct FInteractionIsland
 {
 	UE_NONCOPYABLE(FInteractionIsland);
 	
-	FInteractionIsland(ULevel* Level);
+	FInteractionIsland(ULevel* Level, UPoseSearchInteractionSubsystem* Subsystem);
 	~FInteractionIsland();
 
 	bool DoSearch_AnyThread(const UObject* AnimContext, const FPoseSearchContinuingProperties& ContinuingProperties, FPoseSearchInteractionBlueprintResult& Result);
@@ -75,7 +76,7 @@ struct FInteractionIsland
 
 	bool IsUninjected() const;
 	void InjectToActor(const UObject* AnimContext);
-	void Uninject();
+	void Uninject(bool bValidateTickDependencies);
 
 	static UActorComponent* FindPreTickComponent(const AActor* Actor);
 	static UActorComponent* FindPostTickComponent(const AActor* Actor);
@@ -88,10 +89,14 @@ private:
 	const AActor* GetMainActor() const;
 	bool CanPerformSearch(const UObject* AnimContext) const;
 
+#if ENABLE_ANIM_DEBUG
+	void LogTickDependencies() const;
+#endif // ENABLE_ANIM_DEBUG
+
 	struct FPreTickFunction : public FTickFunction
 	{
 		virtual void ExecuteTick(float DeltaTime, enum ELevelTick TickType, ENamedThreads::Type CurrentThread, const FGraphEventRef& MyCompletionGraphEvent) override;
-		virtual FString DiagnosticMessage() override { return TEXT("FPreTickFunction"); }
+		virtual FString DiagnosticMessage() override { return TEXT("UE::PoseSearch::FInteractionIsland::FPreTickFunction"); }
 		FInteractionIsland* Island = nullptr;
 	};
 	FPreTickFunction PreTickFunction;
@@ -99,7 +104,8 @@ private:
 	struct FPostTickFunction : public FTickFunction
 	{
 		virtual void ExecuteTick(float DeltaTime, enum ELevelTick TickType, ENamedThreads::Type CurrentThread, const FGraphEventRef& MyCompletionGraphEvent) override;
-		virtual FString DiagnosticMessage() override { return TEXT("FPostTickFunction"); }
+		virtual FString DiagnosticMessage() override { return TEXT("UE::PoseSearch::FInteractionIsland::FPostTickFunction"); }
+		FInteractionIsland* Island = nullptr;
 	};
 	FPostTickFunction PostTickFunction;
 
@@ -117,6 +123,8 @@ private:
 	TArray<FInteractionSearchResult> SearchResults;
 	bool bSearchPerfomed = false;
 
+	UPoseSearchInteractionSubsystem* InteractionSubsystem = nullptr;
+
 #if UE_POSE_SEARCH_ENABLE_ISLAND_CRITICAL_SECTION
 	FCriticalSection SearchResultsMutex;
 #endif // UE_POSE_SEARCH_ENABLE_ISLAND_CRITICAL_SECTION
@@ -125,6 +133,9 @@ private:
 	// used to analyze thread safety
 	friend struct UE::PoseSearch::FInteractionValidator;
 	FThreadSafeCounter InteractionIslandThreadSafeCounter = 0;
+	FThreadSafeCounter TickFunctionsThreadSafeCounter = 0;
+	bool bPreTickFunctionExecuted = false;
+	bool bPostTickFunctionExecuted = false;
 #endif // ENABLE_ANIM_DEBUG
 };
 
diff --git a/Engine/Plugins/Animation/PoseSearch/Source/Runtime/Public/PoseSearch/PoseSearchInteractionSubsystem.h b/Engine/Plugins/Animation/PoseSearch/Source/Runtime/Public/PoseSearch/PoseSearchInteractionSubsystem.h
--- a/Engine/Plugins/Animation/PoseSearch/Source/Runtime/Public/PoseSearch/PoseSearchInteractionSubsystem.h
+++ b/Engine/Plugins/Animation/PoseSearch/Source/Runtime/Public/PoseSearch/PoseSearchInteractionSubsystem.h
@@ -138,7 +138,6 @@ private:
 
 	void DestroyIsland(int32 Index);
 	void DestroyAllIslands();
-	void UninjectAllIslands();
 	void RegenerateAllIslands(TConstArrayView<UE::PoseSearch::FInteractionSearchContextGroup> SearchContextGroups, TConstArrayView<UE::PoseSearch::FInteractionSearchContext> SearchContexts);
 
 #if DO_CHECK
diff --git a/Engine/Plugins/Animation/PoseSearch/Source/Runtime/Public/PoseSearch/PoseSearchInteractionValidator.h b/Engine/Plugins/Animation/PoseSearch/Source/Runtime/Public/PoseSearch/PoseSearchInteractionValidator.h
--- a/Engine/Plugins/Animation/PoseSearch/Source/Runtime/Public/PoseSearch/PoseSearchInteractionValidator.h
+++ b/Engine/Plugins/Animation/PoseSearch/Source/Runtime/Public/PoseSearch/PoseSearchInteractionValidator.h
@@ -12,24 +12,17 @@ struct FInteractionIsland;
 // Experimental, this feature might be removed without warning, not for production use
 struct FInteractionValidator
 {
-	enum class EPhase
-	{
-		MotionMatchingInteractionUpdate,
-		MotionMatchingInteractionEvaluate,
-		HistoryCollectorUpdate,
-		HistoryCollectorEvaluate
-	};
-
-	FInteractionValidator(UObject* AnimContext, EPhase Phase);
+	explicit FInteractionValidator(UObject* AnimContext);
+	explicit FInteractionValidator(FInteractionIsland* InValidatingIsland);
 	~FInteractionValidator();
 
 	const UObject* ValidatingAnimContext = nullptr;
 	FInteractionIsland* ValidatingIsland = nullptr;
 };
 
-#define CheckInteractionThreadSafety(AnimContext, Phase) FInteractionValidator(AnimContext, Phase)
+#define CheckInteractionThreadSafety(ValidationContext) const FInteractionValidator InteractionValidator(ValidationContext)
 #else
-#define CheckInteractionThreadSafety(AnimContext, Phase)
+#define CheckInteractionThreadSafety(ValidationContext)
 #endif // ENABLE_ANIM_DEBUG
 
 } // namespace UE::PoseSearch
\ No newline at end of file
