posesearch - making MM interaction respecting FAnimNode_PoseSearchHistoryCollector_Base::bGenerateTrajectory property

[RN] skip

diff --git a/Engine/Plugins/Animation/PoseSearch/Source/Runtime/Private/AnimNode_PoseSearchHistoryCollector.cpp b/Engine/Plugins/Animation/PoseSearch/Source/Runtime/Private/AnimNode_PoseSearchHistoryCollector.cpp
--- a/Engine/Plugins/Animation/PoseSearch/Source/Runtime/Private/AnimNode_PoseSearchHistoryCollector.cpp
+++ b/Engine/Plugins/Animation/PoseSearch/Source/Runtime/Private/AnimNode_PoseSearchHistoryCollector.cpp
@@ -69,16 +69,29 @@ void FAnimNode_PoseSearchHistoryCollector_Base::CacheBones_AnyThread(const FAnim
 
 void FAnimNode_PoseSearchHistoryCollector_Base::Update_AnyThread(const FAnimationUpdateContext& Context)
 {
-	GetEvaluateGraphExposedInputs().Execute(Context);
+	using namespace UE::PoseSearch;
 
-	PoseHistory.SetTrajectory(bGenerateTrajectory ? FPoseSearchQueryTrajectory() : Trajectory, TrajectorySpeedMultiplier);
+	check(Context.AnimInstanceProxy);
+	CheckInteractionThreadSafety(Context.AnimInstanceProxy->GetAnimInstanceObject());
 
+	GetEvaluateGraphExposedInputs().Execute(Context);
+	
+	if (bGenerateTrajectory)
+	{
+		GenerateTrajectory(Cast<UAnimInstance>(Context.AnimInstanceProxy->GetAnimInstanceObject()));
+	}
+	else
+	{
+		PoseHistory.SetTrajectory(Trajectory, TrajectorySpeedMultiplier);
+	}
+
+	bIsTrajectoryGeneratedBeforePreUpdate = false;
 	UpdateCounter.SynchronizeWith(Context.AnimInstanceProxy->GetUpdateCounter());
 }
 
 void FAnimNode_PoseSearchHistoryCollector_Base::GenerateTrajectory(const UAnimInstance* InAnimInstance)
 {
-	if (!bIsTrajectoryGeneratedBeforePreUpdate)
+	if (bGenerateTrajectory && !bIsTrajectoryGeneratedBeforePreUpdate)
 	{
 		FPoseSearchTrajectoryData::FSampling TrajectoryDataSampling;
 		TrajectoryDataSampling.NumHistorySamples = FMath::Max(PoseCount, TrajectoryHistoryCount);
@@ -151,30 +164,6 @@ void FAnimNode_PoseSearchHistoryCollector::Evaluate_AnyThread(FPoseContext& Outp
 #endif // ENABLE_DRAW_DEBUG && ENABLE_ANIM_DEBUG
 }
 
-void FAnimNode_PoseSearchHistoryCollector::Update_AnyThread(const FAnimationUpdateContext& Context)
-{
-	DECLARE_SCOPE_HIERARCHICAL_COUNTER_ANIMNODE(Update_AnyThread);
-
-	using namespace UE::PoseSearch;
-	
-	// scoping CheckInteractionThreadSafety
-	{
-		check(Context.AnimInstanceProxy);
-		CheckInteractionThreadSafety(Context.AnimInstanceProxy->GetAnimInstanceObject());
-
-		if (bGenerateTrajectory)
-		{
-			GenerateTrajectory(Cast<UAnimInstance>(Context.AnimInstanceProxy->GetAnimInstanceObject()));
-		}
-		bIsTrajectoryGeneratedBeforePreUpdate = false;
-
-		Super::Update_AnyThread(Context);
-	}
-	
-	UE::Anim::TScopedGraphMessage<UE::PoseSearch::FPoseHistoryProvider> ScopedMessage(Context, this);
-	Source.Update(Context);
-}
-
 void FAnimNode_PoseSearchHistoryCollector::GatherDebugData(FNodeDebugData& DebugData)
 {
 	Super::GatherDebugData(DebugData);
@@ -238,27 +227,20 @@ void FAnimNode_PoseSearchComponentSpaceHistoryCollector::EvaluateComponentSpace_
 #endif // ENABLE_DRAW_DEBUG && ENABLE_ANIM_DEBUG
 }
 
-void FAnimNode_PoseSearchComponentSpaceHistoryCollector::Update_AnyThread(const FAnimationUpdateContext& Context)
+void FAnimNode_PoseSearchHistoryCollector::Update_AnyThread(const FAnimationUpdateContext& Context)
 {
 	DECLARE_SCOPE_HIERARCHICAL_COUNTER_ANIMNODE(Update_AnyThread);
 
-	using namespace UE::PoseSearch;
-
-
-	// scoping CheckInteractionThreadSafety
-	{
-		check(Context.AnimInstanceProxy);
-		CheckInteractionThreadSafety(Context.AnimInstanceProxy->GetAnimInstanceObject());
-
-		if (bGenerateTrajectory)
-		{
-			GenerateTrajectory(Cast<UAnimInstance>(Context.AnimInstanceProxy->GetAnimInstanceObject()));
-		}
-		bIsTrajectoryGeneratedBeforePreUpdate = false;
+	Super::Update_AnyThread(Context);
+	UE::Anim::TScopedGraphMessage<UE::PoseSearch::FPoseHistoryProvider> ScopedMessage(Context, this);
+	Source.Update(Context);
+}
 
-		Super::Update_AnyThread(Context);
-	}
+void FAnimNode_PoseSearchComponentSpaceHistoryCollector::Update_AnyThread(const FAnimationUpdateContext& Context)
+{
+	DECLARE_SCOPE_HIERARCHICAL_COUNTER_ANIMNODE(Update_AnyThread);
 
+	Super::Update_AnyThread(Context);
 	UE::Anim::TScopedGraphMessage<UE::PoseSearch::FPoseHistoryProvider> ScopedMessage(Context, this);
 	Source.Update(Context);
 }
