pose search - introduced FPoseSearchInteractionAvailability::Tag to facilitate the setup of interactions, and their memory managment

[RN] skip

diff --git a/Engine/Plugins/Animation/PoseSearch/Source/Runtime/Private/PoseSearchInteractionSubsystem.cpp b/Engine/Plugins/Animation/PoseSearch/Source/Runtime/Private/PoseSearchInteractionSubsystem.cpp
--- a/Engine/Plugins/Animation/PoseSearch/Source/Runtime/Private/PoseSearchInteractionSubsystem.cpp
+++ b/Engine/Plugins/Animation/PoseSearch/Source/Runtime/Private/PoseSearchInteractionSubsystem.cpp
@@ -152,6 +152,9 @@ struct FAnimContextInfo
 };
 struct FAnimContextInfos : TArray<FAnimContextInfo, TMemStackAllocator<>> {};
 
+typedef TArray<const UPoseSearchDatabase*, TInlineAllocator<32, TMemStackAllocator<>>> FDatabasesPerTag;
+struct FTagToDatabases : TMap<FName, FDatabasesPerTag, TInlineSetAllocator<8, TMemStackSetAllocator<>>> {};
+
 template <typename EvaluateCombinationType>
 static void GeneratePermutationsRecursive(int32 DataCardinality, int32 DataIndex, TArrayView<int32> Combination, int32 CombinationIndex, EvaluateCombinationType EvaluateCombination)
 {
@@ -180,6 +183,8 @@ static void GenerateInteractionSearchContextGroups(const FInteractionSearchConte
 {
 	QUICK_SCOPE_CYCLE_COUNTER(STAT_UPoseSearchInteractionSubsystem_GenerateInteractionSearchContextGroups);
 
+	check(SearchContextGroups.IsEmpty());
+
 	for (int32 SearchContextIndex = 0; SearchContextIndex < SearchContexts.Num(); ++SearchContextIndex)
 	{
 		const FInteractionSearchContext& SearchContext = SearchContexts[SearchContextIndex];
@@ -490,7 +495,7 @@ void UPoseSearchInteractionSubsystem::ResetAvailabilities()
 	AvailabilitiesMap.Reset();
 }
 
-void UPoseSearchInteractionSubsystem::GenerateAnimContextInfos(UE::PoseSearch::FAnimContextInfos& AnimContextInfos) const
+void UPoseSearchInteractionSubsystem::GenerateAnimContextInfosAndTagToDatabases(UE::PoseSearch::FAnimContextInfos& AnimContextInfos, UE::PoseSearch::FTagToDatabases& TagToDatabases) const
 {
 	QUICK_SCOPE_CYCLE_COUNTER(STAT_UPoseSearchInteractionSubsystem_GenerateAnimContextInfos);
 
@@ -498,6 +503,7 @@ void UPoseSearchInteractionSubsystem::GenerateAnimContextInfos(UE::PoseSearch::F
 
 	const UWorld* SubsystemWorld = GetWorld();
 	check(SubsystemWorld);
+	check(AnimContextInfos.IsEmpty() && TagToDatabases.IsEmpty());
 
 	// collecting and sorting all the AnimContext(s) from the AvailabilitiesMap in AnimContextsWithAvailabilities, to have a stable and deterministic island creation
 	typedef TPair<const UObject*, const FInteractionAvailabilities*> FAnimContextWithAvailabilities;
@@ -508,19 +514,63 @@ void UPoseSearchInteractionSubsystem::GenerateAnimContextInfos(UE::PoseSearch::F
 		const UObject* AnimContext = Iter->Key.Get();
 		if (!AnimContext)
 		{
-			UE_LOG(LogPoseSearch, Log, TEXT("UPoseSearchInteractionSubsystem::GenerateAnimContextInfos null anim context. The associated character got removed from the world"));
+			UE_LOG(LogPoseSearch, Log, TEXT("UPoseSearchInteractionSubsystem::GenerateAnimContextInfosAndTagToDatabaseMap null anim context. The associated character got removed from the world"));
 		}
 		else if (!AnimContext->GetWorld())
 		{
-			UE_LOG(LogPoseSearch, Error, TEXT("UPoseSearchInteractionSubsystem::GenerateAnimContextInfos AnimContext %s is not in a world!"), *AnimContext->GetName());
+			UE_LOG(LogPoseSearch, Error, TEXT("UPoseSearchInteractionSubsystem::GenerateAnimContextInfosAndTagToDatabaseMap AnimContext %s is not in a world!"), *AnimContext->GetName());
 		}
 		else if (SubsystemWorld != AnimContext->GetWorld())
 		{
-			UE_LOG(LogPoseSearch, Error, TEXT("UPoseSearchInteractionSubsystem::GenerateAnimContextInfos AnimContext %s is from World %s, and supposed to be in World @s!"), *AnimContext->GetName(), *AnimContext->GetWorld()->GetName(), *SubsystemWorld->GetName());
+			UE_LOG(LogPoseSearch, Error, TEXT("UPoseSearchInteractionSubsystem::GenerateAnimContextInfosAndTagToDatabaseMap AnimContext %s is from World %s, and supposed to be in World @s!"), *AnimContext->GetName(), *AnimContext->GetWorld()->GetName(), *SubsystemWorld->GetName());
 		}
 		else
 		{
-			AnimContextsWithAvailabilities.Add({ AnimContext, &Iter->Value });
+			bool bContinueEvaluatingAnimContext = !Iter->Value.IsEmpty();
+			if (bContinueEvaluatingAnimContext)
+			{
+				// adding AnimContext to AnimContextsWithAvailabilities only if at least one availability has a valid database or has a valid tag
+				for (const FInteractionAvailability& InteractionAvailability : Iter->Value)
+				{
+					if (const UPoseSearchDatabase* Database = InteractionAvailability.Database.Get())
+					{
+						if (!Database->Schema)
+						{
+							UE_LOG(LogPoseSearch, Error, TEXT("UPoseSearchInteractionSubsystem::GenerateAnimContextInfosAndTagToDatabaseMap null Schema for Database %s. Dropping ALL availabilities for AnimContext %s"), *Database->GetName(), *AnimContext->GetName());
+							bContinueEvaluatingAnimContext = false;
+							break;
+						}
+					}
+				}
+			}
+
+			if (bContinueEvaluatingAnimContext)
+			{
+				bool bAnyValidAvailability = false;
+
+				for (const FInteractionAvailability& InteractionAvailability : Iter->Value)
+				{
+					const bool bIsTagValid = InteractionAvailability.IsTagValid();
+					if (const UPoseSearchDatabase* Database = InteractionAvailability.Database.Get())
+					{
+						check(Database->Schema);
+						bAnyValidAvailability = true;
+						if (bIsTagValid)
+						{
+							TagToDatabases.FindOrAdd(InteractionAvailability.Tag).AddUnique(InteractionAvailability.Database);
+						}
+					}
+					else if (bIsTagValid)
+					{
+						bAnyValidAvailability = true;
+					}
+				}
+			
+				if (bAnyValidAvailability)
+				{
+					AnimContextsWithAvailabilities.Add({ AnimContext, &Iter->Value });
+				}
+			}
 		}
 	}
 
@@ -575,16 +625,18 @@ void UPoseSearchInteractionSubsystem::GenerateSearchContexts(UE::PoseSearch::FIn
 
 	using namespace UE::PoseSearch;
 
+	check(SearchContexts.IsEmpty());
+
 	struct FRoledAnimContextInfo
 	{
-		FRoledAnimContextInfo(const FAnimContextInfo* InAnimContextInfo, const FRole InRole, const FAnimNode_PoseSearchHistoryCollector_Base* InHistoryCollector, const UPoseSearchDatabase* InDatabase, float InMaxDistanceSquared = 0.f)
-			: AnimContextInfo(InAnimContextInfo)
+		FRoledAnimContextInfo(const FAnimContextInfo& InAnimContextInfo, const FRole InRole, const FAnimNode_PoseSearchHistoryCollector_Base& InHistoryCollector, const UPoseSearchDatabase& InDatabase, float InMaxDistanceSquared)
+			: AnimContextInfo(&InAnimContextInfo)
 			, Role(InRole)
-			, HistoryCollector(InHistoryCollector)
-			, Database(InDatabase)
+			, HistoryCollector(&InHistoryCollector)
+			, Database(&InDatabase)
 			, MaxDistanceSquared(InMaxDistanceSquared)
 		{
-			check(InAnimContextInfo && InHistoryCollector && InDatabase && MaxDistanceSquared > UE_SMALL_NUMBER);
+			check(MaxDistanceSquared > UE_SMALL_NUMBER);
 		}
 
 		bool operator==(const FRoledAnimContextInfo& Other) const = default;
@@ -595,7 +647,36 @@ void UPoseSearchInteractionSubsystem::GenerateSearchContexts(UE::PoseSearch::FIn
 		const UPoseSearchDatabase* Database;
 		float MaxDistanceSquared = 0.f;
 	};
-	typedef TArray<FRoledAnimContextInfo, TInlineAllocator<PreallocatedRolesNum, TMemStackAllocator<>>> FRoledAnimContextInfos;
+	struct FRoledAnimContextInfos : TArray<FRoledAnimContextInfo, TInlineAllocator<PreallocatedRolesNum, TMemStackAllocator<>>>
+	{
+		void AddRoledAnimContextInfos(const FAnimContextInfo& AnimContextInfo, const TArray<FName>& RolesFilter, const FAnimNode_PoseSearchHistoryCollector_Base& HistoryCollector, const UPoseSearchDatabase& Database, float MaxDistanceSquared)
+		{
+			const UPoseSearchSchema* Schema = Database.Schema;
+			check(Schema);
+			if (RolesFilter.IsEmpty())
+			{
+				// adding ALL the possible roles from the database:
+				for (const FPoseSearchRoledSkeleton& RoledSkeleton : Schema->GetRoledSkeletons())
+				{
+					AddUnique(FRoledAnimContextInfo(AnimContextInfo, RoledSkeleton.Role, HistoryCollector, Database, MaxDistanceSquared));
+				}
+			}
+			else
+			{
+				for (const FRole& Role : RolesFilter)
+				{
+					if (Schema->GetRoledSkeleton(Role))
+					{
+						AddUnique(FRoledAnimContextInfo(AnimContextInfo, Role, HistoryCollector, Database, MaxDistanceSquared));
+					}
+					else
+					{
+						UE_LOG(LogPoseSearch, Warning, TEXT("UPoseSearchInteractionSubsystem::GenerateSearchContexts unsupported Role %s for Database %s"), *Role.ToString(), *Database.GetName());
+					}
+				}
+			}
+		}
+	};
 
 	// visits the FAnimContextInfos recursively to identify groups of nearby AnimContextInfo(s), relying on FAnimContextInfo::NearbyAnimContextInfos information.
 	// it calls OnNewAnimAnimContextInfoFound on every new FAnimContextInfo found/visited, and OnDoneGroupingAnimContexts once reaches the end of the current group AnimContext(s).
@@ -646,14 +727,16 @@ void UPoseSearchInteractionSubsystem::GenerateSearchContexts(UE::PoseSearch::FIn
 
 	// caching AnimContext(s) locations, max broad phase radiuses (squared) and collect relations of possible interactions between AnimContext(s)
 	// (stored in FAnimContextInfo::NearbyAnimContextInfos::Index) as fast broad phase evaluation refined later on during FInteractionSearchContexts generation
+	// and generating a mapping between availabilities Tag(s) to availabilities published databases
 	FAnimContextInfos AnimContextInfos;
-	GenerateAnimContextInfos(AnimContextInfos);
+	FTagToDatabases TagToDatabases;
+	GenerateAnimContextInfosAndTagToDatabases(AnimContextInfos, TagToDatabases);
 
 	TMap<const UPoseSearchDatabase*, FRoledAnimContextInfos, TInlineSetAllocator<PreallocatedSearchesNum, TMemStackSetAllocator<>>> DatabaseToRoledAnimContextInfos;
 	// visiting ALL the AnimContexts in AnimContextInfos and relying on the cached information to refine potential interactions
 	FAnimContextInfoVisitor AnimContextInfoVisitor(AnimContextInfos,
 		// OnNewAnimContextFound: called when the FAnimContextInfoVisitor finds a new AnimContext that can be grouped in the current group of possibly interacting AnimContext(s)
-		[&DatabaseToRoledAnimContextInfos](const FAnimContextInfo& AnimContextInfo)
+		[&TagToDatabases, &DatabaseToRoledAnimContextInfos](const FAnimContextInfo& AnimContextInfo)
 		{
 			// analyzing all the Availability(s) associated with this AnimContext and eventually generate the associated FRoledAnimContextInfos,
 			// inserted in a per database sorted data structure (DatabaseToRoledAnimContextInfos)
@@ -663,51 +746,46 @@ void UPoseSearchInteractionSubsystem::GenerateSearchContexts(UE::PoseSearch::FIn
 				if (Availability.BroadPhaseRadius <= UE_SMALL_NUMBER)
 				{
 					UE_LOG(LogPoseSearch, Log, TEXT("UPoseSearchInteractionSubsystem::GenerateSearchContexts too small Availability.BroadPhaseRadius. Character will not interact with anything within such small radius!"));
-					continue;
-				}
-
-				if (!Availability.Database)
-				{
-					UE_LOG(LogPoseSearch, Log, TEXT("UPoseSearchInteractionSubsystem::GenerateSearchContexts null Availability.Database"));
-					continue;
-				}
-
-				if (!Availability.Database->Schema)
-				{
-					UE_LOG(LogPoseSearch, Error, TEXT("UPoseSearchInteractionSubsystem::GenerateSearchContexts null Schema for Database %s"), *Availability.Database->GetName());
-					continue;
 				}
-
-				const FAnimNode_PoseSearchHistoryCollector_Base* HistoryCollector = Availability.GetHistoryCollector(AnimContextInfo.AnimContext);
-				if (!HistoryCollector)
-				{
-					UE_LOG(LogPoseSearch, Error, TEXT("UPoseSearchInteractionSubsystem::GenerateSearchContexts couldn't find PoseHistory %s for AnimContext %s"), *Availability.GetPoseHistoryName(), *AnimContextInfo.AnimContext->GetName());
-					continue;
-				}
-
-				FRoledAnimContextInfos& RoledAnimContextInfos = DatabaseToRoledAnimContextInfos.FindOrAdd(Availability.Database);
-				const float MaxDistanceSquared = Availability.BroadPhaseRadius * Availability.BroadPhaseRadius;
-				if (Availability.RolesFilter.IsEmpty())
+				else if (const FAnimNode_PoseSearchHistoryCollector_Base* HistoryCollector = Availability.GetHistoryCollector(AnimContextInfo.AnimContext))
 				{
-					// adding ALL the possible roles from the database:
-					for (const FPoseSearchRoledSkeleton& RoledSkeleton : Availability.Database->Schema->GetRoledSkeletons())
+					if (const UPoseSearchDatabase* Database = Availability.Database.Get())
 					{
-						RoledAnimContextInfos.AddUnique(FRoledAnimContextInfo(&AnimContextInfo, RoledSkeleton.Role, HistoryCollector, Availability.Database, MaxDistanceSquared));
+						check(Database->Schema);
+						FRoledAnimContextInfos& RoledAnimContextInfos = DatabaseToRoledAnimContextInfos.FindOrAdd(Database);
+						const float MaxDistanceSquared = Availability.BroadPhaseRadius * Availability.BroadPhaseRadius;
+						RoledAnimContextInfos.AddRoledAnimContextInfos(AnimContextInfo, Availability.RolesFilter, *HistoryCollector, *Database, MaxDistanceSquared);
 					}
-				}
-				else
-				{
-					for (const FRole& Role : Availability.RolesFilter)
+					else if (Availability.IsTagValid())
 					{
-						if (Availability.Database->Schema->GetRoledSkeleton(Role))
+						// since Database is null, but this availability has a valid Tag, we're looking for valid databases by Availability.Tag
+						if (const FDatabasesPerTag* DatabasesPerTag = TagToDatabases.Find(Availability.Tag))
 						{
-							RoledAnimContextInfos.AddUnique(FRoledAnimContextInfo(&AnimContextInfo, Role, HistoryCollector, Availability.Database, MaxDistanceSquared));
+							check(!DatabasesPerTag->IsEmpty());
+							const float MaxDistanceSquared = Availability.BroadPhaseRadius * Availability.BroadPhaseRadius;
+							for (const UPoseSearchDatabase* DatabaseFromTag : *DatabasesPerTag)
+							{
+								check(DatabaseFromTag && DatabaseFromTag->Schema);
+								FRoledAnimContextInfos& RoledAnimContextInfos = DatabaseToRoledAnimContextInfos.FindOrAdd(DatabaseFromTag);
+								RoledAnimContextInfos.AddRoledAnimContextInfos(AnimContextInfo, Availability.RolesFilter, *HistoryCollector, *DatabaseFromTag, MaxDistanceSquared);
+							}
 						}
 						else
 						{
-							UE_LOG(LogPoseSearch, Warning, TEXT("UPoseSearchInteractionSubsystem::GenerateSearchContexts unsupported Role %s for Database %s"), *Role.ToString(), *Availability.Database->GetName());
+							//@todo: should we add a verbose LOG here? not sure since it'd be very spammy...
+							
+							// this is a valid condition we shouldn't log: for example when the "main character" is loaded and publishing availabilities with a valid tag and null database,
+							// looking for other NPC / seconday characters to interact with, but they are not present of didn't publish any availability
 						}
 					}
+					else
+					{
+						UE_LOG(LogPoseSearch, Log, TEXT("UPoseSearchInteractionSubsystem::GenerateSearchContexts null Availability.Database (with invalid Availability.Tag)"));
+					}
+				}
+				else
+				{
+					UE_LOG(LogPoseSearch, Error, TEXT("UPoseSearchInteractionSubsystem::GenerateSearchContexts couldn't find PoseHistory %s for AnimContext %s"), *Availability.GetPoseHistoryName(), *AnimContextInfo.AnimContext->GetName());
 				}
 			}
 		},
@@ -736,6 +814,7 @@ void UPoseSearchInteractionSubsystem::GenerateSearchContexts(UE::PoseSearch::FIn
 			for (const TPair<const UPoseSearchDatabase*, FRoledAnimContextInfos>& DatabaseToRoledAnimContextInfosPair : DatabaseToRoledAnimContextInfos)
 			{
 				const UPoseSearchDatabase* Database = DatabaseToRoledAnimContextInfosPair.Key;
+				check(Database->Schema);
 				const TArray<FPoseSearchRoledSkeleton>& RoledSkeletons = Database->Schema->GetRoledSkeletons();
 				const int32 CombinationCardinality = RoledSkeletons.Num();
 				const FRoledAnimContextInfos& RoledAnimContextInfos = DatabaseToRoledAnimContextInfosPair.Value;
@@ -905,8 +984,8 @@ void UPoseSearchInteractionSubsystem::Query_AnyThread(const TArrayView<const FPo
 
 			for (const FPoseSearchInteractionAvailability& Availability : Availabilities)
 			{
-				if (Availability.Database == Result.SelectedDatabase &&
-					(Availability.RolesFilter.IsEmpty() || Availability.RolesFilter.Contains(Result.Role)))
+				const bool bIsDatabaseValidates = (Availability.IsTagValid() && !Availability.Database) || (Availability.Database == Result.SelectedDatabase);
+				if (bIsDatabaseValidates &&	(Availability.RolesFilter.IsEmpty() || Availability.RolesFilter.Contains(Result.Role)))
 				{
 					bResultValidated = true;
 					break;
diff --git a/Engine/Plugins/Animation/PoseSearch/Source/Runtime/Public/PoseSearch/PoseSearchInteractionLibrary.h b/Engine/Plugins/Animation/PoseSearch/Source/Runtime/Public/PoseSearch/PoseSearchInteractionLibrary.h
--- a/Engine/Plugins/Animation/PoseSearch/Source/Runtime/Public/PoseSearch/PoseSearchInteractionLibrary.h
+++ b/Engine/Plugins/Animation/PoseSearch/Source/Runtime/Public/PoseSearch/PoseSearchInteractionLibrary.h
@@ -22,6 +22,15 @@ struct POSESEARCH_API FPoseSearchInteractionAvailability
 	UPROPERTY(EditAnywhere, BlueprintReadWrite, Category=Settings)
 	TObjectPtr<UPoseSearchDatabase> Database;
 
+	// in case this availability Database is valid (not null), Tag (if IsTagValid()) is used to flag the Database with a specific name. Different availabilities can share the same Tag.
+	// in case this availability Database is NOT valid, we use the valid Tag to figure out all the possible databases that can be assigned to this availability from all the published availabilities.
+	// The reason behind Tag is, for example, to be able to have NPCs been able to interact with a main character (MC), without the MC having a direct dependency to the database used 
+	// for the interaction allowing those NPCs to be contextually loaded/unloaded, streamed in/out, with the obvious advantages for the the memory managment of the "payload" database.
+	// Another reason for Tag, is to facilitate the setup of interactions, where the MC have to publish only one availability with its own assigned Role (in RolesFilter) 
+	// automatically contextually resolved in multiple different types of possible databases: it could be MC-NPC, MC-Vehicle, MC-Whatever
+	UPROPERTY(EditAnywhere, BlueprintReadWrite, Category=Settings)
+	FName Tag;
+
 	// roles the character is willing to take to partecipate in this interaction. If empty ANY of the Database roles can be taken
 	UPROPERTY(EditAnywhere, BlueprintReadWrite, Category=Settings)
 	TArray<FName> RolesFilter;
@@ -29,6 +38,8 @@ struct POSESEARCH_API FPoseSearchInteractionAvailability
 	// the associated character to this FPoseSearchInteractionAvailability will partecipate in an interaction only if all the necessary roles gest assigned to character within BroadPhaseRadius centimeters
 	UPROPERTY(EditAnywhere, BlueprintReadWrite, Category=Settings)
 	float BroadPhaseRadius = 500.f;
+
+	bool IsTagValid() const { return !Tag.IsNone(); }
 };
 
 // Experimental, this feature might be removed without warning, not for production use
diff --git a/Engine/Plugins/Animation/PoseSearch/Source/Runtime/Public/PoseSearch/PoseSearchInteractionSubsystem.h b/Engine/Plugins/Animation/PoseSearch/Source/Runtime/Public/PoseSearch/PoseSearchInteractionSubsystem.h
--- a/Engine/Plugins/Animation/PoseSearch/Source/Runtime/Public/PoseSearch/PoseSearchInteractionSubsystem.h
+++ b/Engine/Plugins/Animation/PoseSearch/Source/Runtime/Public/PoseSearch/PoseSearchInteractionSubsystem.h
@@ -48,16 +48,21 @@ struct FAvailabilitiesMap : protected TMap<TWeakObjectPtr<const UObject>, FInter
 		FInteractionAvailabilities& AnimContextAvailabilities = FindOrAdd(AnimContext);
 		for (const FPoseSearchInteractionAvailability& Availability : Availabilities)
 		{
-			if (Availability.Database)
+			if (Availability.Database || Availability.IsTagValid())
 			{
 				AnimContextAvailabilities.Emplace(Availability, PoseHistoryName, HistoryCollector);
 			}
+			else
+			{
+				UE_LOG(LogPoseSearch, Warning, TEXT("FAvailabilitiesMap::AddAvailabilities trying adding an invalid Availability for AnimContext %s!"), *AnimContext->GetName());
+			}
 		}
 	}
 };
 
 struct FAnimContextInfos;
 struct FInteractionSearchContextGroup;
+struct FTagToDatabases;
 
 } // namespace UE::PoseSearch
 
@@ -155,7 +160,7 @@ private:
 	bool IsThereAnyAvailability() const;
 	void AddAvailabilities(const TArrayView<const FPoseSearchInteractionAvailability> Availabilities, const UObject* AnimContext, FName PoseHistoryName, const FAnimNode_PoseSearchHistoryCollector_Base* HistoryCollector);
 	void ResetAvailabilities();
-	void GenerateAnimContextInfos(UE::PoseSearch::FAnimContextInfos& AnimContextInfos) const;
+	void GenerateAnimContextInfosAndTagToDatabases(UE::PoseSearch::FAnimContextInfos& AnimContextInfos, UE::PoseSearch::FTagToDatabases& TagToDatabases) const;
 	void GenerateSearchContexts(UE::PoseSearch::FInteractionSearchContexts& SearchContexts) const;
 
 	// @todo: use a thread local storage data structure to make publishing availabilities lockless to avoid the need of AvailabilitiesMapMutex
