posesearch - added UPoseSearchInteractionAsset::GetPreviewMesh

[RN] skip

diff --git a/Engine/Plugins/Animation/PoseSearch/Source/Editor/Private/PoseSearchDatabaseViewModel.cpp b/Engine/Plugins/Animation/PoseSearch/Source/Editor/Private/PoseSearchDatabaseViewModel.cpp
--- a/Engine/Plugins/Animation/PoseSearch/Source/Editor/Private/PoseSearchDatabaseViewModel.cpp
+++ b/Engine/Plugins/Animation/PoseSearch/Source/Editor/Private/PoseSearchDatabaseViewModel.cpp
@@ -80,8 +80,17 @@ bool FDatabasePreviewActor::SpawnPreviewActor(UWorld* World, const UPoseSearchDa
 	Mesh->PreviewInstance = AnimInstance;
 	AnimInstance->InitializeAnimation();
 
-	USkeletalMesh* DatabasePreviewMesh = PoseSearchDatabase->PreviewMesh;
-	Mesh->SetSkeletalMesh(DatabasePreviewMesh ? DatabasePreviewMesh : Skeleton->GetPreviewMesh(true));
+	USkeletalMesh* PreviewMesh = DatabaseAnimationAsset->GetPreviewMeshForRole(Role);
+	if (!PreviewMesh)
+	{
+		PreviewMesh = PoseSearchDatabase->PreviewMesh;
+		if (!PreviewMesh)
+		{
+			PreviewMesh = Skeleton->GetPreviewMesh(true);
+		}
+	}
+	
+	Mesh->SetSkeletalMesh(PreviewMesh);
 	Mesh->EnablePreview(true, PreviewAsset);
 		
 	AnimInstance->SetAnimationAsset(PreviewAsset, IndexAsset.IsLooping(), 0.f);
@@ -161,7 +170,7 @@ void FDatabasePreviewActor::UpdatePreviewActor(const UPoseSearchDatabase* PoseSe
 	const FSearchIndex& SearchIndex = PoseSearchDatabase->GetSearchIndex();
 
 	UAnimPreviewInstance* AnimInstance = GetAnimPreviewInstanceInternal();
-	if (!AnimInstance || IndexAssetIndex >= SearchIndex.Assets.Num())
+	if (!AnimInstance || !SearchIndex.Assets.IsValidIndex(IndexAssetIndex))
 	{
 		return;
 	}
@@ -172,8 +181,34 @@ void FDatabasePreviewActor::UpdatePreviewActor(const UPoseSearchDatabase* PoseSe
 		return;
 	}
 
-	CurrentTime = 0.f;
+	UDebugSkelMeshComponent* Mesh = GetDebugSkelMeshComponent();
+	if (!Mesh)
+	{
+		return;
+	}
+
 	const FSearchIndexAsset& IndexAsset = SearchIndex.Assets[IndexAssetIndex];
+	const FPoseSearchDatabaseAnimationAssetBase* DatabaseAnimationAsset = PoseSearchDatabase->GetDatabaseAnimationAsset<FPoseSearchDatabaseAnimationAssetBase>(IndexAsset.GetSourceAssetIdx());
+
+	USkeletalMesh* PreviewMesh = DatabaseAnimationAsset->GetPreviewMeshForRole(ActorRole);
+	if (!PreviewMesh)
+	{
+		PreviewMesh = PoseSearchDatabase->PreviewMesh;
+		if (!PreviewMesh)
+		{
+			if (USkeleton* Skeleton = PoseSearchDatabase->Schema->GetSkeleton(ActorRole))
+			{
+				PreviewMesh = Skeleton->GetPreviewMesh(true);
+			}
+		}
+	}
+
+	if (Mesh->GetSkeletalMeshAsset() != PreviewMesh)
+	{
+		Mesh->SetSkeletalMesh(PreviewMesh);
+	}	
+
+	CurrentTime = 0.f;
 	float CurrentPlayTime = PlayTime + IndexAsset.GetFirstSampleTime(PoseSearchDatabase->Schema->SampleRate) + PlayTimeOffset;
 	FAnimationRuntime::AdvanceTime(false, CurrentPlayTime, CurrentTime, IndexAsset.GetLastSampleTime(PoseSearchDatabase->Schema->SampleRate));
 			 
@@ -424,6 +459,15 @@ const UDebugSkelMeshComponent* FDatabasePreviewActor::GetDebugSkelMeshComponent(
 	return nullptr;
 }
 
+UDebugSkelMeshComponent* FDatabasePreviewActor::GetDebugSkelMeshComponent()
+{
+	if (ActorPtr != nullptr)
+	{
+		return Cast<UDebugSkelMeshComponent>(ActorPtr->GetRootComponent());
+	}
+	return nullptr;
+}
+
 const UAnimPreviewInstance* FDatabasePreviewActor::GetAnimPreviewInstance() const
 {
 	if (const UDebugSkelMeshComponent* Mesh = GetDebugSkelMeshComponent())
@@ -523,17 +567,14 @@ void FDatabaseViewModel::Tick(float DeltaSeconds)
 		const UPoseSearchDatabase* Database = GetPoseSearchDatabase();
 		if (EAsyncBuildIndexResult::Success == FAsyncPoseSearchDatabasesManagement::RequestAsyncBuildIndex(Database, ERequestAsyncBuildFlag::ContinueRequest))
 		{
-			if (!FMath::IsNearlyZero(DeltaPlayTime))
-			{
-				PlayTime += DeltaPlayTime;
-				PlayTime = FMath::Clamp(PlayTime, MinPreviewPlayLength, MaxPreviewPlayLength);
+			PlayTime += DeltaPlayTime;
+			PlayTime = FMath::Clamp(PlayTime, MinPreviewPlayLength, MaxPreviewPlayLength);
 
-				for (TArray<FDatabasePreviewActor>& PreviewActorGroup : PreviewActors)
+			for (TArray<FDatabasePreviewActor>& PreviewActorGroup : PreviewActors)
+			{
+				for (FDatabasePreviewActor& PreviewActor : PreviewActorGroup)
 				{
-					for (FDatabasePreviewActor& PreviewActor : PreviewActorGroup)
-					{
-						PreviewActor.UpdatePreviewActor(Database, PlayTime, bQuantizeAnimationToPoseData);
-					}
+					PreviewActor.UpdatePreviewActor(Database, PlayTime, bQuantizeAnimationToPoseData);
 				}
 			}
 
@@ -993,23 +1034,18 @@ float FDatabaseViewModel::GetPlayTime() const
 
 void FDatabaseViewModel::SetPlayTime(float NewPlayTime, bool bInTickPlayTime)
 {
-	NewPlayTime = FMath::Clamp(NewPlayTime, MinPreviewPlayLength, MaxPreviewPlayLength);
+	PlayTime = FMath::Clamp(NewPlayTime, MinPreviewPlayLength, MaxPreviewPlayLength);
 	DeltaTimeMultiplier = bInTickPlayTime ? DeltaTimeMultiplier : 0.f;
 
-	if (!FMath::IsNearlyEqual(PlayTime, NewPlayTime))
+	if (const UPoseSearchDatabase* Database = GetPoseSearchDatabase())
 	{
-		PlayTime = NewPlayTime;
-		
-		if (const UPoseSearchDatabase* Database = GetPoseSearchDatabase())
+		if (EAsyncBuildIndexResult::Success == FAsyncPoseSearchDatabasesManagement::RequestAsyncBuildIndex(Database, ERequestAsyncBuildFlag::ContinueRequest))
 		{
-			if (EAsyncBuildIndexResult::Success == FAsyncPoseSearchDatabasesManagement::RequestAsyncBuildIndex(Database, ERequestAsyncBuildFlag::ContinueRequest))
+			for (TArray<FDatabasePreviewActor>& PreviewActorGroup : PreviewActors)
 			{
-				for (TArray<FDatabasePreviewActor>& PreviewActorGroup : PreviewActors)
+				for (FDatabasePreviewActor& PreviewActor : PreviewActorGroup)
 				{
-					for (FDatabasePreviewActor& PreviewActor : PreviewActorGroup)
-					{
-						PreviewActor.UpdatePreviewActor(Database, PlayTime, bQuantizeAnimationToPoseData);
-					}
+					PreviewActor.UpdatePreviewActor(Database, PlayTime, bQuantizeAnimationToPoseData);
 				}
 			}
 		}
diff --git a/Engine/Plugins/Animation/PoseSearch/Source/Editor/Private/PoseSearchInteractionAssetEditor.cpp b/Engine/Plugins/Animation/PoseSearch/Source/Editor/Private/PoseSearchInteractionAssetEditor.cpp
--- a/Engine/Plugins/Animation/PoseSearch/Source/Editor/Private/PoseSearchInteractionAssetEditor.cpp
+++ b/Engine/Plugins/Animation/PoseSearch/Source/Editor/Private/PoseSearchInteractionAssetEditor.cpp
@@ -63,7 +63,12 @@ bool FInteractionAssetPreviewActor::SpawnPreviewActor(UWorld* World, const UPose
 		AnimInstance->InitializeAnimation();
 
 		USkeleton* Skeleton = PreviewAsset->GetSkeleton();
-		Mesh->SetSkeletalMesh(Skeleton->GetPreviewMesh(true));
+		USkeletalMesh* PreviewMesh = InteractionAsset->GetPreviewMesh(Role);
+		if (!PreviewMesh)
+		{
+			PreviewMesh = Skeleton->GetPreviewMesh(true);
+		}
+		Mesh->SetSkeletalMesh(PreviewMesh);
 		Mesh->EnablePreview(true, PreviewAsset);
 
 		AnimInstance->SetAnimationAsset(PreviewAsset, InteractionAsset->IsLooping(), 0.f);
@@ -119,6 +124,18 @@ void FInteractionAssetPreviewActor::UpdatePreviewActor(const UPoseSearchInteract
 			{
 				if (UDebugSkelMeshComponent* Mesh = Cast<UDebugSkelMeshComponent>(ActorPtr->GetRootComponent()))
 				{
+					USkeleton* Skeleton = PreviewAsset->GetSkeleton();
+					USkeletalMesh* PreviewMesh = InteractionAsset->GetPreviewMesh(GetRole());
+					if (!PreviewMesh)
+					{
+						PreviewMesh = Skeleton->GetPreviewMesh(true);
+					}
+
+					if (Mesh->GetSkeletalMeshAsset() != PreviewMesh)
+					{
+						Mesh->SetSkeletalMesh(PreviewMesh);
+					}					
+
 					if (UAnimPreviewInstance* AnimInstance = Mesh->PreviewInstance.Get())
 					{
 						bool bPreviewAssetChanged = false;
diff --git a/Engine/Plugins/Animation/PoseSearch/Source/Editor/Public/PoseSearchDatabaseViewModel.h b/Engine/Plugins/Animation/PoseSearch/Source/Editor/Public/PoseSearchDatabaseViewModel.h
--- a/Engine/Plugins/Animation/PoseSearch/Source/Editor/Public/PoseSearchDatabaseViewModel.h
+++ b/Engine/Plugins/Animation/PoseSearch/Source/Editor/Public/PoseSearchDatabaseViewModel.h
@@ -38,6 +38,7 @@ namespace UE::PoseSearch
 
 		void Destroy();
 
+		UDebugSkelMeshComponent* GetDebugSkelMeshComponent();
 		const UDebugSkelMeshComponent* GetDebugSkelMeshComponent() const;
 		const UAnimPreviewInstance* GetAnimPreviewInstance() const;
 		const FAnimationAssetSampler& GetSampler() const { return Sampler; }
diff --git a/Engine/Plugins/Animation/PoseSearch/Source/Runtime/Private/PoseSearchDatabase.cpp b/Engine/Plugins/Animation/PoseSearch/Source/Runtime/Private/PoseSearchDatabase.cpp
--- a/Engine/Plugins/Animation/PoseSearch/Source/Runtime/Private/PoseSearchDatabase.cpp
+++ b/Engine/Plugins/Animation/PoseSearch/Source/Runtime/Private/PoseSearchDatabase.cpp
@@ -279,6 +279,12 @@ bool FPoseSearchDatabaseAnimationAssetBase::IsSkeletonCompatible(TObjectPtr<cons
 	
 	return false;
 }
+
+USkeletalMesh* FPoseSearchDatabaseAnimationAssetBase::GetPreviewMeshForRole(const UE::PoseSearch::FRole& Role) const
+{
+	return nullptr;
+}
+
 #endif // WITH_EDITOR
 
 #if WITH_EDITORONLY_DATA
@@ -527,6 +533,13 @@ bool FPoseSearchDatabaseAnimMontage::IsRootMotionEnabled() const
 
 //////////////////////////////////////////////////////////////////////////
 // FPoseSearchDatabaseMultiAnimAsset
+#if WITH_EDITOR
+USkeletalMesh* FPoseSearchDatabaseMultiAnimAsset::GetPreviewMeshForRole(const UE::PoseSearch::FRole& Role) const
+{
+	return MultiAnimAsset ? MultiAnimAsset->GetPreviewMesh(Role) : nullptr;
+}
+#endif // WITH_EDITOR
+
 UObject* FPoseSearchDatabaseMultiAnimAsset::GetAnimationAsset() const
 {
 	return MultiAnimAsset.Get();
diff --git a/Engine/Plugins/Animation/PoseSearch/Source/Runtime/Private/PoseSearchInteractionAsset.cpp b/Engine/Plugins/Animation/PoseSearch/Source/Runtime/Private/PoseSearchInteractionAsset.cpp
--- a/Engine/Plugins/Animation/PoseSearch/Source/Runtime/Private/PoseSearchInteractionAsset.cpp
+++ b/Engine/Plugins/Animation/PoseSearch/Source/Runtime/Private/PoseSearchInteractionAsset.cpp
@@ -256,6 +256,20 @@ FTransform UPoseSearchInteractionAsset::GetDebugWarpOrigin(const UE::PoseSearch:
 	}
 	return FTransform::Identity;
 }
+
+USkeletalMesh* UPoseSearchInteractionAsset::GetPreviewMesh(const UE::PoseSearch::FRole& Role) const
+{
+	for (int32 ItemIndex = 0; ItemIndex < Items.Num(); ++ItemIndex)
+	{
+		const FPoseSearchInteractionAssetItem& Item = Items[ItemIndex];
+		if (Item.Role == Role)
+		{
+			return Item.PreviewMesh.Get();
+		}
+	}
+	return nullptr;
+}
+
 #endif // WITH_EDITOR
 
 void UPoseSearchInteractionAsset::CalculateWarpTransforms(float Time, const TArrayView<const FTransform> ActorRootBoneTransforms, TArrayView<FTransform> FullAlignedActorRootBoneTransforms) const
diff --git a/Engine/Plugins/Animation/PoseSearch/Source/Runtime/Public/PoseSearch/MultiAnimAsset.h b/Engine/Plugins/Animation/PoseSearch/Source/Runtime/Public/PoseSearch/MultiAnimAsset.h
--- a/Engine/Plugins/Animation/PoseSearch/Source/Runtime/Public/PoseSearch/MultiAnimAsset.h
+++ b/Engine/Plugins/Animation/PoseSearch/Source/Runtime/Public/PoseSearch/MultiAnimAsset.h
@@ -7,6 +7,7 @@
 #include "MultiAnimAsset.generated.h"
 
 class UAnimationAsset;
+class USkeletalMesh;
 
 // UObject defining tuples of UAnimationAsset(s) with associated Role(s) and relative transforms from a shared reference system via GetOrigin
 UCLASS(Abstract, Experimental, BlueprintType, Category = "Animation")
@@ -21,6 +22,7 @@ public:
 
 #if WITH_EDITOR
 	[[nodiscard]] virtual int32 GetFrameAtTime(float Time) const PURE_VIRTUAL(UMultiAnimAsset::GetFrameAtTime, return 0;);
+	[[nodiscard]] virtual USkeletalMesh* GetPreviewMesh(const FName& Role) const PURE_VIRTUAL(UMultiAnimAsset::GetPreviewMesh, return nullptr;);
 #endif // WITH_EDITOR
 
 	[[nodiscard]] virtual int32 GetNumRoles() const PURE_VIRTUAL(UMultiAnimAsset::GetNumRoles, return 0;);
diff --git a/Engine/Plugins/Animation/PoseSearch/Source/Runtime/Public/PoseSearch/PoseSearchDatabase.h b/Engine/Plugins/Animation/PoseSearch/Source/Runtime/Public/PoseSearch/PoseSearchDatabase.h
--- a/Engine/Plugins/Animation/PoseSearch/Source/Runtime/Public/PoseSearch/PoseSearchDatabase.h
+++ b/Engine/Plugins/Animation/PoseSearch/Source/Runtime/Public/PoseSearch/PoseSearchDatabase.h
@@ -63,6 +63,8 @@ struct POSESEARCH_API FPoseSearchDatabaseAnimationAssetBase
 #if WITH_EDITOR
 	virtual int32 GetFrameAtTime(float Time) const;
 	virtual bool IsSkeletonCompatible(TObjectPtr<const UPoseSearchSchema> InSchema) const;
+	// Experimental, this feature might be removed without warning, not for production use
+	virtual USkeletalMesh* GetPreviewMeshForRole(const UE::PoseSearch::FRole& Role) const;
 #endif // WITH_EDITOR
 
 #if WITH_EDITORONLY_DATA
@@ -337,7 +339,7 @@ template<> struct TStructOpsTypeTraits<FPoseSearchDatabaseAnimMontage> : public
 	enum { WithIdenticalViaEquality = true };
 };
 
-USTRUCT(BlueprintType, Category = "Animation|Pose Search")
+USTRUCT(Experimental, BlueprintType, Category = "Animation|Pose Search")
 struct POSESEARCH_API FPoseSearchDatabaseMultiAnimAsset : public FPoseSearchDatabaseAnimationAssetBase
 {
 	GENERATED_BODY()
@@ -359,6 +361,10 @@ struct POSESEARCH_API FPoseSearchDatabaseMultiAnimAsset : public FPoseSearchData
 	virtual FFloatInterval GetSamplingRange() const override { return SamplingRange; }
 #endif // WITH_EDITORONLY_DATA
 
+#if WITH_EDITOR
+	USkeletalMesh* GetPreviewMeshForRole(const UE::PoseSearch::FRole& Role) const;
+#endif // WITH_EDITOR
+
 	virtual UObject* GetAnimationAsset() const override;
 	virtual float GetPlayLength() const override;
 
@@ -443,7 +449,7 @@ public:
 
 	// If null, the default preview mesh for the skeleton will be used. Otherwise, this will be used in preview scenes.
 	// @todo: Move this to be a setting in the Pose Search Database editor. 
-	UPROPERTY(EditAnywhere, Category = "Preview")
+	UPROPERTY(EditAnywhere, Category = "Preview", meta = (ExcludeFromHash))
 	TObjectPtr<USkeletalMesh> PreviewMesh = nullptr;
 #endif // WITH_EDITORONLY_DATA
 
diff --git a/Engine/Plugins/Animation/PoseSearch/Source/Runtime/Public/PoseSearch/PoseSearchInteractionAsset.h b/Engine/Plugins/Animation/PoseSearch/Source/Runtime/Public/PoseSearch/PoseSearchInteractionAsset.h
--- a/Engine/Plugins/Animation/PoseSearch/Source/Runtime/Public/PoseSearch/PoseSearchInteractionAsset.h
+++ b/Engine/Plugins/Animation/PoseSearch/Source/Runtime/Public/PoseSearch/PoseSearchInteractionAsset.h
@@ -35,6 +35,12 @@ struct POSESEARCH_API FPoseSearchInteractionAssetItem
 	// offset from the origin
 	UPROPERTY(EditAnywhere, Category = "Settings")
 	FTransform Origin = FTransform::Identity;
+
+#if WITH_EDITORONLY_DATA
+	// If null, the default preview mesh for the skeleton will be used. Otherwise, this will be used in preview scenes.
+	UPROPERTY(EditAnywhere, Category = "Preview", meta = (ExcludeFromHash))
+	TObjectPtr<USkeletalMesh> PreviewMesh = nullptr;
+#endif // WITH_EDITORONLY_DATA
 };
 
 UCLASS(Experimental, BlueprintType, Category = "Animation|Pose Search")
@@ -77,12 +83,13 @@ public:
 	virtual int32 GetNumRoles() const override { return Items.Num(); }
 	virtual UE::PoseSearch::FRole GetRole(int32 RoleIndex) const override { return Items[RoleIndex].Role; }
 	
-	virtual UAnimationAsset* GetAnimationAsset(const FName /*UE::PoseSearch::FRole*/& Role) const override;
+	virtual UAnimationAsset* GetAnimationAsset(const UE::PoseSearch::FRole& Role) const override;
 
 	virtual FTransform GetOrigin(const UE::PoseSearch::FRole& Role) const override;
 
 #if WITH_EDITOR
 	FTransform GetDebugWarpOrigin(const UE::PoseSearch::FRole& Role, bool bComposeWithDebugWarpOffset) const;
+	virtual USkeletalMesh* GetPreviewMesh(const UE::PoseSearch::FRole& Role) const override;
 #endif // WITH_EDITOR
 
 	virtual void CalculateWarpTransforms(float Time, const TArrayView<const FTransform> ActorRootBoneTransforms, TArrayView<FTransform> FullAlignedActorRootBoneTransforms) const override;
