posesearch - experimental UPoseSearchFeatureChannel_Position::bNormalizeDisplacement

diff --git a/Engine/Plugins/Animation/PoseSearch/Source/Runtime/Private/PoseSearchFeatureChannel_Position.cpp b/Engine/Plugins/Animation/PoseSearch/Source/Runtime/Private/PoseSearchFeatureChannel_Position.cpp
--- a/Engine/Plugins/Animation/PoseSearch/Source/Runtime/Private/PoseSearchFeatureChannel_Position.cpp
+++ b/Engine/Plugins/Animation/PoseSearch/Source/Runtime/Private/PoseSearchFeatureChannel_Position.cpp
@@ -31,7 +31,8 @@ void UPoseSearchFeatureChannel_Position::FindOrAddToSchema(UPoseSearchSchema* Sc
 					Position->PermutationTimeType == PermutationTimeType &&
 					Position->SampleRole == Role &&
 					Position->OriginRole == Role &&
-					Position->bDefaultWithRootBone)
+					Position->bDefaultWithRootBone && 
+					!Position->bNormalizeDisplacement)
 				{
 					return Position;
 				}
@@ -137,6 +138,7 @@ void UPoseSearchFeatureChannel_Position::BuildQuery(UE::PoseSearch::FSearchConte
 		UniqueIdentifier = HashCombineFast(UniqueIdentifier, GetTypeHash(InputQueryPose));
 		UniqueIdentifier = HashCombineFast(UniqueIdentifier, GetTypeHash(ComponentStripping));
 		UniqueIdentifier = HashCombineFast(UniqueIdentifier, GetTypeHash(PermutationTimeType));
+		UniqueIdentifier = HashCombineFast(UniqueIdentifier, GetTypeHash(bNormalizeDisplacement));
 
 		TConstArrayView<float> CachedChannelData;
 		if (const UPoseSearchFeatureChannel* CachedChannel = SearchContext.GetCachedChannelData(UniqueIdentifier, this, CachedChannelData))
@@ -158,6 +160,7 @@ void UPoseSearchFeatureChannel_Position::BuildQuery(UE::PoseSearch::FSearchConte
 			check(CachedPositionChannel->InputQueryPose == InputQueryPose);
 			check(CachedPositionChannel->ComponentStripping == ComponentStripping);
 			check(CachedPositionChannel->PermutationTimeType == PermutationTimeType);
+			check(CachedPositionChannel->bNormalizeDisplacement == bNormalizeDisplacement);
 #endif //DO_CHECK
 
 			// copying the CachedChannelData into this channel portion of the FeatureVectorBuilder
@@ -183,7 +186,11 @@ void UPoseSearchFeatureChannel_Position::BuildQuery(UE::PoseSearch::FSearchConte
 	}
 	
 	// calculating the BonePosition in root bone space for the bone indexed by SchemaBoneIdx
-	const FVector BonePosition = SearchContext.GetSamplePosition(SampleTimeOffset, OriginTimeOffset, SchemaBoneIdx, SchemaOriginBoneIdx, SampleRole, OriginRole, PermutationTimeType);
+	FVector BonePosition = SearchContext.GetSamplePosition(SampleTimeOffset, OriginTimeOffset, SchemaBoneIdx, SchemaOriginBoneIdx, SampleRole, OriginRole, PermutationTimeType);
+	if (bNormalizeDisplacement)
+	{
+		BonePosition = BonePosition.GetClampedToMaxSize(1.f);
+	}
 	FFeatureVectorHelper::EncodeVector(SearchContext.EditFeatureVector(), ChannelDataOffset, BonePosition, ComponentStripping, false);
 }
 
@@ -218,7 +225,6 @@ void UPoseSearchFeatureChannel_Position::DebugDraw(const UE::PoseSearch::FDebugD
 		const FVector OriginBonePos = DrawParams.ExtractPosition(PoseVector, OriginTimeOffset, SchemaOriginBoneIdx, OriginRole, OriginPermutationTimeType, INDEX_NONE, PermutationOriginTimeOffset);
 		const FVector DeltaPos = DrawParams.ExtractRotation(PoseVector, OriginTimeOffset, RootSchemaBoneIdx, OriginRole, OriginPermutationTimeType, INDEX_NONE, PermutationOriginTimeOffset).RotateVector(FeaturesVector);
 		const FVector BonePos = OriginBonePos + DeltaPos;
-		DrawParams.DrawPoint(BonePos, Color);
 
 		if (MaxPositionDistanceSquared > 0.f)
 		{
@@ -235,12 +241,22 @@ void UPoseSearchFeatureChannel_Position::DebugDraw(const UE::PoseSearch::FDebugD
 			}
 		}
 
-		const bool bDrawOrigin = !DeltaPos.IsNearlyZero() && (SchemaOriginBoneIdx != RootSchemaBoneIdx || !FMath::IsNearlyZero(OriginTimeOffset) ||
-			SampleRole != OriginRole || PermutationTimeType != EPermutationTimeType::UseSampleTime || bUseBlueprintQueryOverride);
-		if (bDrawOrigin)
+		if (bNormalizeDisplacement)
+		{
+			static const float NormalizeDisplacementLength = 100.f;
+			DrawParams.DrawLine(OriginBonePos, OriginBonePos + DeltaPos * NormalizeDisplacementLength, Color);
+		}
+		else
 		{
-			DrawParams.DrawPoint(OriginBonePos, Color);
-			DrawParams.DrawLine(OriginBonePos, BonePos, Color);
+			DrawParams.DrawPoint(BonePos, Color);
+
+			const bool bDrawOrigin = !DeltaPos.IsNearlyZero() && (SchemaOriginBoneIdx != RootSchemaBoneIdx || !FMath::IsNearlyZero(OriginTimeOffset) ||
+				SampleRole != OriginRole || PermutationTimeType != EPermutationTimeType::UseSampleTime || bUseBlueprintQueryOverride);
+			if (bDrawOrigin)
+			{
+				DrawParams.DrawPoint(OriginBonePos, Color);
+				DrawParams.DrawLine(OriginBonePos, BonePos, Color);
+			}
 		}
 	}
 }
@@ -264,6 +280,10 @@ bool UPoseSearchFeatureChannel_Position::IndexAsset(UE::PoseSearch::FAssetIndexe
 	{
 		if (Indexer.GetSamplePosition(BonePosition, SampleTimeOffset, OriginTimeOffset, SampleIdx, SchemaBoneIdx, SchemaOriginBoneIdx, SampleRole, OriginRole, PermutationTimeType, SamplingAttributeId))
 		{
+			if (bNormalizeDisplacement)
+			{
+				BonePosition = BonePosition.GetClampedToMaxSize(1.f);
+			}
 			FFeatureVectorHelper::EncodeVector(Indexer.GetPoseVector(SampleIdx), ChannelDataOffset, BonePosition, ComponentStripping, false);
 		}
 		else
@@ -282,6 +302,10 @@ UE::PoseSearch::TLabelBuilder& UPoseSearchFeatureChannel_Position::GetLabel(UE::
 	AppendLabelSeparator(LabelBuilder, LabelFormat);
 
 	LabelBuilder.Append(TEXT("Pos"));
+	if (bNormalizeDisplacement)
+	{
+		LabelBuilder.Append(TEXT("_ND"));
+	}
 
 	if (ComponentStripping == EComponentStrippingVector::StripXY)
 	{
diff --git a/Engine/Plugins/Animation/PoseSearch/Source/Runtime/Public/PoseSearch/PoseSearchFeatureChannel_Position.h b/Engine/Plugins/Animation/PoseSearch/Source/Runtime/Public/PoseSearch/PoseSearchFeatureChannel_Position.h
--- a/Engine/Plugins/Animation/PoseSearch/Source/Runtime/Public/PoseSearch/PoseSearchFeatureChannel_Position.h
+++ b/Engine/Plugins/Animation/PoseSearch/Source/Runtime/Public/PoseSearch/PoseSearchFeatureChannel_Position.h
@@ -80,6 +80,14 @@ public:
 	UPROPERTY(EditAnywhere, Category = "Experimental|Filter")
 	float MaxPositionDistanceSquared = 0.f;
 
+	// Experimental, this feature might be removed without warning, not for production use
+	// UPoseSearchFeatureChannel_Position represents / calculates the position (displacement) of Bone relative to OriginBone
+	// if bNormalizeDisplacement is true, the displacement will be normalized. if zero it'll stay zero, so it'll be like matching the 
+	// direction towards the position is pointing at. useful to match "trajectory points" when the velocity is not consistent between 
+	// query and database poses. Look at the debug lines to have an idea what it'd match
+	UPROPERTY(EditAnywhere, Category = "Experimental|Settings")
+	bool bNormalizeDisplacement = false;
+
 #if WITH_EDITORONLY_DATA
 	// if set, all the channels of the same class with the same cardinality, and the same NormalizationGroup, will be normalized together.
 	// for example in a locomotion database of a character holding a weapon, containing non mirrorable animations, you'd still want to normalize together 
