pose search
- support for additional roles in UPoseSearchInteractionAsset not covered by the databse schema
- support for databases containing UPoseSearchInteractionAsset with a single role schema in FAnimNode_MotionMatching (via UPoseSearchLibrary::UpdateMotionMatchingState)
- support for "interactions" with a single character, without injected tick dependencies
- not performing the MM search in FAnimNode_MotionMatching::UpdateAssetPlayer if the node is not relevant enough

[RN] skip

diff --git a/Engine/Plugins/Animation/PoseSearch/Source/Editor/Private/PoseSearchDatabaseViewModel.cpp b/Engine/Plugins/Animation/PoseSearch/Source/Editor/Private/PoseSearchDatabaseViewModel.cpp
--- a/Engine/Plugins/Animation/PoseSearch/Source/Editor/Private/PoseSearchDatabaseViewModel.cpp
+++ b/Engine/Plugins/Animation/PoseSearch/Source/Editor/Private/PoseSearchDatabaseViewModel.cpp
@@ -47,17 +47,18 @@ bool FDatabasePreviewActor::SpawnPreviewActor(UWorld* World, const UPoseSearchDa
 {
 	check(PoseSearchDatabase && PoseSearchDatabase->Schema);
 	const FSearchIndex& SearchIndex = PoseSearchDatabase->GetSearchIndex();
+	const FSearchIndexAsset& IndexAsset = SearchIndex.Assets[IndexAssetIdx];
+	const FPoseSearchDatabaseAnimationAssetBase* DatabaseAnimationAsset = PoseSearchDatabase->GetDatabaseAnimationAsset<FPoseSearchDatabaseAnimationAssetBase>(IndexAsset.GetSourceAssetIdx());
+	check(DatabaseAnimationAsset);
 
+	
 	USkeleton* Skeleton = PoseSearchDatabase->Schema->GetSkeleton(Role);
 	if (!Skeleton)
 	{
-		UE_LOG(LogPoseSearchEditor, Error, TEXT("Couldn't spawn preview Actor because of missing Role '%s' in Schema '%s'"), *Role.ToString(), *PoseSearchDatabase->Schema->GetName());
+		UE_LOG(LogPoseSearchEditor, Log, TEXT("Couldn't spawn preview Actor for asset %s because its Role '%s' is missing in Schema '%s'"), *GetNameSafe(DatabaseAnimationAsset->GetAnimationAsset()), *Role.ToString(), *PoseSearchDatabase->Schema->GetName());
 		return false;
 	}
 
-	const FSearchIndexAsset& IndexAsset = SearchIndex.Assets[IndexAssetIdx];
-
-	const FPoseSearchDatabaseAnimationAssetBase* DatabaseAnimationAsset = PoseSearchDatabase->GetDatabaseAnimationAsset<FPoseSearchDatabaseAnimationAssetBase>(IndexAsset.GetSourceAssetIdx());
 	UAnimationAsset* PreviewAsset = Cast<UAnimationAsset>(DatabaseAnimationAsset->GetAnimationAssetForRole(Role));
 	if (!PreviewAsset)
 	{
diff --git a/Engine/Plugins/Animation/PoseSearch/Source/Runtime/Private/AnimNode_MotionMatching.cpp b/Engine/Plugins/Animation/PoseSearch/Source/Runtime/Private/AnimNode_MotionMatching.cpp
--- a/Engine/Plugins/Animation/PoseSearch/Source/Runtime/Private/AnimNode_MotionMatching.cpp
+++ b/Engine/Plugins/Animation/PoseSearch/Source/Runtime/Private/AnimNode_MotionMatching.cpp
@@ -178,13 +178,14 @@ void FAnimNode_MotionMatching::UpdateAssetPlayer(const FAnimationUpdateContext&
 	}
 #endif // ENABLE_ANIM_DEBUG
 
+	const bool bIsRelevantEnough = Context.GetFinalBlendWeight() > FAnimWeight::GetSmallestRelevantWeight();
 	if (DatabasesToSearch.IsEmpty())
 	{
 		Context.LogMessage(
 			EMessageSeverity::Error,
 			LOCTEXT("NoDatabases", "No database assets provided for motion matching."));
 	}
-	else
+	else if (bIsRelevantEnough)
 	{
 		const float DeltaTime = Context.GetDeltaTime();
 
@@ -213,7 +214,7 @@ void FAnimNode_MotionMatching::UpdateAssetPlayer(const FAnimationUpdateContext&
 		);
 	}
 
-	const bool bJumpToPose = MotionMatchingState.CurrentSearchResult.IsValid() && !MotionMatchingState.CurrentSearchResult.bIsContinuingPoseSearch;
+	const bool bJumpToPose = bIsRelevantEnough && MotionMatchingState.CurrentSearchResult.IsValid() && !MotionMatchingState.CurrentSearchResult.bIsContinuingPoseSearch;
 
 	PRAGMA_DISABLE_DEPRECATION_WARNINGS
 	MotionMatchingState.bJumpedToPose = bJumpToPose;
@@ -230,8 +231,8 @@ void FAnimNode_MotionMatching::UpdateAssetPlayer(const FAnimationUpdateContext&
 		{
 			const FPoseSearchDatabaseAnimationAssetBase* DatabaseAsset = CurrentResultDatabase->GetDatabaseAnimationAsset<FPoseSearchDatabaseAnimationAssetBase>(*SearchIndexAsset);
 			check(DatabaseAsset);
-
-			if (UAnimationAsset* AnimationAsset = Cast<UAnimationAsset>(DatabaseAsset->GetAnimationAsset()))
+			
+			if (UAnimationAsset* AnimationAsset = Cast<UAnimationAsset>(DatabaseAsset->GetAnimationAssetForRole(CurrentResultDatabase->Schema->GetDefaultRole())))
 			{
 				// Clear up any sync group info before pushing new asset player (which will have sync info since its the highest weighted).
 				for (FBlendStackAnimPlayer& AnimPlayer : AnimPlayers)
diff --git a/Engine/Plugins/Animation/PoseSearch/Source/Runtime/Private/AnimNode_MotionMatchingInteraction.cpp b/Engine/Plugins/Animation/PoseSearch/Source/Runtime/Private/AnimNode_MotionMatchingInteraction.cpp
--- a/Engine/Plugins/Animation/PoseSearch/Source/Runtime/Private/AnimNode_MotionMatchingInteraction.cpp
+++ b/Engine/Plugins/Animation/PoseSearch/Source/Runtime/Private/AnimNode_MotionMatchingInteraction.cpp
@@ -90,18 +90,15 @@ void FAnimNode_MotionMatchingInteraction::UpdateAssetPlayer(const FAnimationUpda
 			RoledAnimAsset = SingleAnimAsset;
 		}
 
-		if (RoledAnimAsset)
+		if (RoledAnimAsset && (AnimPlayers.IsEmpty() || !CurrentResult.bIsContinuingPoseSearch))
 		{
-			if (AnimPlayers.IsEmpty() || !CurrentResult.bIsContinuingPoseSearch)
-			{
-				const FPoseSearchRoledSkeleton* RoledSkeleton = CurrentResult.SelectedDatabase->Schema->GetRoledSkeleton(CurrentResult.Role);
-				check(RoledSkeleton);
+			const FPoseSearchRoledSkeleton* RoledSkeleton = CurrentResult.SelectedDatabase->Schema->GetRoledSkeleton(CurrentResult.Role);
+			check(RoledSkeleton);
 
-				BlendTo(Context, RoledAnimAsset, CurrentResult.SelectedTime, CurrentResult.bLoop, CurrentResult.bIsMirrored, RoledSkeleton->MirrorDataTable.Get(),
-					BlendTime, BlendProfile, BlendOption, bUseInertialBlend, CurrentResult.BlendParameters, CurrentResult.WantedPlayRate);
+			BlendTo(Context, RoledAnimAsset, CurrentResult.SelectedTime, CurrentResult.bLoop, CurrentResult.bIsMirrored, RoledSkeleton->MirrorDataTable.Get(),
+				BlendTime, BlendProfile, BlendOption, bUseInertialBlend, CurrentResult.BlendParameters, CurrentResult.WantedPlayRate);
 
-				bBlendToExecuted = true;
-			}
+			bBlendToExecuted = true;
 		}
 	}
 	else
@@ -123,8 +120,7 @@ void FAnimNode_MotionMatchingInteraction::UpdateAssetPlayer(const FAnimationUpda
 	}
 #endif // ENABLE_ANIM_DEBUG
 
-	// bypassing FAnimNode_BlendStack::UpdateAssetPlayer, since we overridden its behaviour
-	FAnimNode_BlendStack_Standalone::UpdateAssetPlayer(Context);
+	Super::UpdateAssetPlayer(Context);
 }
 
 void FAnimNode_MotionMatchingInteraction::Evaluate_AnyThread(FPoseContext& Output)
diff --git a/Engine/Plugins/Animation/PoseSearch/Source/Runtime/Private/PoseSearchContext.cpp b/Engine/Plugins/Animation/PoseSearch/Source/Runtime/Private/PoseSearchContext.cpp
--- a/Engine/Plugins/Animation/PoseSearch/Source/Runtime/Private/PoseSearchContext.cpp
+++ b/Engine/Plugins/Animation/PoseSearch/Source/Runtime/Private/PoseSearchContext.cpp
@@ -770,6 +770,39 @@ bool FSearchContext::ArePoseHistoriesValid() const
 	return true;
 }
 
+const IPoseHistory* FSearchContext::GetPoseHistory(const FRole& Role) const
+{
+	if (const int32* RoleIndex = RoleToIndex.Find(Role))
+	{
+		return PoseHistories[*RoleIndex];
+	}
+
+	UE_LOG(LogPoseSearch, Error, TEXT("FSearchContext::GetPoseHistory - Role %s could not be found!"), *Role.ToString());
+	return nullptr;
+}
+
+const UAnimInstance* FSearchContext::GetAnimInstance(const FRole& Role) const
+{
+	if (const int32* RoleIndex = RoleToIndex.Find(Role))
+	{
+		return Cast<UAnimInstance>(AnimContexts[*RoleIndex]);
+	}
+
+	UE_LOG(LogPoseSearch, Error, TEXT("FSearchContext::GetAnimInstance - Role %s could not be found!"), *Role.ToString());
+	return nullptr;
+}
+
+const UObject* FSearchContext::GetAnimContext(const FRole& Role) const
+{
+	if (const int32* RoleIndex = RoleToIndex.Find(Role))
+	{
+		return AnimContexts[*RoleIndex];
+	}
+		
+	UE_LOG(LogPoseSearch, Error, TEXT("FSearchContext::GetAnimContext - Role %s could not be found!"), *Role.ToString());
+	return nullptr;
+}
+
 FTransform FSearchContext::GetWorldBoneTransformAtTime(float SampleTime, const FRole& SampleRole, int8 SchemaBoneIdx)
 {
 	// CachedQueries.Last is the query we're building 
diff --git a/Engine/Plugins/Animation/PoseSearch/Source/Runtime/Private/PoseSearchDerivedData.cpp b/Engine/Plugins/Animation/PoseSearch/Source/Runtime/Private/PoseSearchDerivedData.cpp
--- a/Engine/Plugins/Animation/PoseSearch/Source/Runtime/Private/PoseSearchDerivedData.cpp
+++ b/Engine/Plugins/Animation/PoseSearch/Source/Runtime/Private/PoseSearchDerivedData.cpp
@@ -420,13 +420,6 @@ static bool InitSearchIndexAssets(FSearchIndexBase& SearchIndex, const UPoseSear
 				continue;
 			}
 
-			if (DatabaseAsset->GetNumRoles() != Schema->GetRoledSkeletons().Num())
-			{
-				UE_LOG(LogPoseSearch, Error, TEXT("DatabaseAsset '%s' Roles don't match Schema '%s' Skeletons Roles"), *DatabaseAsset->GetAnimationAsset()->GetName(), *Schema->GetName());
-				bAnyErrors = true;
-				continue;
-			}
-
 			// checking for valid roles in DatabaseMultiAnimAsset against the Schema
 			bool bAreAllRolesSupported = true;
 			for (const FPoseSearchRoledSkeleton& RoledSkeleton : Schema->GetRoledSkeletons())
@@ -1058,11 +1051,6 @@ static bool IndexDatabase(FSearchIndexBase& SearchIndexBase, const UPoseSearchDa
 				TempBoneContainers.Emplace(RoledBoneContainers[Role]);
 				TempRoleToIndex.Add(Role) = RoleIndex;
 			}
-			else
-			{
-				UE_LOG(LogPoseSearch, Warning, TEXT("Role '%s' from asset '%s' cannot be found in Skeletons of Schema '%s'"), *Role.ToString(), *DatabaseAnimationAssetBase->GetName(), *Schema->GetName());
-				return false;
-			}
 		}
 
 		const FFloatInterval ExtrapolationTimeInterval = SearchIndexAsset.GetExtrapolationTimeInterval(Schema->SampleRate, AdditionalExtrapolationTime);
diff --git a/Engine/Plugins/Animation/PoseSearch/Source/Runtime/Private/PoseSearchInteractionIsland.cpp b/Engine/Plugins/Animation/PoseSearch/Source/Runtime/Private/PoseSearchInteractionIsland.cpp
--- a/Engine/Plugins/Animation/PoseSearch/Source/Runtime/Private/PoseSearchInteractionIsland.cpp
+++ b/Engine/Plugins/Animation/PoseSearch/Source/Runtime/Private/PoseSearchInteractionIsland.cpp
@@ -18,9 +18,6 @@
 namespace UE::PoseSearch
 {
 
-static int32 GVarPoseSearchInteractionIslandResultPolicy = 1;
-static FAutoConsoleVariableRef CVarPoseSearchInteractionIslandResultPolicy(TEXT("a.PoseSearchInteraction.IslandResultPolicy"), GVarPoseSearchInteractionIslandResultPolicy, TEXT("0: only best result is valid\n1: assign from best to worst result"));
-
 // Utils functions
 ///////////////////////////////////////////////////////////
 static FInteractionSearchResult InitSearchResult(const FSearchResult& SearchResult, const FInteractionSearchContext& SearchContext, int32 SearchIndex)
@@ -148,31 +145,18 @@ static void InitSearchResults(TArray<FInteractionSearchResult>& SearchResults, T
 				return PoseSearchResults[IndexA].PoseCost < PoseSearchResults[IndexB].PoseCost;
 			});
 
-		if (GVarPoseSearchInteractionIslandResultPolicy == 0)
-		{
-			// only best, valid with the highest cardinality in roles, result is valid
-			const int32 BestSearchIndex = SortedPoseSearchResults[0];
-			if (PoseSearchResults[BestSearchIndex].IsValid())
-			{
-				SearchResults.SetNum(1, EAllowShrinking::No);
-				SearchResults[0] = InitSearchResult(PoseSearchResults[BestSearchIndex], SearchContexts[BestSearchIndex], BestSearchIndex);
-			}
-		}
-		else
+		// assign from best to worst result
+		FVisitedAnimContexts VisitedAnimContexts;
+		for (int32 SearchIndex : SortedPoseSearchResults)
 		{
-			// assign from best to worst result
-			FVisitedAnimContexts VisitedAnimContexts;
-			for (int32 SearchIndex : SortedPoseSearchResults)
+			if (IsPoseSearchResultUsable(SearchIndex, PoseSearchResults, SearchContexts, VisitedAnimContexts))
 			{
-				if (IsPoseSearchResultUsable(SearchIndex, PoseSearchResults, SearchContexts, VisitedAnimContexts))
+				SearchResults.Add(InitSearchResult(PoseSearchResults[SearchIndex], SearchContexts[SearchIndex], SearchIndex));
+
+				for (const TWeakObjectPtr<const UObject>& SearchAnimContext : SearchContexts[SearchIndex].AnimContexts)
 				{
-					SearchResults.Add(InitSearchResult(PoseSearchResults[SearchIndex], SearchContexts[SearchIndex], SearchIndex));
-					
-					for (const TWeakObjectPtr<const UObject>& SearchAnimContext : SearchContexts[SearchIndex].AnimContexts)
-					{
-						check(SearchAnimContext.Get());
-						VisitedAnimContexts.Add(SearchAnimContext.Get());
-					}
+					check(SearchAnimContext.Get());
+					VisitedAnimContexts.Add(SearchAnimContext.Get());
 				}
 			}
 		}
@@ -322,7 +306,7 @@ void FInteractionIsland::FPreTickFunction::ExecuteTick(float DeltaTime, enum ELe
 	// generating trajectories before running any of the skeletal mesh component ticks
 	check(Island);
 
-	if (!Island->IsUninjected())
+	if (Island->HasTickDependencies())
 	{
 		CheckInteractionThreadSafety(Island);
 
@@ -375,7 +359,7 @@ void FInteractionIsland::FPostTickFunction::ExecuteTick(float DeltaTime, enum EL
 	check(Island);
 
 #if ENABLE_ANIM_DEBUG
-	if (!Island->IsUninjected())
+	if (Island->HasTickDependencies())
 	{
 		CheckInteractionThreadSafety(Island);
 
@@ -418,7 +402,7 @@ FInteractionIsland::FInteractionIsland(ULevel* Level, UPoseSearchInteractionSubs
 
 FInteractionIsland::~FInteractionIsland()
 {
-	Uninject(false);
+	Uninitialize(false);
 
 	PreTickFunction.UnRegisterTickFunction();
 	PostTickFunction.UnRegisterTickFunction();
@@ -461,7 +445,7 @@ UActorComponent* FInteractionIsland::FindPostTickComponent(const AActor* Actor)
 	return nullptr;
 }
 
-void FInteractionIsland::InjectToActor(const UObject* AnimContext)
+void FInteractionIsland::InjectToActor(const UObject* AnimContext, bool bAddTickDependencies)
 {
 	check(IsInGameThread());
 
@@ -492,7 +476,7 @@ void FInteractionIsland::InjectToActor(const UObject* AnimContext)
 			{
 				if (UActorComponent* PostTickComponent = FindPostTickComponent(Actor))
 				{
-					const bool bIsFirstInjectedActor = IsUninjected();
+					const bool bIsMainActor = PostTickComponents.IsEmpty();
 
 					//	tick order: 
 					//		PreTickComponents (UCharacterMovementComponent, or UCharacterMoverComponent) ->
@@ -507,18 +491,29 @@ void FInteractionIsland::InjectToActor(const UObject* AnimContext)
 					// (so we can remove them later on in a consistent fashion)
 					check(PreTickComponents.Num() == PostTickComponents.Num());
 
-					PreTickFunction.AddPrerequisite(PreTickComponent, PreTickComponent->PrimaryComponentTick);
-
-					if (bIsFirstInjectedActor)
+					if (bAddTickDependencies)
 					{
-						// it should be this island, but it's not a UObject, so we use the InteractionSubsystem
-						PostTickComponent->PrimaryComponentTick.AddPrerequisite(InteractionSubsystem, PreTickFunction);
-						PostTickFunction.AddPrerequisite(PostTickComponent, PostTickComponent->PrimaryComponentTick);
+						PreTickFunction.AddPrerequisite(PreTickComponent, PreTickComponent->PrimaryComponentTick);
+
+						if (bIsMainActor)
+						{
+							// it should be this island, but it's not a UObject, so we use the InteractionSubsystem
+							PostTickComponent->PrimaryComponentTick.AddPrerequisite(InteractionSubsystem, PreTickFunction);
+							PostTickFunction.AddPrerequisite(PostTickComponent, PostTickComponent->PrimaryComponentTick);
+
+							check(!bHasTickDependencies);
+							bHasTickDependencies = bAddTickDependencies;
+						}
+						else
+						{
+							// it should be this island, but it's not a UObject, so we use the InteractionSubsystem
+							PostTickComponent->PrimaryComponentTick.AddPrerequisite(InteractionSubsystem, PostTickFunction);
+							check(bHasTickDependencies);
+						}
 					}
 					else
 					{
-						// it should be this island, but it's not a UObject, so we use the InteractionSubsystem
-						PostTickComponent->PrimaryComponentTick.AddPrerequisite(InteractionSubsystem, PostTickFunction);
+						check(!bHasTickDependencies);
 					}
 				}
 			}
@@ -537,27 +532,25 @@ void FInteractionIsland::AddSearchContext(const FInteractionSearchContext& Searc
 	SearchContexts.Add(SearchContext);
 }
 
-void FInteractionIsland::Uninject(bool bValidateTickDependencies)
+void FInteractionIsland::RemoveTickDependencies(bool bValidateTickDependencies)
 {
 	// Called by UPoseSearchInteractionSubsystem::Tick when there aren't animation jobs flying. No need to FScopeLock Lock(&Mutex);
 	check(IsInGameThread());
 
 	check(PreTickComponents.Num() == PostTickComponents.Num());
 
-	if (IsUninjected())
+	if (!bHasTickDependencies)
 	{
-		check(SearchContexts.IsEmpty() && SearchResults.IsEmpty() && !bSearchPerfomed);
-
 #if ENABLE_ANIM_DEBUG
 		if (bValidateTickDependencies && (bPreTickFunctionExecuted || bPostTickFunctionExecuted))
 		{
 			if (bPreTickFunctionExecuted)
 			{
-				UE_LOG(LogPoseSearch, Error, TEXT("FInteractionIsland::Uninject, unexpected FPreTickFunction::ExecuteTick run!"));
+				UE_LOG(LogPoseSearch, Error, TEXT("FInteractionIsland::RemoveTickDependencies, unexpected FPreTickFunction::ExecuteTick run!"));
 			}
 			else
 			{
-				UE_LOG(LogPoseSearch, Error, TEXT("FInteractionIsland::Uninject, unexpected FPostTickFunction::ExecuteTick run!"));
+				UE_LOG(LogPoseSearch, Error, TEXT("FInteractionIsland::RemoveTickDependencies, unexpected FPostTickFunction::ExecuteTick run!"));
 			}
 				
 			LogTickDependencies();
@@ -571,11 +564,11 @@ void FInteractionIsland::Uninject(bool bValidateTickDependencies)
 		{
 			if (!bPreTickFunctionExecuted)
 			{
-				UE_LOG(LogPoseSearch, Error, TEXT("FInteractionIsland::Uninject, expected FPreTickFunction::ExecuteTick didn't run!"));
+				UE_LOG(LogPoseSearch, Error, TEXT("FInteractionIsland::RemoveTickDependencies, expected FPreTickFunction::ExecuteTick didn't run!"));
 			}
 			else
 			{
-				UE_LOG(LogPoseSearch, Error, TEXT("FInteractionIsland::Uninject, expected FPostTickFunction::ExecuteTick didn't run!"));
+				UE_LOG(LogPoseSearch, Error, TEXT("FInteractionIsland::RemoveTickDependencies, expected FPostTickFunction::ExecuteTick didn't run!"));
 			}
 			
 			LogTickDependencies();
@@ -589,7 +582,6 @@ void FInteractionIsland::Uninject(bool bValidateTickDependencies)
 				PreTickFunction.RemovePrerequisite(PreTickComponent, PreTickComponent->PrimaryComponentTick);
 			}
 		}
-		check(PreTickFunction.GetPrerequisites().IsEmpty());
 
 		bool bMainActor = true;
 		for (TWeakObjectPtr<UActorComponent>& PostTickComponentPtr : PostTickComponents)
@@ -609,8 +601,24 @@ void FInteractionIsland::Uninject(bool bValidateTickDependencies)
 			bMainActor = false;
 		}
 
+		check(PreTickFunction.GetPrerequisites().IsEmpty());
 		check(PostTickFunction.GetPrerequisites().IsEmpty());
 
+		bHasTickDependencies = false;
+	}
+
+#if ENABLE_ANIM_DEBUG
+	bPreTickFunctionExecuted = false;
+	bPostTickFunctionExecuted = false;
+#endif // ENABLE_ANIM_DEBUG
+}
+
+void FInteractionIsland::Uninitialize(bool bValidateTickDependencies)
+{
+	RemoveTickDependencies(bValidateTickDependencies);
+
+	if (IsInitialized())
+	{
 		PreTickComponents.Reset();
 		PostTickComponents.Reset();
 
@@ -618,22 +626,27 @@ void FInteractionIsland::Uninject(bool bValidateTickDependencies)
 		SearchResults.Reset();
 		bSearchPerfomed = false;
 	}
+	else
+	{
+		check(PreTickComponents.IsEmpty() && PostTickComponents.IsEmpty() && SearchContexts.IsEmpty() && SearchResults.IsEmpty() && !bSearchPerfomed);
+	}
+}
 
-#if ENABLE_ANIM_DEBUG
-	bPreTickFunctionExecuted = false;
-	bPostTickFunctionExecuted = false;
-#endif // ENABLE_ANIM_DEBUG
+bool FInteractionIsland::HasTickDependencies() const
+{
+	return bHasTickDependencies;
 }
 
-bool FInteractionIsland::IsUninjected() const
+bool FInteractionIsland::IsInitialized() const
 {
-	return PostTickComponents.IsEmpty();
+	return !SearchContexts.IsEmpty();
 }
 
 const AActor* FInteractionIsland::GetMainActor() const
 {
-	if (!IsUninjected())
+	if (IsInitialized())
 	{
+		check(!PostTickComponents.IsEmpty());
 		if (const UActorComponent* ActorComponent = PostTickComponents[0].Get())
 		{
 			const AActor* MainActor = ActorComponent->GetOwner();
diff --git a/Engine/Plugins/Animation/PoseSearch/Source/Runtime/Private/PoseSearchInteractionSubsystem.cpp b/Engine/Plugins/Animation/PoseSearch/Source/Runtime/Private/PoseSearchInteractionSubsystem.cpp
--- a/Engine/Plugins/Animation/PoseSearch/Source/Runtime/Private/PoseSearchInteractionSubsystem.cpp
+++ b/Engine/Plugins/Animation/PoseSearch/Source/Runtime/Private/PoseSearchInteractionSubsystem.cpp
@@ -22,15 +22,20 @@ static bool GVarPoseSearchInteractionCacheIslands = true;
 static FAutoConsoleVariableRef CVarPoseSearchInteractionCacheIslands(TEXT("a.PoseSearchInteraction.CacheIslands"), GVarPoseSearchInteractionCacheIslands, TEXT("Cache Pose Search Interaction Islands for future reuse instead of destrying them"));
 #endif // !NO_CVARS
 
-static FVector GetContextLocation(const UObject* AnimContext)
+static const FTransform& GetContextTransform(const UObject* AnimContext)
 {
 	if (const UAnimInstance* AnimInstance = Cast<UAnimInstance>(AnimContext))
 	{
-		return AnimInstance->GetSkelMeshComponent()->GetComponentLocation();
+		return AnimInstance->GetSkelMeshComponent()->GetComponentTransform();
 	}
 
 	unimplemented();
-	return FVector::ZeroVector;
+	return FTransform::Identity;
+}
+
+static FVector GetContextLocation(const UObject* AnimContext)
+{
+	return GetContextTransform(AnimContext).GetLocation();
 }
 
 FString FInteractionAvailability::GetPoseHistoryName() const
@@ -58,52 +63,6 @@ const FAnimNode_PoseSearchHistoryCollector_Base* FInteractionAvailability::GetHi
 	return nullptr;
 }
 
-struct FInteractionSearchContextGroup
-{
-	bool Contains(const FInteractionSearchContext& SearchContext) const
-	{
-		for (const TWeakObjectPtr<const UObject>& AnimContext : SearchContext.AnimContexts)
-		{
-			if (AnimContexts.Find(AnimContext.Get()))
-			{
-				return true;
-			}
-		}
-		return false;
-	}
-
-	void Add(const FInteractionSearchContext& SearchContext, int32 SearchContextIndex)
-	{
-		for (const TWeakObjectPtr<const UObject>& AnimContext : SearchContext.AnimContexts)
-		{
-			if (AnimContext.Get())
-			{
-				AnimContexts.Add(AnimContext.Get());
-			}
-		}
-
-		SearchContextsIndices.Add(SearchContextIndex);
-	}
-
-	void Merge(const FInteractionSearchContextGroup& SearchContextGroup)
-	{
-		for (const UObject* AnimContext : SearchContextGroup.AnimContexts)
-		{
-			check(AnimContext);
-			AnimContexts.Add(AnimContext);
-		}
-
-		for (int32 SearchContextsIndex : SearchContextGroup.SearchContextsIndices)
-		{
-			SearchContextsIndices.Add(SearchContextsIndex);
-		}
-	}
-
-	TSet<const UObject*, DefaultKeyFuncs<const UObject*>, TInlineSetAllocator<16, TMemStackSetAllocator<>>> AnimContexts;
-	TArray<int32, TInlineAllocator<16, TMemStackAllocator<>>> SearchContextsIndices;
-};
-typedef TArray<FInteractionSearchContextGroup, TInlineAllocator<PreallocatedSearchesNum, TMemStackAllocator<>>> FInteractionSearchContextGroups;
-
 struct FAnimContextInfo;
 struct FNearbyAnimContextInfo
 {
@@ -178,46 +137,6 @@ static void GeneratePermutations(int32 DataCardinality, int32 CombinationCardina
 	GeneratePermutationsRecursive(DataCardinality, 0, Combination, 0, EvaluateCombination);
 }
 
-
-static void GenerateInteractionSearchContextGroups(const FInteractionSearchContexts& SearchContexts, FInteractionSearchContextGroups& SearchContextGroups)
-{
-	QUICK_SCOPE_CYCLE_COUNTER(STAT_UPoseSearchInteractionSubsystem_GenerateInteractionSearchContextGroups);
-
-	check(SearchContextGroups.IsEmpty());
-
-	for (int32 SearchContextIndex = 0; SearchContextIndex < SearchContexts.Num(); ++SearchContextIndex)
-	{
-		const FInteractionSearchContext& SearchContext = SearchContexts[SearchContextIndex];
-
-		int32 MainSearchContextGroupIndex = INDEX_NONE;
-		for (int32 SearchContextGroupIndex = 0; SearchContextGroupIndex < SearchContextGroups.Num();)
-		{
-			if (SearchContextGroups[SearchContextGroupIndex].Contains(SearchContext))
-			{
-				if (MainSearchContextGroupIndex == INDEX_NONE)
-				{
-					MainSearchContextGroupIndex = SearchContextGroupIndex;
-					SearchContextGroups[MainSearchContextGroupIndex].Add(SearchContext, SearchContextIndex);
-					++SearchContextGroupIndex;
-				}
-				else
-				{
-					SearchContextGroups[MainSearchContextGroupIndex].Merge(SearchContextGroups[SearchContextGroupIndex]);
-					SearchContextGroups.RemoveAt(SearchContextGroupIndex);
-				}
-			}
-			else
-			{
-				++SearchContextGroupIndex;
-			}
-		}
-		if (MainSearchContextGroupIndex == INDEX_NONE)
-		{
-			SearchContextGroups.AddDefaulted_GetRef().Add(SearchContext, SearchContextIndex);
-		}
-	}
-}
-
 } // namespace UE::PoseSearch
 
 #if UE_POSE_SEARCH_ENABLE_INTERACTION_SUBSYSTEM_CRITICAL_SECTION
@@ -241,7 +160,7 @@ UE::PoseSearch::FInteractionIsland& UPoseSearchInteractionSubsystem::GetAvailabl
 
 	for (FInteractionIsland* Island : Islands)
 	{
-		if (Island->IsUninjected())
+		if (!Island->IsInitialized())
 		{
 			return *Island;
 		}
@@ -258,7 +177,7 @@ void UPoseSearchInteractionSubsystem::DestroyAllIslands()
 	}
 }
 
-void UPoseSearchInteractionSubsystem::RegenerateAllIslands(TConstArrayView<UE::PoseSearch::FInteractionSearchContextGroup> SearchContextGroups, TConstArrayView<UE::PoseSearch::FInteractionSearchContext> SearchContexts)
+void UPoseSearchInteractionSubsystem::RegenerateAllIslands(TConstArrayView<UE::PoseSearch::FInteractionSearchContext> SearchContexts)
 {
 	QUICK_SCOPE_CYCLE_COUNTER(STAT_UPoseSearchInteractionSubsystem_RegenerateAllIslands);
 
@@ -278,23 +197,111 @@ void UPoseSearchInteractionSubsystem::RegenerateAllIslands(TConstArrayView<UE::P
 		for (FInteractionIsland* Island : Islands)
 		{
 			CheckInteractionThreadSafety(Island);
-			Island->Uninject(true);
+			Island->Uninitialize(true);
 		}
 	}
 
-	for (const FInteractionSearchContextGroup& SearchContextGroup : SearchContextGroups)
+	struct FInteractionSearchContextGroup
 	{
+		bool Contains(const FInteractionSearchContext& SearchContext) const
+		{
+			for (const TWeakObjectPtr<const UObject>& AnimContext : SearchContext.AnimContexts)
+			{
+				if (AnimContexts.Find(AnimContext.Get()))
+				{
+					return true;
+				}
+			}
+			return false;
+		}
+
+		void Add(const FInteractionSearchContext& SearchContext, int32 SearchContextIndex)
+		{
+			for (const TWeakObjectPtr<const UObject>& AnimContext : SearchContext.AnimContexts)
+			{
+				if (AnimContext.Get())
+				{
+					AnimContexts.Add(AnimContext.Get());
+				}
+			}
+
+			SearchContextsIndices.Add(SearchContextIndex);
+		}
+
+		void Merge(const FInteractionSearchContextGroup& SearchContextGroup)
+		{
+			for (const UObject* AnimContext : SearchContextGroup.AnimContexts)
+			{
+				check(AnimContext);
+				AnimContexts.Add(AnimContext);
+			}
+
+			for (int32 SearchContextsIndex : SearchContextGroup.SearchContextsIndices)
+			{
+				SearchContextsIndices.Add(SearchContextsIndex);
+			}
+		}
+
+		TSet<const UObject*, DefaultKeyFuncs<const UObject*>, TInlineSetAllocator<16, TMemStackSetAllocator<>>> AnimContexts;
+		TArray<int32, TInlineAllocator<16, TMemStackAllocator<>>> SearchContextsIndices;
+	};
+
+	// grouping SearchContexts AnimContext(s) in FInteractionSearchContextGroup(s). We'll create as many interaction islands as many groups
+	TArray<FInteractionSearchContextGroup, TInlineAllocator<PreallocatedSearchesNum, TMemStackAllocator<>>> SearchContextGroups;
+	for (int32 SearchContextIndex = 0; SearchContextIndex < SearchContexts.Num(); ++SearchContextIndex)
+	{
+		const FInteractionSearchContext& SearchContext = SearchContexts[SearchContextIndex];
+
+		int32 MainSearchContextGroupIndex = INDEX_NONE;
+		for (int32 SearchContextGroupIndex = 0; SearchContextGroupIndex < SearchContextGroups.Num();)
+		{
+			if (SearchContextGroups[SearchContextGroupIndex].Contains(SearchContext))
+			{
+				if (MainSearchContextGroupIndex == INDEX_NONE)
+				{
+					MainSearchContextGroupIndex = SearchContextGroupIndex;
+					SearchContextGroups[MainSearchContextGroupIndex].Add(SearchContext, SearchContextIndex);
+					++SearchContextGroupIndex;
+				}
+				else
+				{
+					SearchContextGroups[MainSearchContextGroupIndex].Merge(SearchContextGroups[SearchContextGroupIndex]);
+					SearchContextGroups.RemoveAt(SearchContextGroupIndex);
+				}
+			}
+			else
+			{
+				++SearchContextGroupIndex;
+			}
+		}
+		if (MainSearchContextGroupIndex == INDEX_NONE)
+		{
+			SearchContextGroups.AddDefaulted_GetRef().Add(SearchContext, SearchContextIndex);
+		}
+	}
+
+	for (FInteractionSearchContextGroup& SearchContextGroup : SearchContextGroups)
+	{
+		// @todo: search for the most suitable island to reuse to avoid having to Uninitialize/RemoveTickDependencies and InjectToActor right away
 		FInteractionIsland& Island = GetAvailableIsland();
 		CheckInteractionThreadSafety(&Island);
 
-		for (const UObject* AnimContext : SearchContextGroup.AnimContexts)
+		// initializing the island with its assigned SearchContexts
+		bool bAreTickDependenciesRequired = false;
+		check(Island.GetSearchContexts().IsEmpty());
+		for (int32 SearchContextsIndex : SearchContextGroup.SearchContextsIndices)
 		{
-			Island.InjectToActor(AnimContext);
+			const FInteractionSearchContext& SearchContext = SearchContexts[SearchContextsIndex];
+			// if there're at least two AnimContext(s) potentially interacting with each other 
+			// (where the search involves 2+ characters) tick dependencies are required to be thread safe
+			bAreTickDependenciesRequired |= SearchContext.AnimContexts.Num() > 1;
+			Island.AddSearchContext(SearchContext);
 		}
 
-		for (int32 SearchContextsIndex : SearchContextGroup.SearchContextsIndices)
+		// injecting tick dependencies between island AnimContext
+		for (const UObject* AnimContext : SearchContextGroup.AnimContexts)
 		{
-			Island.AddSearchContext(SearchContexts[SearchContextsIndex]);
+			Island.InjectToActor(AnimContext, bAreTickDependenciesRequired);
 		}
 	}
 }
@@ -690,7 +697,14 @@ void UPoseSearchInteractionSubsystem::GenerateSearchContexts(UE::PoseSearch::FIn
 		{
 			for (const FAnimContextInfo& AnimContextInfo : AnimContextInfos)
 			{
-				if (!AnimContextInfo.NearbyAnimContextInfos.IsEmpty() && !VisitedAnimContextInfos.Find(&AnimContextInfo))
+				if (AnimContextInfo.NearbyAnimContextInfos.IsEmpty())
+				{
+					check(!VisitedAnimContextInfos.Find(&AnimContextInfo));
+					// no need to add this context to the VisitedAnimContextInfos since it's isolated!
+					OnNewAnimAnimContextInfoFound(AnimContextInfo);
+					OnDoneGroupingAnimContexts();
+				}
+				else if(!VisitedAnimContextInfos.Find(&AnimContextInfo))
 				{
 					// starting the evaluation of a new set of grouped AnimContext(s)
 
@@ -910,7 +924,7 @@ void UPoseSearchInteractionSubsystem::Tick(float DeltaSeconds)
 		bool bAllUninjected = true;
 		for (FInteractionIsland* Island : Islands)
 		{
-			if (!Island->IsUninjected())
+			if (Island->IsInitialized())
 			{
 				bAllUninjected = false;
 			}
@@ -935,15 +949,12 @@ void UPoseSearchInteractionSubsystem::Tick(float DeltaSeconds)
 	// populating the continuing pose properties for the SearchContexts from the current Islands
 	PopulateContinuingProperties(DeltaSeconds, SearchContexts);
 	
-	// generating groups of AnimContext(s) from the SearchContexts to be able to generate islands later on
-	FInteractionSearchContextGroups SearchContextGroups;
-	GenerateInteractionSearchContextGroups(SearchContexts, SearchContextGroups);
-
 #if ENABLE_DRAW_DEBUG && ENABLE_ANIM_DEBUG
-	DebugDraw();
+	// drawing the current frame islands to be consistent with the search, before regenerating the islands with the newly published availabilities
+	DebugDrawIslands();
 #endif // ENABLE_DRAW_DEBUG && ENABLE_ANIM_DEBUG
 
-	RegenerateAllIslands(SearchContextGroups, SearchContexts);
+	RegenerateAllIslands(SearchContexts);
 
 	// clearing up AvailabilitiesMap for the next frame
 	ResetAvailabilities();
@@ -1007,7 +1018,7 @@ void UPoseSearchInteractionSubsystem::Query_AnyThread(const TArrayView<const FPo
 }
 
 #if ENABLE_DRAW_DEBUG && ENABLE_ANIM_DEBUG
-void UPoseSearchInteractionSubsystem::DebugDraw() const
+void UPoseSearchInteractionSubsystem::DebugDrawIslands() const
 {
 	using namespace UE::PoseSearch;
 
@@ -1039,7 +1050,7 @@ void UPoseSearchInteractionSubsystem::DebugDraw() const
 
 	for (const FInteractionIsland* Island : Islands)
 	{
-		if (!Island->IsUninjected())
+		if (Island->IsInitialized())
 		{
 			const FColor& Color = Colors[CurrentColorIndex];
 
@@ -1059,8 +1070,21 @@ void UPoseSearchInteractionSubsystem::DebugDraw() const
 
 							if (MaxBroadPhaseRadius > UE_SMALL_NUMBER)
 							{
-								const FVector Position = GetContextLocation(AnimContext);
-								DrawDebugCircle(GetWorld(), Position, MaxBroadPhaseRadius, 40, Color, false, 0.f, SDPG_Foreground, 0.f, FVector::XAxisVector, FVector::YAxisVector, false);
+								const UWorld* World = GetWorld();
+								const FTransform& Transform = GetContextTransform(AnimContext);
+								DrawDebugCircle(World, Transform.GetLocation(), MaxBroadPhaseRadius, 40, Color, false, 0.f, SDPG_Foreground, 0.f, FVector::XAxisVector, FVector::YAxisVector, false);
+
+								if (!Island->HasTickDependencies())
+								{
+									const FVector ForwardAxisStart = Transform.TransformPosition(FVector::ForwardVector * MaxBroadPhaseRadius);
+									const FVector ForwardAxisEnd = Transform.TransformPosition(FVector::ForwardVector * -MaxBroadPhaseRadius);
+
+									const FVector LeftAxisStart = Transform.TransformPosition(FVector::LeftVector * MaxBroadPhaseRadius);
+									const FVector LeftAxisEnd = Transform.TransformPosition(FVector::LeftVector * -MaxBroadPhaseRadius);
+
+									DrawDebugLine(World, ForwardAxisStart, ForwardAxisEnd, Color, false, 0.f, SDPG_Foreground, 0.f);
+									DrawDebugLine(World, LeftAxisStart, LeftAxisEnd, Color, false, 0.f, SDPG_Foreground, 0.f);
+								}
 							}
 						}
 					}
diff --git a/Engine/Plugins/Animation/PoseSearch/Source/Runtime/Private/PoseSearchLibrary.cpp b/Engine/Plugins/Animation/PoseSearch/Source/Runtime/Private/PoseSearchLibrary.cpp
--- a/Engine/Plugins/Animation/PoseSearch/Source/Runtime/Private/PoseSearchLibrary.cpp
+++ b/Engine/Plugins/Animation/PoseSearch/Source/Runtime/Private/PoseSearchLibrary.cpp
@@ -260,6 +260,39 @@ namespace UE::PoseSearch
 
 		return false;
 	}
+
+	static FRole GetCommonDefaultRole(const TConstArrayView<TObjectPtr<const UPoseSearchDatabase>> Databases)
+	{
+		check(!Databases.IsEmpty());
+
+		FRole Role = DefaultRole;
+		if (const UPoseSearchDatabase* Database = Databases[0].Get())
+		{
+			if (const UPoseSearchSchema* Schema = Database->Schema)
+			{
+				Role = Schema->GetDefaultRole();
+			}
+		}
+
+#if WITH_EDITOR && ENABLE_ANIM_DEBUG 
+		for (int32 DatabaseIndex = 1; DatabaseIndex < Databases.Num(); ++DatabaseIndex)
+		{
+			if (const UPoseSearchDatabase* Database = Databases[DatabaseIndex].Get())
+			{
+				if (const UPoseSearchSchema* Schema = Database->Schema)
+				{
+					if (Role != Schema->GetDefaultRole())
+					{
+						UE_LOG(LogPoseSearch, Error, TEXT("GetCommonDefaultRole - inconsistent Role between provided Databases!"));
+						break;
+					}
+				}
+			}
+		}
+#endif // WITH_EDITOR && ENABLE_ANIM_DEBUG
+
+		return Role;
+	}
 }
 
 //////////////////////////////////////////////////////////////////////////
@@ -622,7 +655,7 @@ void UPoseSearchLibrary::UpdateMotionMatchingState(
 	}
 
 	FSearchContext SearchContext(0.f, &InOutMotionMatchingState.PoseIndicesHistory, InOutMotionMatchingState.CurrentSearchResult, PoseJumpThresholdTime);
-	SearchContext.AddRole(DefaultRole, AnimContext, PoseHistory);
+	SearchContext.AddRole(GetCommonDefaultRole(Databases), AnimContext, PoseHistory);
 
 	const bool bCanAdvance = InOutMotionMatchingState.CurrentSearchResult.CanAdvance(DeltaTime);
 
diff --git a/Engine/Plugins/Animation/PoseSearch/Source/Runtime/Private/PoseSearchSchema.cpp b/Engine/Plugins/Animation/PoseSearch/Source/Runtime/Private/PoseSearchSchema.cpp
--- a/Engine/Plugins/Animation/PoseSearch/Source/Runtime/Private/PoseSearchSchema.cpp
+++ b/Engine/Plugins/Animation/PoseSearch/Source/Runtime/Private/PoseSearchSchema.cpp
@@ -127,7 +127,6 @@ FPoseSearchRoledSkeleton* UPoseSearchSchema::GetRoledSkeleton(const UE::PoseSear
 	return nullptr;
 }
 
-#if WITH_EDITOR
 const UE::PoseSearch::FRole UPoseSearchSchema::GetDefaultRole() const
 {
 	if (!Skeletons.IsEmpty())
@@ -136,7 +135,6 @@ const UE::PoseSearch::FRole UPoseSearchSchema::GetDefaultRole() const
 	}
 	return UE::PoseSearch::DefaultRole;
 }
-#endif // WITH_EDITOR
 
 USkeleton* UPoseSearchSchema::GetSkeleton(const UE::PoseSearch::FRole& Role) const
 {
diff --git a/Engine/Plugins/Animation/PoseSearch/Source/Runtime/Public/PoseSearch/PoseSearchContext.h b/Engine/Plugins/Animation/PoseSearch/Source/Runtime/Public/PoseSearch/PoseSearchContext.h
--- a/Engine/Plugins/Animation/PoseSearch/Source/Runtime/Public/PoseSearch/PoseSearchContext.h
+++ b/Engine/Plugins/Animation/PoseSearch/Source/Runtime/Public/PoseSearch/PoseSearchContext.h
@@ -179,13 +179,13 @@ struct POSESEARCH_API FSearchContext
 	
 	bool ArePoseHistoriesValid() const;
 	const TArray<const IPoseHistory*, TInlineAllocator<PreallocatedRolesNum, TMemStackAllocator<>>>& GetPoseHistories() const { return PoseHistories; }
-	const IPoseHistory* GetPoseHistory(const FRole& Role) const { return PoseHistories[RoleToIndex[Role]]; }
+	const IPoseHistory* GetPoseHistory(const FRole& Role) const;
 
 	float GetDesiredPermutationTimeOffset() const { return DesiredPermutationTimeOffset; }
 	
 	UE_DEPRECATED(5.6, "Use GetAnimContext instead")
-	const UAnimInstance* GetAnimInstance(const FRole& Role) const { return Cast<UAnimInstance>(AnimContexts[RoleToIndex[Role]]); }
-	const UObject* GetAnimContext(const FRole& Role) const { return AnimContexts[RoleToIndex[Role]]; }
+	const UAnimInstance* GetAnimInstance(const FRole& Role) const;
+	const UObject* GetAnimContext(const FRole& Role) const;
 
 	UE_DEPRECATED(5.6, "Use GetAnimContexts instead")
 	const TArray<const UAnimInstance*, TInlineAllocator<PreallocatedRolesNum, TMemStackAllocator<>>> GetAnimInstances() const 
diff --git a/Engine/Plugins/Animation/PoseSearch/Source/Runtime/Public/PoseSearch/PoseSearchInteractionIsland.h b/Engine/Plugins/Animation/PoseSearch/Source/Runtime/Public/PoseSearch/PoseSearchInteractionIsland.h
--- a/Engine/Plugins/Animation/PoseSearch/Source/Runtime/Public/PoseSearch/PoseSearchInteractionIsland.h
+++ b/Engine/Plugins/Animation/PoseSearch/Source/Runtime/Public/PoseSearch/PoseSearchInteractionIsland.h
@@ -74,17 +74,21 @@ struct FInteractionIsland
 	const TArray<FInteractionSearchContext>& GetSearchContexts() const { return SearchContexts; }
 	const TArray<FInteractionSearchResult>& GetSearchResults() const { return SearchResults; }
 
-	bool IsUninjected() const;
-	void InjectToActor(const UObject* AnimContext);
-	void Uninject(bool bValidateTickDependencies);
-
 	static UActorComponent* FindPreTickComponent(const AActor* Actor);
 	static UActorComponent* FindPostTickComponent(const AActor* Actor);
 
 	const FInteractionSearchResult* FindSearchResult(const FInteractionSearchContext& SearchContext) const;
+
+	bool IsInitialized() const;
 	void AddSearchContext(const FInteractionSearchContext& SearchContext);
+	void Uninitialize(bool bValidateTickDependencies);
+
+	bool HasTickDependencies() const;
+	void InjectToActor(const UObject* AnimContext, bool bAddTickDependencies);
 
 private:
+	void RemoveTickDependencies(bool bValidateTickDependencies);
+
 	bool GetResult_AnyThread(const UObject* AnimContext, FPoseSearchInteractionBlueprintResult& Result);
 	const AActor* GetMainActor() const;
 	bool CanPerformSearch(const UObject* AnimContext) const;
@@ -108,6 +112,7 @@ private:
 		FInteractionIsland* Island = nullptr;
 	};
 	FPostTickFunction PostTickFunction;
+	bool bHasTickDependencies = false;
 
 	// islands don't get deallocated, so this array once warmed up will not hit the allocator and waste cycles
 	TArray<TWeakObjectPtr<UActorComponent>> PreTickComponents;
diff --git a/Engine/Plugins/Animation/PoseSearch/Source/Runtime/Public/PoseSearch/PoseSearchInteractionSubsystem.h b/Engine/Plugins/Animation/PoseSearch/Source/Runtime/Public/PoseSearch/PoseSearchInteractionSubsystem.h
--- a/Engine/Plugins/Animation/PoseSearch/Source/Runtime/Public/PoseSearch/PoseSearchInteractionSubsystem.h
+++ b/Engine/Plugins/Animation/PoseSearch/Source/Runtime/Public/PoseSearch/PoseSearchInteractionSubsystem.h
@@ -52,16 +52,16 @@ struct FAvailabilitiesMap : protected TMap<TWeakObjectPtr<const UObject>, FInter
 			{
 				AnimContextAvailabilities.Emplace(Availability, PoseHistoryName, HistoryCollector);
 			}
-			else
-			{
-				UE_LOG(LogPoseSearch, Warning, TEXT("FAvailabilitiesMap::AddAvailabilities trying adding an invalid Availability for AnimContext %s!"), *AnimContext->GetName());
-			}
+			// @todo: make this Warning less spammy and reenable it
+			//else
+			//{
+			//	UE_LOG(LogPoseSearch, Warning, TEXT("FAvailabilitiesMap::AddAvailabilities trying adding an invalid Availability for AnimContext %s!"), *AnimContext->GetName());
+			//}
 		}
 	}
 };
 
 struct FAnimContextInfos;
-struct FInteractionSearchContextGroup;
 struct FTagToDatabases;
 
 } // namespace UE::PoseSearch
@@ -143,7 +143,7 @@ private:
 
 	void DestroyIsland(int32 Index);
 	void DestroyAllIslands();
-	void RegenerateAllIslands(TConstArrayView<UE::PoseSearch::FInteractionSearchContextGroup> SearchContextGroups, TConstArrayView<UE::PoseSearch::FInteractionSearchContext> SearchContexts);
+	void RegenerateAllIslands(TConstArrayView<UE::PoseSearch::FInteractionSearchContext> SearchContexts);
 
 #if DO_CHECK
 	bool ValidateAllIslands() const;
@@ -153,7 +153,7 @@ private:
 	UE::PoseSearch::FInteractionIsland* FindIsland(const UObject* AnimContext);
 	
 #if ENABLE_DRAW_DEBUG && ENABLE_ANIM_DEBUG
-	void DebugDraw() const;
+	void DebugDrawIslands() const;
 	const UE::PoseSearch::FInteractionAvailabilities* FindAvailabilities(const UObject* AnimContext) const;
 #endif // ENABLE_DRAW_DEBUG && ENABLE_ANIM_DEBUG
 
