Automatically adjust the pose search columns size for better readability. Store settings (checkboxes) for better user experience when scrubbing.

#rnx

diff --git a/Engine/Plugins/Animation/PoseSearch/Source/Editor/Private/PoseSearchDebuggerDatabaseColumns.h b/Engine/Plugins/Animation/PoseSearch/Source/Editor/Private/PoseSearchDebuggerDatabaseColumns.h
--- a/Engine/Plugins/Animation/PoseSearch/Source/Editor/Private/PoseSearchDebuggerDatabaseColumns.h
+++ b/Engine/Plugins/Animation/PoseSearch/Source/Editor/Private/PoseSearchDebuggerDatabaseColumns.h
@@ -11,6 +11,8 @@
 #include "PoseSearchDatabaseViewModel.h"
 #include "PoseSearchDebuggerDatabaseRowData.h"
 #include "PoseSearchDebuggerViewModel.h"
+#include "Fonts/FontMeasure.h"
+#include "Framework/Application/SlateApplication.h"
 #include "Preferences/PersonaOptions.h"
 #include "Styling/AppStyle.h"
 #include "Subsystems/AssetEditorSubsystem.h"
@@ -52,6 +54,8 @@ struct IColumn : TSharedFromThis<IColumn>
 	virtual FSortPredicate GetSortPredicate() const = 0;
 
 	virtual TSharedRef<SWidget> GenerateWidget(const FRowDataRef& RowData) const = 0;
+
+	virtual float ComputeColumnWidth(TConstArrayView<FRowDataRef> Rows) const = 0;
 };
 
 /** Column struct to represent each column in the debugger database */
@@ -86,6 +90,22 @@ protected:
 	{
 		return FSlateColor(FLinearColor::White);
 	}
+	
+	virtual float ComputeColumnWidth(TConstArrayView<TSharedRef<FDebuggerDatabaseRowData>> Rows) const override
+	{
+		const TSharedRef<FSlateFontMeasure> FontMeasure = FSlateApplication::Get().GetRenderer()->GetFontMeasureService();
+		const FSlateFontInfo & CellTextStyle = FAppStyle::GetFontStyle("SmallText");
+		const float CellPadding = 10.0f;
+
+		float MinWidth = FontMeasure->Measure(GetLabel(), CellTextStyle).X + CellPadding;
+		
+		for (const TSharedRef<FDebuggerDatabaseRowData>& Row : Rows)
+		{
+			MinWidth = FMath::Max((FontMeasure->Measure(GetRowText(Row), CellTextStyle).X + CellPadding), MinWidth);
+		}
+		
+		return MinWidth;
+	}
 };
 
 struct FPoseIdx : ITextColumn
@@ -217,6 +237,22 @@ struct FDatabaseName : IColumn
 			}
 		});
 	}
+
+	virtual float ComputeColumnWidth(TConstArrayView<FRowDataRef> Rows) const override
+	{
+		const TSharedRef<FSlateFontMeasure> FontMeasure = FSlateApplication::Get().GetRenderer()->GetFontMeasureService();
+		const FSlateFontInfo & CellTextStyle = FAppStyle::GetFontStyle("SmallText");
+		const float CellPadding = 10.0f;
+
+		float MinWidth = FontMeasure->Measure(GetLabel(), CellTextStyle).X + CellPadding;
+		
+		for (const TSharedRef<FDebuggerDatabaseRowData>& Row : Rows)
+		{
+			MinWidth = FMath::Max(FontMeasure->Measure(Row->SharedData->DatabaseName, CellTextStyle).X + CellPadding, MinWidth);
+		}
+		
+		return MinWidth;
+	}
 };
 
 struct FAssetName : IColumn
@@ -291,6 +327,22 @@ struct FAssetName : IColumn
 					}
 				});
 	}
+
+	virtual float ComputeColumnWidth(TConstArrayView<FRowDataRef> Rows) const override
+	{
+		const TSharedRef<FSlateFontMeasure> FontMeasure = FSlateApplication::Get().GetRenderer()->GetFontMeasureService();
+		const FSlateFontInfo & CellTextStyle = FAppStyle::GetFontStyle("SmallText");
+		const float CellPadding = 10.0f;
+
+		float MinWidth = FontMeasure->Measure(GetLabel(), CellTextStyle).X + CellPadding;
+		
+		for (const TSharedRef<FDebuggerDatabaseRowData>& Row : Rows)
+		{
+			MinWidth = FMath::Max(FontMeasure->Measure(Row->AssetName, CellTextStyle).X, MinWidth);
+		}
+		
+		return MinWidth;
+	}
 };
 
 struct FFrame : ITextColumn
diff --git a/Engine/Plugins/Animation/PoseSearch/Source/Editor/Private/PoseSearchDebuggerDatabaseView.cpp b/Engine/Plugins/Animation/PoseSearch/Source/Editor/Private/PoseSearchDebuggerDatabaseView.cpp
--- a/Engine/Plugins/Animation/PoseSearch/Source/Editor/Private/PoseSearchDebuggerDatabaseView.cpp
+++ b/Engine/Plugins/Animation/PoseSearch/Source/Editor/Private/PoseSearchDebuggerDatabaseView.cpp
@@ -8,6 +8,7 @@
 #include "StructUtils/InstancedStruct.h"
 #include "Internationalization/Regex.h"
 #include "PoseSearchDebuggerDatabaseRow.h"
+#include "PoseSearchDebuggerSettings.h"
 #include "PoseSearchDebuggerView.h"
 #include "PoseSearchDebuggerViewModel.h"
 #include "PoseSearch/PoseSearchDatabase.h"
@@ -286,7 +287,7 @@ void SDebuggerDatabaseView::Update(const FTraceMotionMatchingStateMessage& State
 				AddUnfilteredDatabaseRow(Database, DynamicWeightsSqrt, UnfilteredDatabaseRows, SharedData, PoseEntry.DbPoseIdx, PoseEntry.PoseCandidateFlags, PoseToPCAValuesVectorIndexes, PoseEntry.Cost);
 			}
 
-			if (bShowAllPoses)
+			if (UPoseSearchDebuggerConfig::Get().bShowAllPoses)
 			{
 				bool bAddMissingPoseEntries = false;
 
@@ -451,6 +452,12 @@ void SDebuggerDatabaseView::Update(const FTraceMotionMatchingStateMessage& State
 			return Column0->SortIndex < Column1->SortIndex;
 		});
 
+		// Compute columns widths
+		for (TPair<FName, TSharedRef<IColumn>>& ColumnPair : Columns)
+		{
+			ColumnPair.Value->Width = ColumnPair.Value->ComputeColumnWidth(UnfilteredDatabaseRows);
+		}
+		
 		// Add columns from map to header row
 		for (TPair<FName, TSharedRef<IColumn>>& ColumnPair : Columns)
 		{
@@ -462,10 +469,11 @@ void SDebuggerDatabaseView::Update(const FTraceMotionMatchingStateMessage& State
 					.DefaultTooltip(Column.GetLabelTooltip())
 					.SortMode(this, &SDebuggerDatabaseView::GetColumnSortMode, Column.ColumnId)
 					.OnSort(this, &SDebuggerDatabaseView::OnColumnSortModeChanged)
-					.FillWidth(this, &SDebuggerDatabaseView::GetColumnWidth, Column.ColumnId)
-					.VAlignCell(VAlign_Center)
+					.ManualWidth(this, &SDebuggerDatabaseView::GetColumnWidth, Column.ColumnId)
+					.OnWidthChanged(this, &SDebuggerDatabaseView::OnColumnWidthChanged, Column.ColumnId)
 					.VAlignHeader(VAlign_Center)
-					.HAlignHeader(HAlign_Center)
+					.VAlignCell(VAlign_Center)
+					.HAlignHeader(HAlign_Fill)
 					.HAlignCell(HAlign_Fill)
 					.HeaderContent()
 					[
@@ -483,9 +491,7 @@ void SDebuggerDatabaseView::Update(const FTraceMotionMatchingStateMessage& State
 
 				FilteredDatabaseView.HeaderRow->AddColumn(ColumnArgs);
 				ContinuingPoseView.HeaderRow->AddColumn(ColumnArgs);
-
-				// Every time the active column is changed, update the database column
-				ActiveView.HeaderRow->AddColumn(ColumnArgs.OnWidthChanged(this, &SDebuggerDatabaseView::OnColumnWidthChanged, Column.ColumnId));
+				ActiveView.HeaderRow->AddColumn(ColumnArgs);
 			}
 		}
 	}
@@ -535,7 +541,7 @@ void SDebuggerDatabaseView::OnFilterTextChanged(const FText& SearchText)
 
 void SDebuggerDatabaseView::OnShowAllPosesCheckboxChanged(ECheckBoxState State)
 {
-	bShowAllPoses = (State == ECheckBoxState::Checked);
+	UPoseSearchDebuggerConfig::Get().bShowAllPoses = State == ECheckBoxState::Checked;
 
 	if (TSharedPtr<SDebuggerView> DebuggerView = ParentDebuggerViewPtr.Pin())
 	{
@@ -551,41 +557,19 @@ void SDebuggerDatabaseView::OnShowAllPosesCheckboxChanged(ECheckBoxState State)
 
 void SDebuggerDatabaseView::OnShowOnlyBestAssetPoseCheckboxChanged(ECheckBoxState State)
 {
-	if (State == ECheckBoxState::Checked)
-	{
-		bShowOnlyBestAssetPose = true;
-	}
-	else
-	{
-		bShowOnlyBestAssetPose = false;
-	}
-
+	UPoseSearchDebuggerConfig::Get().bShowOnlyBestAssetPose = State == ECheckBoxState::Checked;
 	PopulateViewRows();
 }
 
 void SDebuggerDatabaseView::OnHideInvalidPosesCheckboxChanged(ECheckBoxState State)
 {
-	if (State == ECheckBoxState::Checked) 
-	{
-		bHideInvalidPoses = true;
-	}
-	else 
-	{
-		bHideInvalidPoses = false;
-	}
+	UPoseSearchDebuggerConfig::Get().bHideInvalidPoses = State == ECheckBoxState::Checked;
 	PopulateViewRows();
 }
 
 void SDebuggerDatabaseView::OnUseRegexCheckboxChanged(ECheckBoxState State)
 {
-	if (State == ECheckBoxState::Checked)
-	{
-		bUseRegex = true;
-	}
-	else
-	{
-		bUseRegex = false;
-	}
+	UPoseSearchDebuggerConfig::Get().bUseRegex = State == ECheckBoxState::Checked;
 	PopulateViewRows();
 }
 
@@ -633,7 +617,7 @@ void SDebuggerDatabaseView::PopulateViewRows()
 		FRegexPattern Pattern(FilterString);
 	
 		TSet<int32> BestAssetPoseIndex;
-		if (bShowOnlyBestAssetPose)
+		if (UPoseSearchDebuggerConfig::Get().bShowOnlyBestAssetPose)
 		{
 			TArray<int32> SortIndex;
 			SortIndex.SetNumUninitialized(UnfilteredDatabaseRowsNum);
@@ -700,7 +684,7 @@ void SDebuggerDatabaseView::PopulateViewRows()
 				bTryAddToFilteredDatabaseViewRows = false;
 			}
 
-			if (bShowOnlyBestAssetPose)
+			if (UPoseSearchDebuggerConfig::Get().bShowOnlyBestAssetPose)
 			{
 				if (!BestAssetPoseIndex.Find(UnfilteredRow->PoseIdx))
 				{
@@ -711,11 +695,11 @@ void SDebuggerDatabaseView::PopulateViewRows()
 			if (bTryAddToFilteredDatabaseViewRows)
 			{
 				bool bPassesNameFilter = true;
-				if (bHideInvalidPoses && !EnumHasAnyFlags(UnfilteredRow->PoseCandidateFlags, EPoseCandidateFlags::AnyValidMask))
+				if (UPoseSearchDebuggerConfig::Get().bHideInvalidPoses && !EnumHasAnyFlags(UnfilteredRow->PoseCandidateFlags, EPoseCandidateFlags::AnyValidMask))
 				{
 					bPassesNameFilter = false;
 				}
-				else if (bUseRegex)
+				else if (UPoseSearchDebuggerConfig::Get().bUseRegex)
 				{
 					FRegexMatcher Matcher(Pattern, UnfilteredRow->AssetName);
 					bPassesNameFilter = Matcher.FindNext();
@@ -1100,6 +1084,10 @@ void SDebuggerDatabaseView::Construct(const FArguments& InArgs)
 					.Padding(10, 5, 10, 5)
 					[
 						SNew(SCheckBox)
+						.IsChecked_Lambda([this]
+						{
+							return UPoseSearchDebuggerConfig::Get().bShowAllPoses ? ECheckBoxState::Checked : ECheckBoxState::Unchecked; 
+						})
 						.OnCheckStateChanged_Lambda([this](ECheckBoxState State)
 						{
 							   SDebuggerDatabaseView::OnShowAllPosesCheckboxChanged(State);
@@ -1114,6 +1102,10 @@ void SDebuggerDatabaseView::Construct(const FArguments& InArgs)
 					.Padding(10, 5, 10, 5)
 					[
 						SNew(SCheckBox)
+						.IsChecked_Lambda([this]
+						{
+							return UPoseSearchDebuggerConfig::Get().bShowOnlyBestAssetPose ? ECheckBoxState::Checked : ECheckBoxState::Unchecked; 
+						})
 						.OnCheckStateChanged_Lambda([this](ECheckBoxState State)
 						{
 							SDebuggerDatabaseView::OnShowOnlyBestAssetPoseCheckboxChanged(State);
@@ -1128,6 +1120,10 @@ void SDebuggerDatabaseView::Construct(const FArguments& InArgs)
 					.Padding(10, 5, 10, 5)
 					[
 						SNew(SCheckBox)
+						.IsChecked_Lambda([this]
+						{
+							return UPoseSearchDebuggerConfig::Get().bHideInvalidPoses ? ECheckBoxState::Checked : ECheckBoxState::Unchecked; 
+						})
 						.OnCheckStateChanged_Lambda([this](ECheckBoxState State)
 						{
 							   SDebuggerDatabaseView::OnHideInvalidPosesCheckboxChanged(State);
@@ -1142,14 +1138,18 @@ void SDebuggerDatabaseView::Construct(const FArguments& InArgs)
 					.Padding(10, 5, 10, 5)
 					[
 						SNew(SCheckBox)
+						.IsChecked_Lambda([this]
+						{
+							return UPoseSearchDebuggerConfig::Get().bUseRegex ? ECheckBoxState::Checked : ECheckBoxState::Unchecked; 
+						})
 						.OnCheckStateChanged_Lambda([this](ECheckBoxState State)
 							{
 								SDebuggerDatabaseView::OnUseRegexCheckboxChanged(State);
 							})
-					[
-						SNew(STextBlock)
-						.Text(LOCTEXT("PoseSearchDebuggerUseRegexFlag", "Use Regex"))
-					]
+						[
+							SNew(STextBlock)
+							.Text(LOCTEXT("PoseSearchDebuggerUseRegexFlag", "Use Regex"))
+						]
 					]
 				]
 		
diff --git a/Engine/Plugins/Animation/PoseSearch/Source/Editor/Private/PoseSearchDebuggerDatabaseView.h b/Engine/Plugins/Animation/PoseSearch/Source/Editor/Private/PoseSearchDebuggerDatabaseView.h
--- a/Engine/Plugins/Animation/PoseSearch/Source/Editor/Private/PoseSearchDebuggerDatabaseView.h
+++ b/Engine/Plugins/Animation/PoseSearch/Source/Editor/Private/PoseSearchDebuggerDatabaseView.h
@@ -146,18 +146,6 @@ private:
 	/** Text used to filter DatabaseView */
 	FText FilterText;
 
-	/* Bool used to show all poses from display */
-	bool bShowAllPoses = false;
-	
-	/* Bool used to show only the best pose of every asset */
-	bool bShowOnlyBestAssetPose = false;
-
-	/* Bool used to hide invalid poses from display */
-	bool bHideInvalidPoses = false;
-
-	/* Bool used to use FilterText as regex filter*/
-	bool bUseRegex = false;
-
 	FText ReasonForNoActivePose;
 	FText ReasonForNoContinuingPose;
 	FText ReasonForNoCandidates;
diff --git a/Engine/Plugins/Animation/PoseSearch/Source/Editor/Private/PoseSearchDebuggerSettings.cpp b/Engine/Plugins/Animation/PoseSearch/Source/Editor/Private/PoseSearchDebuggerSettings.cpp
new file mode 100644
--- /dev/null
+++ b/Engine/Plugins/Animation/PoseSearch/Source/Editor/Private/PoseSearchDebuggerSettings.cpp
@@ -0,0 +1,26 @@
+// Copyright Epic Games, Inc. All Rights Reserved.
+
+#include "PoseSearchDebuggerSettings.h"
+#include "Misc/CoreDelegates.h"
+
+#include UE_INLINE_GENERATED_CPP_BY_NAME(PoseSearchDebuggerSettings)
+
+UPoseSearchDebuggerConfig::UPoseSearchDebuggerConfig()
+{
+	// Ensure we save settings on exit.
+	FCoreDelegates::OnPreExit.AddLambda([]()
+	{
+		if (UPoseSearchDebuggerConfig* Config = GetMutableDefault<UPoseSearchDebuggerConfig>())
+		{
+			Config->SaveConfig();
+		}
+	});
+}
+
+UPoseSearchDebuggerConfig& UPoseSearchDebuggerConfig::Get()
+{
+	UPoseSearchDebuggerConfig* MutableCDO = GetMutableDefault<UPoseSearchDebuggerConfig>();
+	check(MutableCDO != nullptr)
+	
+	return *MutableCDO;
+}
diff --git a/Engine/Plugins/Animation/PoseSearch/Source/Editor/Private/PoseSearchDebuggerSettings.h b/Engine/Plugins/Animation/PoseSearch/Source/Editor/Private/PoseSearchDebuggerSettings.h
new file mode 100644
--- /dev/null
+++ b/Engine/Plugins/Animation/PoseSearch/Source/Editor/Private/PoseSearchDebuggerSettings.h
@@ -0,0 +1,57 @@
+// Copyright Epic Games, Inc. All Rights Reserved.
+
+#pragma once
+
+#include "CoreMinimal.h"
+#include "Engine/DeveloperSettings.h"
+
+#include "PoseSearchDebuggerSettings.generated.h"
+
+/** Settings that holds editor configurations. Not accessible in Project Settings. */
+UCLASS(Config = EditorPerProjectUserSettings, MinimalAPI)
+class UPoseSearchDebuggerConfig : public UObject
+{
+	GENERATED_BODY()
+
+public:
+
+	UPoseSearchDebuggerConfig();
+	
+	static UPoseSearchDebuggerConfig& Get();
+	
+	// General options
+
+	/** Used to draw the query used to get pose for this frame from Motion Matching algorithm */
+	UPROPERTY(config)
+	bool bDrawQuery = false;
+
+	/** Used to show the trajectory used for this frame to run Motion Matching against */
+	UPROPERTY(config)
+	bool bDrawTrajectory = false;
+
+	/** Used to show the traced pose search history value */
+	UPROPERTY(config)
+	bool bDrawHistory = false;
+
+	/** Bool used to break down channels and display a complete picture of the weights that determined the final pose from Motion Matching algorithm */
+	UPROPERTY(config)
+	bool bIsVerbose = false;
+
+	// Pose Candidates options
+
+	/* Bool used to show all poses from display */
+	UPROPERTY(config)
+	bool bShowAllPoses = false;
+
+	/* Bool used to show only the best pose of every asset */
+	UPROPERTY(config)
+	bool bShowOnlyBestAssetPose = false;
+
+	/* Bool used to hide invalid poses from display */
+	UPROPERTY(config)
+	bool bHideInvalidPoses = false;
+
+	/* Bool used to use FilterText as regex filter*/
+	UPROPERTY(config)
+	bool bUseRegex = false;
+};
\ No newline at end of file
diff --git a/Engine/Plugins/Animation/PoseSearch/Source/Editor/Private/PoseSearchDebuggerViewModel.cpp b/Engine/Plugins/Animation/PoseSearch/Source/Editor/Private/PoseSearchDebuggerViewModel.cpp
--- a/Engine/Plugins/Animation/PoseSearch/Source/Editor/Private/PoseSearchDebuggerViewModel.cpp
+++ b/Engine/Plugins/Animation/PoseSearch/Source/Editor/Private/PoseSearchDebuggerViewModel.cpp
@@ -9,6 +9,7 @@
 #include "IGameplayProvider.h"
 #include "StructUtils/InstancedStruct.h"
 #include "IRewindDebugger.h"
+#include "PoseSearchDebuggerSettings.h"
 #include "PoseSearch/PoseSearchDatabase.h"
 #include "PoseSearch/PoseSearchDerivedData.h"
 #include "PoseSearch/PoseSearchSchema.h"
@@ -205,4 +206,44 @@ void FDebuggerViewModel::OnUpdateSearchSelection(int32 InSearchId)
 	}
 }
 
+void FDebuggerViewModel::SetVerbose(bool bVerbose)
+{
+	UPoseSearchDebuggerConfig::Get().bIsVerbose = bVerbose;
+}
+
+bool FDebuggerViewModel::IsVerbose() const
+{
+	return UPoseSearchDebuggerConfig::Get().bIsVerbose;
+}
+
+void FDebuggerViewModel::SetDrawQuery(bool bInDrawQuery)
+{
+	UPoseSearchDebuggerConfig::Get().bDrawQuery = bInDrawQuery;
+}
+
+bool FDebuggerViewModel::GetDrawQuery() const
+{
+	return UPoseSearchDebuggerConfig::Get().bDrawQuery;
+}
+
+void FDebuggerViewModel::SetDrawTrajectory(bool bInDrawTrajectory)
+{
+	UPoseSearchDebuggerConfig::Get().bDrawTrajectory = bInDrawTrajectory;
+}
+
+bool FDebuggerViewModel::GetDrawTrajectory() const
+{
+	return UPoseSearchDebuggerConfig::Get().bDrawTrajectory;
+}
+
+void FDebuggerViewModel::SetDrawHistory(bool bInDrawHistory)
+{
+	UPoseSearchDebuggerConfig::Get().bDrawHistory = bInDrawHistory;
+}
+
+bool FDebuggerViewModel::GetDrawHistory() const
+{
+	return UPoseSearchDebuggerConfig::Get().bDrawHistory;
+}
+	
 } // namespace UE::PoseSearch
diff --git a/Engine/Plugins/Animation/PoseSearch/Source/Editor/Private/PoseSearchDebuggerViewModel.h b/Engine/Plugins/Animation/PoseSearch/Source/Editor/Private/PoseSearchDebuggerViewModel.h
--- a/Engine/Plugins/Animation/PoseSearch/Source/Editor/Private/PoseSearchDebuggerViewModel.h
+++ b/Engine/Plugins/Animation/PoseSearch/Source/Editor/Private/PoseSearchDebuggerViewModel.h
@@ -38,17 +38,17 @@ public:
 	/** Updates active motion matching state based on node selection */
 	void OnUpdateSearchSelection(int32 InSearchId);
 
-	void SetVerbose(bool bVerbose) { bIsVerbose = bVerbose; }
-	bool IsVerbose() const { return bIsVerbose; }
+	void SetVerbose(bool bVerbose);
+	bool IsVerbose() const;
 
-	void SetDrawQuery(bool bInDrawQuery) { bDrawQuery = bInDrawQuery; }
-	bool GetDrawQuery() const { return bDrawQuery; }
+	void SetDrawQuery(bool bInDrawQuery);
+	bool GetDrawQuery() const;
 
-	void SetDrawTrajectory(bool bInDrawTrajectory) { bDrawTrajectory = bInDrawTrajectory; }
-	bool GetDrawTrajectory() const { return bDrawTrajectory; }
+	void SetDrawTrajectory(bool bInDrawTrajectory);
+	bool GetDrawTrajectory() const;
 
-	void SetDrawHistory(bool bInDrawHistory) { bDrawHistory = bInDrawHistory; }
-	bool GetDrawHistory() const { return bDrawHistory; }
+	void SetDrawHistory(bool bInDrawHistory);
+	bool GetDrawHistory() const;
 
 private:
 
@@ -67,11 +67,6 @@ private:
 
 	/** Actor object populated by the timeline mapped from the MotionMatchingStates.SkeletalMeshComponentIds. used as input for additional bone transforms to draw channels */
 	TMap<uint64, TWeakObjectPtr<AActor>> DebugDrawActors;
-
-	bool bIsVerbose = false;
-	bool bDrawQuery = true;	
-	bool bDrawTrajectory = false;
-	bool bDrawHistory = false;
 	
 	/** Limits some public API */
 	friend class FDebugger;
