This change lays the ground for doing pose trajectory generation inside mass for mass agents. It makes the following changes
* A new UpdateHistory_WorldSpace method that doesn't require the mover relative velocity. This is required because Mass processor doesn't have fast access to the up to date mover velocity at the time of calculation
* Deprecated some methods which took FPoseSearchTrajectoryData. This parameter was unused and is not available in mass, and it would be a waste of time to pipe it in
* Add profile tag to RigUnit_GenerateMoverTrajectory

#tests editor win64 animsandbox

#virtualized

diff --git a/Engine/Plugins/Animation/PoseSearch/Source/Runtime/Private/PoseSearchHistory.cpp b/Engine/Plugins/Animation/PoseSearch/Source/Runtime/Private/PoseSearchHistory.cpp
--- a/Engine/Plugins/Animation/PoseSearch/Source/Runtime/Private/PoseSearchHistory.cpp
+++ b/Engine/Plugins/Animation/PoseSearch/Source/Runtime/Private/PoseSearchHistory.cpp
@@ -804,8 +804,8 @@ void FPoseHistory::GenerateTrajectory(const UAnimInstance* AnimInstance, float D
 	FPoseSearchTrajectoryData::FDerived TrajectoryDataDerived;
 	if (TrajectoryData.UpdateData(DeltaTime, AnimInstance, TrajectoryDataDerived, TrajectoryDataState))
 	{
-		UPoseSearchTrajectoryLibrary::InitTrajectorySamples(Trajectory, TrajectoryData, TrajectoryDataDerived.Position, TrajectoryDataDerived.Facing, TrajectoryDataSampling, DeltaTime);
-		UPoseSearchTrajectoryLibrary::UpdateHistory_TransformHistory(Trajectory, TrajectoryData, TrajectoryDataDerived.Position, TrajectoryDataDerived.Velocity, TrajectoryDataSampling, DeltaTime);
+		UPoseSearchTrajectoryLibrary::InitTrajectorySamples(Trajectory, TrajectoryDataDerived.Position, TrajectoryDataDerived.Facing, TrajectoryDataSampling, DeltaTime);
+		UPoseSearchTrajectoryLibrary::UpdateHistory_TransformHistory(Trajectory, TrajectoryDataDerived.Position, TrajectoryDataDerived.Velocity, TrajectoryDataSampling, DeltaTime);
 		UPoseSearchTrajectoryLibrary::UpdatePrediction_SimulateCharacterMovement(Trajectory, TrajectoryData, TrajectoryDataDerived, TrajectoryDataSampling, DeltaTime);
 
 		// @todo: support TrajectorySpeedMultiplier
diff --git a/Engine/Plugins/Animation/PoseSearch/Source/Runtime/Private/PoseSearchTrajectoryLibrary.cpp b/Engine/Plugins/Animation/PoseSearch/Source/Runtime/Private/PoseSearchTrajectoryLibrary.cpp
--- a/Engine/Plugins/Animation/PoseSearch/Source/Runtime/Private/PoseSearchTrajectoryLibrary.cpp
+++ b/Engine/Plugins/Animation/PoseSearch/Source/Runtime/Private/PoseSearchTrajectoryLibrary.cpp
@@ -166,6 +166,7 @@ FVector FPoseSearchTrajectoryData::StepCharacterMovementGroundPrediction(
 	return OutVelocity;
 }
 
+// Deprecated
 void UPoseSearchTrajectoryLibrary::InitTrajectorySamples(
 	FPoseSearchQueryTrajectory& Trajectory,
 	const FPoseSearchTrajectoryData& TrajectoryData,
@@ -175,15 +176,29 @@ void UPoseSearchTrajectoryLibrary::InitTrajectorySamples(
 {
 	InitTrajectorySamples(
 		Trajectory,
-		TrajectoryData,
 		TrajectoryDataDerived.Position, TrajectoryDataDerived.Facing,
 		TrajectoryDataSampling, 
 		DeltaTime );
 }
 
+// Deprecated
+void UPoseSearchTrajectoryLibrary::InitTrajectorySamples(
+	FPoseSearchQueryTrajectory& Trajectory,
+	const FPoseSearchTrajectoryData& TrajectoryData,	// Unused parameter
+	FVector DefaultPosition, FQuat DefaultFacing,
+	const FPoseSearchTrajectoryData::FSampling& TrajectoryDataSampling,
+	float DeltaTime)
+{
+	InitTrajectorySamples(
+		Trajectory,
+		DefaultPosition,
+		DefaultFacing,
+		TrajectoryDataSampling,
+		DeltaTime);
+}
+
 void UPoseSearchTrajectoryLibrary::InitTrajectorySamples(
 	FPoseSearchQueryTrajectory& Trajectory,
-	const FPoseSearchTrajectoryData& TrajectoryData,
 	FVector DefaultPosition, FQuat DefaultFacing,
 	const FPoseSearchTrajectoryData::FSampling& TrajectoryDataSampling,
 	float DeltaTime)
@@ -218,24 +233,25 @@ void UPoseSearchTrajectoryLibrary::InitTrajectorySamples(
 	}
 }
 
+// Deprecated
 void UPoseSearchTrajectoryLibrary::UpdateHistory_TransformHistory(
 	FPoseSearchQueryTrajectory& Trajectory,
-	const FPoseSearchTrajectoryData& TrajectoryData,
-	const FPoseSearchTrajectoryData::FDerived& TrajectoryDataDerived,
+	const FPoseSearchTrajectoryData& TrajectoryData,	// Unused parameter
+	FVector CurrentPosition,
+	FVector CurrentVelocity,
 	const FPoseSearchTrajectoryData::FSampling& TrajectoryDataSampling,
 	float DeltaTime)
 {
 	UpdateHistory_TransformHistory(
 		Trajectory,
-		TrajectoryData,
-		TrajectoryDataDerived.Position, TrajectoryDataDerived.Velocity,
+		CurrentPosition,
+		CurrentVelocity,
 		TrajectoryDataSampling,
 		DeltaTime);
 }
 
 void UPoseSearchTrajectoryLibrary::UpdateHistory_TransformHistory(
 	FPoseSearchQueryTrajectory& Trajectory,
-	const FPoseSearchTrajectoryData& TrajectoryData,
 	FVector CurrentPosition,
 	FVector CurrentVelocity, 
 	const FPoseSearchTrajectoryData::FSampling& TrajectoryDataSampling,
@@ -282,6 +298,47 @@ void UPoseSearchTrajectoryLibrary::UpdateHistory_TransformHistory(
 	}
 }
 
+
+void UPoseSearchTrajectoryLibrary::UpdateHistory_WorldSpace(
+	FPoseSearchQueryTrajectory& Trajectory,
+	FVector CurrentPosition,
+	const FPoseSearchTrajectoryData::FSampling& TrajectoryDataSampling,
+	float DeltaTime)
+{
+	const int32 NumHistorySamples = TrajectoryDataSampling.NumHistorySamples;
+	if (NumHistorySamples > 0)
+	{
+		const float SecondsPerHistorySample = TrajectoryDataSampling.SecondsPerHistorySample;
+
+		check(NumHistorySamples <= Trajectory.Samples.Num());
+
+		// Shift history Samples when it's time to record a new one.
+		if (SecondsPerHistorySample <= 0.f || FMath::Abs(Trajectory.Samples[NumHistorySamples - 1].AccumulatedSeconds) >= SecondsPerHistorySample)
+		{
+			for (int32 Index = 0; Index < NumHistorySamples - 1; ++Index)
+			{
+				Trajectory.Samples[Index].AccumulatedSeconds = Trajectory.Samples[Index + 1].AccumulatedSeconds - DeltaTime;
+				Trajectory.Samples[Index].Position = Trajectory.Samples[Index + 1].Position;
+				Trajectory.Samples[Index].Facing = Trajectory.Samples[Index + 1].Facing;
+			}
+
+			// Adding a new history record
+			// Note we add current frame at delta time in the future since t = 0 should be the previous frame
+			Trajectory.Samples[NumHistorySamples - 1].AccumulatedSeconds = DeltaTime;
+			Trajectory.Samples[NumHistorySamples - 1].Position = CurrentPosition;
+			Trajectory.Samples[NumHistorySamples - 1].Facing = Trajectory.Samples[NumHistorySamples].Facing;
+		}
+		else
+		{
+			// Didn't record a new history position, update timers and shift by ground translation
+			for (int32 Index = 0; Index < NumHistorySamples; ++Index)
+			{
+				Trajectory.Samples[Index].AccumulatedSeconds -= DeltaTime;
+			}
+		}
+	}
+}
+
 FVector UPoseSearchTrajectoryLibrary::RemapVectorMagnitudeWithCurve(
 	const FVector& Vector,
 	bool bUseCurve,
@@ -404,8 +461,8 @@ void UPoseSearchTrajectoryLibrary::PoseSearchGenerateTrajectory(
 
 	FPoseSearchTrajectoryData::FDerived TrajectoryDataDerived;
 	InTrajectoryData.UpdateData(InDeltaTime, Context, TrajectoryDataDerived, TrajectoryDataState);
-	InitTrajectorySamples(InOutTrajectory, InTrajectoryData, TrajectoryDataDerived.Position, TrajectoryDataDerived.Facing, TrajectoryDataSampling, InDeltaTime);
-	UpdateHistory_TransformHistory(InOutTrajectory, InTrajectoryData, TrajectoryDataDerived.Position, TrajectoryDataDerived.Velocity, TrajectoryDataSampling, InDeltaTime);
+	InitTrajectorySamples(InOutTrajectory, TrajectoryDataDerived.Position, TrajectoryDataDerived.Facing, TrajectoryDataSampling, InDeltaTime);
+	UpdateHistory_TransformHistory(InOutTrajectory, TrajectoryDataDerived.Position, TrajectoryDataDerived.Velocity, TrajectoryDataSampling, InDeltaTime);
 	UpdatePrediction_SimulateCharacterMovement(InOutTrajectory, InTrajectoryData, TrajectoryDataDerived, TrajectoryDataSampling, InDeltaTime);
 
 	InOutDesiredControllerYawLastUpdate = TrajectoryDataState.DesiredControllerYawLastUpdate;
@@ -414,6 +471,7 @@ void UPoseSearchTrajectoryLibrary::PoseSearchGenerateTrajectory(
 }
 
 
+// Deprecated
 void UPoseSearchTrajectoryLibrary::PoseSearchGeneratePredictorTrajectory(
 	UObject* InPredictor,	// must implement IPoseSearchTrajectoryPredictorInterface
 	UPARAM(ref)	const FPoseSearchTrajectoryData& InTrajectoryData,
@@ -425,6 +483,32 @@ void UPoseSearchTrajectoryLibrary::PoseSearchGeneratePredictorTrajectory(
 	int32 InTrajectoryHistoryCount,
 	float InPredictionSamplingInterval,
 	int32 InTrajectoryPredictionCount)
+{
+	IPoseSearchTrajectoryPredictorInterface* predictor = Cast<IPoseSearchTrajectoryPredictorInterface>(InPredictor);
+	if (predictor != nullptr)
+	{
+		PoseSearchGenerateTrajectoryWithPredictor(InPredictor,
+			InDeltaTime,
+			InOutTrajectory,
+			InOutDesiredControllerYawLastUpdate,
+			OutTrajectory,
+			InHistorySamplingInterval,
+			InTrajectoryHistoryCount,
+			InPredictionSamplingInterval,
+			InTrajectoryPredictionCount);
+	}
+}
+
+void UPoseSearchTrajectoryLibrary::PoseSearchGenerateTrajectoryWithPredictor(
+	TScriptInterface<IPoseSearchTrajectoryPredictorInterface> InPredictor,	// must implement IPoseSearchTrajectoryPredictorInterface
+	float InDeltaTime,
+	UPARAM(ref) FPoseSearchQueryTrajectory& InOutTrajectory,
+	UPARAM(ref) float& InOutDesiredControllerYawLastUpdate,
+	FPoseSearchQueryTrajectory& OutTrajectory,
+	float InHistorySamplingInterval,
+	int32 InTrajectoryHistoryCount,
+	float InPredictionSamplingInterval,
+	int32 InTrajectoryPredictionCount)
 {
 	FPoseSearchTrajectoryData::FSampling TrajectoryDataSampling;
 	TrajectoryDataSampling.NumHistorySamples = InTrajectoryHistoryCount;
@@ -439,19 +523,17 @@ void UPoseSearchTrajectoryLibrary::PoseSearchGeneratePredictorTrajectory(
 	FVector CurrentVelocity = FVector::ZeroVector;
 	FQuat CurrentFacing = FQuat::Identity;
 
-	IPoseSearchTrajectoryPredictorInterface* PredictorObject = Cast<IPoseSearchTrajectoryPredictorInterface>(InPredictor);
-
-	if (PredictorObject)
+	if (InPredictor != nullptr)
 	{
-		PredictorObject->GetCurrentState(CurrentPosition, CurrentFacing, CurrentVelocity);
+		InPredictor->GetCurrentState(CurrentPosition, CurrentFacing, CurrentVelocity);
 	}
 
-	InitTrajectorySamples(InOutTrajectory, InTrajectoryData, CurrentPosition, CurrentFacing, TrajectoryDataSampling, InDeltaTime);
-	UpdateHistory_TransformHistory(InOutTrajectory, InTrajectoryData, CurrentPosition, CurrentVelocity, TrajectoryDataSampling, InDeltaTime);
+	InitTrajectorySamples(InOutTrajectory, CurrentPosition, CurrentFacing, TrajectoryDataSampling, InDeltaTime);
+	UpdateHistory_TransformHistory(InOutTrajectory, CurrentPosition, CurrentVelocity, TrajectoryDataSampling, InDeltaTime);
 
-	if (PredictorObject)
+	if (InPredictor != nullptr)
 	{
-		PredictorObject->Predict(InOutTrajectory, InTrajectoryPredictionCount + 1, InPredictionSamplingInterval, InTrajectoryHistoryCount);
+		InPredictor->Predict(InOutTrajectory, InTrajectoryPredictionCount + 1, InPredictionSamplingInterval, InTrajectoryHistoryCount);
 	}
 
 	//InOutDesiredControllerYawLastUpdate = TrajectoryDataState.DesiredControllerYawLastUpdate;
diff --git a/Engine/Plugins/Animation/PoseSearch/Source/Runtime/Public/PoseSearch/PoseSearchTrajectoryLibrary.h b/Engine/Plugins/Animation/PoseSearch/Source/Runtime/Public/PoseSearch/PoseSearchTrajectoryLibrary.h
--- a/Engine/Plugins/Animation/PoseSearch/Source/Runtime/Public/PoseSearch/PoseSearchTrajectoryLibrary.h
+++ b/Engine/Plugins/Animation/PoseSearch/Source/Runtime/Public/PoseSearch/PoseSearchTrajectoryLibrary.h
@@ -109,15 +109,27 @@ public:
 	static void InitTrajectorySamples(FPoseSearchQueryTrajectory& Trajectory, const FPoseSearchTrajectoryData& TrajectoryData, const FPoseSearchTrajectoryData::FDerived& TrajectoryDataDerived, const FPoseSearchTrajectoryData::FSampling& TrajectoryDataSampling, float DeltaTime);
 
 	// Initialize history and predicted samples based on sampling settings and a default state
+	UE_DEPRECATED(5.6, "Use the InitTrajectorySamples function that doesn't take FPoseSearchTrajectoryData")
 	static void InitTrajectorySamples(FPoseSearchQueryTrajectory& Trajectory, const FPoseSearchTrajectoryData& TrajectoryData, FVector DefaultPosition, FQuat DefaultFacing, const FPoseSearchTrajectoryData::FSampling& TrajectoryDataSampling, float DeltaTime);
-	
+
+	// Initialize history and predicted samples based on sampling settings and a default state
+	static void InitTrajectorySamples(FPoseSearchQueryTrajectory& Trajectory, FVector DefaultPosition, FQuat DefaultFacing, const FPoseSearchTrajectoryData::FSampling& TrajectoryDataSampling, float DeltaTime);
+
 	UE_DEPRECATED(5.5, "Use the UpdateHistory_TransformHistory function that takes CurrentPosition and CurrentVelocity arguments instead of TrajectoryDataDerived")
 	static void UpdateHistory_TransformHistory(FPoseSearchQueryTrajectory& Trajectory, const FPoseSearchTrajectoryData& TrajectoryData, const FPoseSearchTrajectoryData::FDerived& TrajectoryDataDerived, const FPoseSearchTrajectoryData::FSampling& TrajectoryDataSampling, float DeltaTime);
 
+	UE_DEPRECATED(5.6, "Use the UpdateHistory_TransformHistory function that doesn't take FPoseSearchTrajectoryData")
+	static void UpdateHistory_TransformHistory(FPoseSearchQueryTrajectory& Trajectory, const FPoseSearchTrajectoryData& TrajectoryData, FVector CurrentPosition, FVector CurrentVelocity, const FPoseSearchTrajectoryData::FSampling& TrajectoryDataSampling, float DeltaTime);
+
 	// Update history by tracking offsets that result from character intent (e.g. movement component velocity) and applying
 	// that to the current world transform. This works well on moving platforms as it only stores a history of movement
 	// that results from character intent, not movement from platforms.
-	static void UpdateHistory_TransformHistory(FPoseSearchQueryTrajectory& Trajectory, const FPoseSearchTrajectoryData& TrajectoryData, FVector CurrentPosition, FVector CurrentVelocity, const FPoseSearchTrajectoryData::FSampling& TrajectoryDataSampling, float DeltaTime);
+	// Important: CurrentVelocity should be the velocity relative to the ground as reported by the character movement component or character mover etc
+	static void UpdateHistory_TransformHistory(FPoseSearchQueryTrajectory& Trajectory, FVector CurrentPosition, FVector CurrentVelocity, const FPoseSearchTrajectoryData::FSampling& TrajectoryDataSampling, float DeltaTime);
+
+	// Update the history purely based on current position, without taking into account ground velocities
+	// DeltaTime is the time since the last call to this method
+	static void UpdateHistory_WorldSpace(FPoseSearchQueryTrajectory& Trajectory, FVector CurrentPosition, const FPoseSearchTrajectoryData::FSampling& TrajectoryDataSampling, float DeltaTime);
 
 	// Update prediction by simulating the movement math for ground locomotion from UCharacterMovementComponent.
 	static void UpdatePrediction_SimulateCharacterMovement(FPoseSearchQueryTrajectory& Trajectory, const FPoseSearchTrajectoryData& TrajectoryData, const FPoseSearchTrajectoryData::FDerived& TrajectoryDataDerived, const FPoseSearchTrajectoryData::FSampling& TrajectoryDataSampling, float DeltaTime);
@@ -131,12 +143,20 @@ public:
 		float InHistorySamplingInterval = 0.04f, int32 InTrajectoryHistoryCount = 10, float InPredictionSamplingInterval = 0.2f, int32 InTrajectoryPredictionCount = 8);
 
 	// Generates a prediction trajectory based of the current movement intent. For use with predictors. InPredictor must implement IPoseSearchTrajectoryPredictorInterface
-	UFUNCTION(BlueprintCallable, Category = "Animation|PoseSearch|Experimental", meta = (BlueprintThreadSafe, DisplayName = "Pose Search Generate Trajectory (using Predictor)"))
+	UE_DEPRECATED(5.6, "Use PoseSearchGenerateTrajectoryWithPredictor function that doesn't take FPoseSearchTrajectoryData")
+	UFUNCTION(BlueprintCallable, Category = "Animation|PoseSearch|Experimental", meta = (BlueprintThreadSafe, DisplayName = "Pose Search Generate Trajectory (using Predictor) (DEPRECATED)"))
 	static void PoseSearchGeneratePredictorTrajectory(UObject* InPredictor,
 		UPARAM(ref) const FPoseSearchTrajectoryData& InTrajectoryData, float InDeltaTime,
 		UPARAM(ref) FPoseSearchQueryTrajectory& InOutTrajectory, UPARAM(ref) float& InOutDesiredControllerYawLastUpdate, FPoseSearchQueryTrajectory& OutTrajectory,
 		float InHistorySamplingInterval = 0.04f, int32 InTrajectoryHistoryCount = 10, float InPredictionSamplingInterval = 0.2f, int32 InTrajectoryPredictionCount = 8);
 
+	// Generates a prediction trajectory based of the current movement intent. For use with predictors. InPredictor must implement IPoseSearchTrajectoryPredictorInterface
+	UFUNCTION(BlueprintCallable, Category = "Animation|PoseSearch|Experimental", meta = (BlueprintThreadSafe, DisplayName = "Pose Search Generate Trajectory (using Predictor)"))
+	static void PoseSearchGenerateTrajectoryWithPredictor(TScriptInterface<IPoseSearchTrajectoryPredictorInterface> InPredictor,
+		float InDeltaTime,
+		UPARAM(ref) FPoseSearchQueryTrajectory& InOutTrajectory, UPARAM(ref) float& InOutDesiredControllerYawLastUpdate, FPoseSearchQueryTrajectory& OutTrajectory,
+		float InHistorySamplingInterval = 0.04f, int32 InTrajectoryHistoryCount = 10, float InPredictionSamplingInterval = 0.2f, int32 InTrajectoryPredictionCount = 8);
+
 	// Experimental: Process InTrajectory to apply gravity and handle collisions. Eventually returns the modified OutTrajectory.
 	// If bApplyGravity is true, gravity from the UCharacterMovementComponent will be applied.
 	// If FloorCollisionsOffset > 0, vertical collision will be performed to every sample of the trajectory to have the samples float over the geometry (by FloorCollisionsOffset).
diff --git a/Engine/Plugins/Experimental/MoverAnimNext/Source/MoverAnimNext/Private/Graph/RigUnit_GenerateMoverTrajectory.cpp b/Engine/Plugins/Experimental/MoverAnimNext/Source/MoverAnimNext/Private/Graph/RigUnit_GenerateMoverTrajectory.cpp
--- a/Engine/Plugins/Experimental/MoverAnimNext/Source/MoverAnimNext/Private/Graph/RigUnit_GenerateMoverTrajectory.cpp
+++ b/Engine/Plugins/Experimental/MoverAnimNext/Source/MoverAnimNext/Private/Graph/RigUnit_GenerateMoverTrajectory.cpp
@@ -8,6 +8,7 @@
 
 FRigUnit_GenerateMoverTrajectory_Execute()
 {
+	TRACE_CPUPROFILER_EVENT_SCOPE(FRigUnit_GenerateMoverTrajectory_Execute);
 	using namespace UE::AnimNext;
 
 	const FAnimNextModuleContextData& ModuleContextData = ExecuteContext.GetContextData<FAnimNextModuleContextData>();
@@ -40,8 +41,8 @@ FRigUnit_GenerateMoverTrajectory_Execute()
 	}
 
 	UMoverTrajectoryPredictor::GetCurrentState(*MoverComponent, CurrentPosition, CurrentFacing, CurrentVelocity);
-	UPoseSearchTrajectoryLibrary::InitTrajectorySamples(InOutTrajectory, TrajectoryData, CurrentPosition, CurrentFacing, TrajectoryDataSampling, DeltaTime);
-	UPoseSearchTrajectoryLibrary::UpdateHistory_TransformHistory(InOutTrajectory, TrajectoryData, CurrentPosition, CurrentVelocity, TrajectoryDataSampling, DeltaTime);
+	UPoseSearchTrajectoryLibrary::InitTrajectorySamples(InOutTrajectory, CurrentPosition, CurrentFacing, TrajectoryDataSampling, DeltaTime);
+	UPoseSearchTrajectoryLibrary::UpdateHistory_TransformHistory(InOutTrajectory, CurrentPosition, CurrentVelocity, TrajectoryDataSampling, DeltaTime);
 	UMoverTrajectoryPredictor::Predict(*MoverComponent, InOutTrajectory,
 		NumPredictionSamples + 1, PredictionSamplingInterval, NumHistorySamples);
 }
diff --git a/Engine/Plugins/Experimental/MoverAnimNext/Source/MoverAnimNext/Private/Graph/RigUnit_GenerateMoverTrajectory.h b/Engine/Plugins/Experimental/MoverAnimNext/Source/MoverAnimNext/Private/Graph/RigUnit_GenerateMoverTrajectory.h
--- a/Engine/Plugins/Experimental/MoverAnimNext/Source/MoverAnimNext/Private/Graph/RigUnit_GenerateMoverTrajectory.h
+++ b/Engine/Plugins/Experimental/MoverAnimNext/Source/MoverAnimNext/Private/Graph/RigUnit_GenerateMoverTrajectory.h
@@ -21,9 +21,6 @@ struct FRigUnit_GenerateMoverTrajectory : public FRigUnit_AnimNextBase
 	UPROPERTY(EditAnywhere, Category = "Graph", meta = (Input))
 	TObjectPtr<UMoverComponent> MoverComponent;
 
-	UPROPERTY(EditAnywhere, Category = "Graph", meta = (Input))
-	FPoseSearchTrajectoryData TrajectoryData;
-
 	UPROPERTY(EditAnywhere, Category = "Graph", meta = (Input))
 	// This should be the most recent simulation time that was used to get us to our current state
 	float DeltaTime;
