posesearch - fix for BuildQuery UE_LOG message

diff --git a/Engine/Plugins/Animation/PoseSearch/Source/Runtime/Private/PoseSearchFeatureChannel_Heading.cpp b/Engine/Plugins/Animation/PoseSearch/Source/Runtime/Private/PoseSearchFeatureChannel_Heading.cpp
--- a/Engine/Plugins/Animation/PoseSearch/Source/Runtime/Private/PoseSearchFeatureChannel_Heading.cpp
+++ b/Engine/Plugins/Animation/PoseSearch/Source/Runtime/Private/PoseSearchFeatureChannel_Heading.cpp
@@ -116,7 +116,7 @@ void UPoseSearchFeatureChannel_Heading::BuildQuery(UE::PoseSearch::FSearchContex
 		else
 		{
 			// @todo: support non UAnimInstance anim contexts for AnimNext
-			UE_LOG(LogPoseSearch, Warning, TEXT("UPoseSearchFeatureChannel_Curve::BuildQuery - unsupported null UAnimInstance: WIP support for AnimNext!"));
+			UE_LOG(LogPoseSearch, Warning, TEXT("UPoseSearchFeatureChannel_Heading::BuildQuery - unsupported null UAnimInstance: WIP support for AnimNext!"));
 		}
 		return;
 	}
diff --git a/Engine/Plugins/Animation/PoseSearch/Source/Runtime/Private/PoseSearchFeatureChannel_Position.cpp b/Engine/Plugins/Animation/PoseSearch/Source/Runtime/Private/PoseSearchFeatureChannel_Position.cpp
--- a/Engine/Plugins/Animation/PoseSearch/Source/Runtime/Private/PoseSearchFeatureChannel_Position.cpp
+++ b/Engine/Plugins/Animation/PoseSearch/Source/Runtime/Private/PoseSearchFeatureChannel_Position.cpp
@@ -118,7 +118,7 @@ void UPoseSearchFeatureChannel_Position::BuildQuery(UE::PoseSearch::FSearchConte
 		else
 		{
 			// @todo: support non UAnimInstance anim contexts for AnimNext
-			UE_LOG(LogPoseSearch, Warning, TEXT("UPoseSearchFeatureChannel_Curve::BuildQuery - unsupported null UAnimInstance: WIP support for AnimNext!"));
+			UE_LOG(LogPoseSearch, Warning, TEXT("UPoseSearchFeatureChannel_Position::BuildQuery - unsupported null UAnimInstance: WIP support for AnimNext!"));
 		}
 		return;
 	}
diff --git a/Engine/Plugins/Animation/PoseSearch/Source/Runtime/Private/PoseSearchFeatureChannel_Velocity.cpp b/Engine/Plugins/Animation/PoseSearch/Source/Runtime/Private/PoseSearchFeatureChannel_Velocity.cpp
--- a/Engine/Plugins/Animation/PoseSearch/Source/Runtime/Private/PoseSearchFeatureChannel_Velocity.cpp
+++ b/Engine/Plugins/Animation/PoseSearch/Source/Runtime/Private/PoseSearchFeatureChannel_Velocity.cpp
@@ -62,7 +62,7 @@ void UPoseSearchFeatureChannel_Velocity::BuildQuery(UE::PoseSearch::FSearchConte
 		else
 		{
 			// @todo: support non UAnimInstance anim contexts for AnimNext
-			UE_LOG(LogPoseSearch, Warning, TEXT("UPoseSearchFeatureChannel_Curve::BuildQuery - unsupported null UAnimInstance: WIP support for AnimNext!"));
+			UE_LOG(LogPoseSearch, Warning, TEXT("UPoseSearchFeatureChannel_Velocity::BuildQuery - unsupported null UAnimInstance: WIP support for AnimNext!"));
 		}
 		return;
 	}
