posesearch
- fix for warping with optional interaction characters
- fix for UPoseSearchInteractionSubsystem::GenerateSearchContexts narrow phase

[RN] skip

diff --git a/Engine/Plugins/Animation/PoseSearch/Source/Runtime/Private/AnimNode_MotionMatchingInteraction.cpp b/Engine/Plugins/Animation/PoseSearch/Source/Runtime/Private/AnimNode_MotionMatchingInteraction.cpp
--- a/Engine/Plugins/Animation/PoseSearch/Source/Runtime/Private/AnimNode_MotionMatchingInteraction.cpp
+++ b/Engine/Plugins/Animation/PoseSearch/Source/Runtime/Private/AnimNode_MotionMatchingInteraction.cpp
@@ -33,6 +33,37 @@ static float DamperImplicit(float Halflife, float DeltaTime, float Epsilon = 1e-
 	return FMath::Clamp((1.0f - FMath::InvExpApprox((0.69314718056f * DeltaTime) / (Halflife + Epsilon))), 0.0f, 1.0f);
 }
 
+static FRoleToIndex MakeRoleToIndex(const UMultiAnimAsset* MultiAnimAsset)
+{
+	check(MultiAnimAsset);
+	const int32 NumRoles = MultiAnimAsset->GetNumRoles();
+
+	FRoleToIndex RoleToIndex;
+	RoleToIndex.Reserve(NumRoles);
+	for (int32 MultiAnimAssetRoleIndex = 0; MultiAnimAssetRoleIndex < NumRoles; ++MultiAnimAssetRoleIndex)
+	{
+		RoleToIndex.Add(MultiAnimAsset->GetRole(MultiAnimAssetRoleIndex)) = MultiAnimAssetRoleIndex;
+	}
+	return RoleToIndex;
+}
+
+typedef TArray<bool, TInlineAllocator<PreallocatedRolesNum, TMemStackAllocator<>>> FRelevantRoleIndexes;
+static FRelevantRoleIndexes GetRelevantRoleIndexes(const UMultiAnimAsset* MultiAnimAsset, const UPoseSearchSchema* Schema)
+{
+	check(Schema && MultiAnimAsset);
+
+	const FRoleToIndex RoleToIndex = MakeRoleToIndex(MultiAnimAsset);
+	const int32 NumRoles = MultiAnimAsset->GetNumRoles();
+
+	FRelevantRoleIndexes RelevantRoleIndexes;
+	RelevantRoleIndexes.SetNum(NumRoles);
+	for (const FPoseSearchRoledSkeleton& RoledSkeleton : Schema->GetRoledSkeletons())
+	{
+		RelevantRoleIndexes[RoleToIndex[RoledSkeleton.Role]] = true;
+	}
+	return RelevantRoleIndexes;
+}
+
 } // UE::PoseSearch
 
 void FAnimNode_MotionMatchingInteraction::Reset()
@@ -256,11 +287,14 @@ void FAnimNode_MotionMatchingInteraction::Evaluate_AnyThread(FPoseContext& Outpu
 	// NoTe that ActorRootTransforms[CurrentResultRoleIndex] is Output.AnimInstanceProxy->GetSkelMeshComponent()->GetComponentTransform();
 	const FTransform AlignmentTransformNoWarpAdjustment = RootMotionTransform * ActorRootTransforms[CurrentResultRoleIndex];
 					
+	check(CurrentResult.SelectedDatabase != nullptr && CurrentResult.SelectedDatabase->Schema != nullptr);
+	const FRelevantRoleIndexes RelevantRoleIndexes = GetRelevantRoleIndexes(MultiAnimAsset, CurrentResult.SelectedDatabase->Schema);
+
 	// calculating the world space warp tranform to have the character fully aligned to what MultiAnimAsset wants.
 	// using ActorTransforms instead of ActorRootTransforms, since we want the previous frame root bone offset rotation to be embedded into the calculation
 	TArray<FTransform, TInlineAllocator<PreallocatedRolesNum, TMemStackAllocator<>>> FullAlignedActorRootBoneTransformsAtExtractionTime;
 	FullAlignedActorRootBoneTransformsAtExtractionTime.SetNum(NumRoles);
-	MultiAnimAsset->CalculateWarpTransforms(WarpingAlignmentExtractionTime, ActorTransforms, FullAlignedActorRootBoneTransformsAtExtractionTime);
+	MultiAnimAsset->CalculateWarpTransforms(WarpingAlignmentExtractionTime, ActorTransforms, FullAlignedActorRootBoneTransformsAtExtractionTime, RelevantRoleIndexes);
 	const FTransform AlignmentTransformFullWarpAdjustment = FullAlignedActorRootBoneTransformsAtExtractionTime[CurrentResultRoleIndex];
 					
 	// @todo: compensating AlignmentTransformFullWarpAdjustment by the delta time between AccumulatedTime and WarpingAlignmentExtractionTime?				
@@ -314,7 +348,7 @@ void FAnimNode_MotionMatchingInteraction::Evaluate_AnyThread(FPoseContext& Outpu
 		// using ActorTransforms instead of ActorRootTransforms, since we want the previous frame root bone offset rotation to be embedded into the calculation
 		TArray<FTransform, TInlineAllocator<PreallocatedRolesNum, TMemStackAllocator<>>> FullAlignedActorRootBoneTransformsAtAccumulatedTime;
 		FullAlignedActorRootBoneTransformsAtAccumulatedTime.SetNum(NumRoles);
-		MultiAnimAsset->CalculateWarpTransforms(AccumulatedTime, ActorTransforms, FullAlignedActorRootBoneTransformsAtAccumulatedTime);
+		MultiAnimAsset->CalculateWarpTransforms(AccumulatedTime, ActorTransforms, FullAlignedActorRootBoneTransformsAtAccumulatedTime, RelevantRoleIndexes);
 		FullAlignedPoseRootTransform = FullAlignedActorRootBoneTransformsAtAccumulatedTime[CurrentResultRoleIndex];
 	}
 								
diff --git a/Engine/Plugins/Animation/PoseSearch/Source/Runtime/Private/PoseSearchInteractionAsset.cpp b/Engine/Plugins/Animation/PoseSearch/Source/Runtime/Private/PoseSearchInteractionAsset.cpp
--- a/Engine/Plugins/Animation/PoseSearch/Source/Runtime/Private/PoseSearchInteractionAsset.cpp
+++ b/Engine/Plugins/Animation/PoseSearch/Source/Runtime/Private/PoseSearchInteractionAsset.cpp
@@ -272,15 +272,34 @@ USkeletalMesh* UPoseSearchInteractionAsset::GetPreviewMesh(const UE::PoseSearch:
 
 #endif // WITH_EDITOR
 
-void UPoseSearchInteractionAsset::CalculateWarpTransforms(float Time, const TArrayView<const FTransform> ActorRootBoneTransforms, TArrayView<FTransform> FullAlignedActorRootBoneTransforms) const
+void UPoseSearchInteractionAsset::CalculateWarpTransforms(float Time, const TConstArrayView<const FTransform> ActorRootBoneTransforms, TArrayView<FTransform> FullAlignedActorRootBoneTransforms, const TConstArrayView<bool> RelevantRoleIndexes) const
 {
 	using namespace UE::PoseSearch;
 
 	check(ActorRootBoneTransforms.Num() == GetNumRoles());
 	check(FullAlignedActorRootBoneTransforms.Num() == GetNumRoles());
-
+	check(RelevantRoleIndexes.IsEmpty() || RelevantRoleIndexes.Num() == GetNumRoles());
+	
 	const int32 ItemsNum = Items.Num();
-	if (ItemsNum < 2)
+
+	int32 RelevantItemsNum;
+	if (RelevantRoleIndexes.IsEmpty())
+	{
+		RelevantItemsNum = ItemsNum;
+	}
+	else
+	{
+		RelevantItemsNum = 0;
+		for (bool bIsRelevant : RelevantRoleIndexes)
+		{
+			if (bIsRelevant)
+			{
+				++RelevantItemsNum;
+			}
+		}
+	}
+
+	if (RelevantItemsNum < 2)
 	{
 		for (int32 ItemIndex = 0; ItemIndex < ItemsNum; ++ItemIndex)
 		{
@@ -289,6 +308,8 @@ void UPoseSearchInteractionAsset::CalculateWarpTransforms(float Time, const TArr
 	}
 	else
 	{
+		// we have at least one relevant item!
+
 		TArray<FTransform, TInlineAllocator<PreallocatedRolesNum>> AssetRootBoneTransforms;
 		AssetRootBoneTransforms.SetNum(ItemsNum);
 
@@ -337,46 +358,130 @@ void UPoseSearchInteractionAsset::CalculateWarpTransforms(float Time, const TArr
 
 		float WarpingWeightTranslationSum = 0.f;
 		float WarpingWeightRotationSum = 0.f;
-		for (int32 ItemIndex = 0; ItemIndex < ItemsNum; ++ItemIndex)
-		{
-			SortedByWarpingWeightRotationItemIndex[ItemIndex] = ItemIndex;
-			WarpingWeightTranslationSum += Items[ItemIndex].WarpingWeightTranslation;
-			WarpingWeightRotationSum += Items[ItemIndex].WarpingWeightRotation;
-		}
 
-		const float NormalizedHomogeneousWeight = 1.f / ItemsNum;
-		if (WarpingWeightTranslationSum > UE_KINDA_SMALL_NUMBER)
+		if (RelevantRoleIndexes.IsEmpty())
 		{
 			for (int32 ItemIndex = 0; ItemIndex < ItemsNum; ++ItemIndex)
 			{
-				NormalizedWarpingWeightTranslation[ItemIndex] = Items[ItemIndex].WarpingWeightTranslation / WarpingWeightTranslationSum;
+				SortedByWarpingWeightRotationItemIndex[ItemIndex] = ItemIndex;
+				WarpingWeightTranslationSum += Items[ItemIndex].WarpingWeightTranslation;
+				WarpingWeightRotationSum += Items[ItemIndex].WarpingWeightRotation;
 			}
 		}
 		else
 		{
 			for (int32 ItemIndex = 0; ItemIndex < ItemsNum; ++ItemIndex)
 			{
-				NormalizedWarpingWeightTranslation[ItemIndex] = NormalizedHomogeneousWeight;
+				SortedByWarpingWeightRotationItemIndex[ItemIndex] = ItemIndex;
+				if (RelevantRoleIndexes[ItemIndex])
+				{
+					WarpingWeightTranslationSum += Items[ItemIndex].WarpingWeightTranslation;
+					WarpingWeightRotationSum += Items[ItemIndex].WarpingWeightRotation;
+				}
 			}
 		}
 
-		if (WarpingWeightRotationSum > UE_KINDA_SMALL_NUMBER)
+		const float NormalizedHomogeneousWeight = 1.f / RelevantItemsNum;
+		if (RelevantRoleIndexes.IsEmpty())
 		{
-			SortedByWarpingWeightRotationItemIndex.Sort([this](const int32 A, const int32 B)
+			if (WarpingWeightTranslationSum > UE_KINDA_SMALL_NUMBER)
+			{
+				for (int32 ItemIndex = 0; ItemIndex < ItemsNum; ++ItemIndex)
 				{
-					return Items[A].WarpingWeightRotation < Items[B].WarpingWeightRotation;
-				});
+					NormalizedWarpingWeightTranslation[ItemIndex] = Items[ItemIndex].WarpingWeightTranslation / WarpingWeightTranslationSum;
+				}
+			}
+			else
+			{
+				for (int32 ItemIndex = 0; ItemIndex < ItemsNum; ++ItemIndex)
+				{
+					NormalizedWarpingWeightTranslation[ItemIndex] = NormalizedHomogeneousWeight;
+				}
+			}
 
-			for (int32 ItemIndex = 0; ItemIndex < ItemsNum; ++ItemIndex)
+			if (WarpingWeightRotationSum > UE_KINDA_SMALL_NUMBER)
 			{
-				NormalizedWarpingWeightRotation[ItemIndex] = Items[ItemIndex].WarpingWeightRotation / WarpingWeightRotationSum;
+				SortedByWarpingWeightRotationItemIndex.Sort([this](const int32 A, const int32 B)
+					{
+						return Items[A].WarpingWeightRotation < Items[B].WarpingWeightRotation;
+					});
+
+				for (int32 ItemIndex = 0; ItemIndex < ItemsNum; ++ItemIndex)
+				{
+					NormalizedWarpingWeightRotation[ItemIndex] = Items[ItemIndex].WarpingWeightRotation / WarpingWeightRotationSum;
+				}
+			}
+			else
+			{
+				for (int32 ItemIndex = 0; ItemIndex < ItemsNum; ++ItemIndex)
+				{
+					NormalizedWarpingWeightRotation[ItemIndex] = NormalizedHomogeneousWeight;
+				}
 			}
 		}
 		else
 		{
-			for (int32 ItemIndex = 0; ItemIndex < ItemsNum; ++ItemIndex)
+			if (WarpingWeightTranslationSum > UE_KINDA_SMALL_NUMBER)
 			{
-				NormalizedWarpingWeightRotation[ItemIndex] = NormalizedHomogeneousWeight;
+				for (int32 ItemIndex = 0; ItemIndex < ItemsNum; ++ItemIndex)
+				{
+					if (RelevantRoleIndexes[ItemIndex])
+					{
+						NormalizedWarpingWeightTranslation[ItemIndex] = Items[ItemIndex].WarpingWeightTranslation / WarpingWeightTranslationSum;
+					}
+					else
+					{
+						NormalizedWarpingWeightTranslation[ItemIndex] = 0.f;
+					}
+				}
+			}
+			else
+			{
+				for (int32 ItemIndex = 0; ItemIndex < ItemsNum; ++ItemIndex)
+				{
+					if (RelevantRoleIndexes[ItemIndex])
+					{
+						NormalizedWarpingWeightTranslation[ItemIndex] = NormalizedHomogeneousWeight;
+					}
+					else
+					{
+						NormalizedWarpingWeightTranslation[ItemIndex] = 0.f;
+					}
+				}
+			}
+
+			if (WarpingWeightRotationSum > UE_KINDA_SMALL_NUMBER)
+			{
+				SortedByWarpingWeightRotationItemIndex.Sort([this](const int32 A, const int32 B)
+					{
+						return Items[A].WarpingWeightRotation < Items[B].WarpingWeightRotation;
+					});
+
+				for (int32 ItemIndex = 0; ItemIndex < ItemsNum; ++ItemIndex)
+				{
+					if (RelevantRoleIndexes[ItemIndex])
+					{
+						NormalizedWarpingWeightRotation[ItemIndex] = Items[ItemIndex].WarpingWeightRotation / WarpingWeightRotationSum;
+					}
+					else
+					{
+						NormalizedWarpingWeightRotation[ItemIndex] = 0.f;
+					}
+				}
+			}
+			else
+			{
+				for (int32 ItemIndex = 0; ItemIndex < ItemsNum; ++ItemIndex)
+				{
+					if (RelevantRoleIndexes[ItemIndex])
+					{
+						NormalizedWarpingWeightRotation[ItemIndex] = NormalizedHomogeneousWeight;
+					}
+					else
+					{
+						NormalizedWarpingWeightRotation[ItemIndex] = 0.f;
+					}
+				}
 			}
 		}
 
diff --git a/Engine/Plugins/Animation/PoseSearch/Source/Runtime/Private/PoseSearchInteractionSubsystem.cpp b/Engine/Plugins/Animation/PoseSearch/Source/Runtime/Private/PoseSearchInteractionSubsystem.cpp
--- a/Engine/Plugins/Animation/PoseSearch/Source/Runtime/Private/PoseSearchInteractionSubsystem.cpp
+++ b/Engine/Plugins/Animation/PoseSearch/Source/Runtime/Private/PoseSearchInteractionSubsystem.cpp
@@ -651,7 +651,7 @@ void UPoseSearchInteractionSubsystem::GenerateSearchContexts(UE::PoseSearch::FIn
 		const FAnimContextInfo* AnimContextInfo = nullptr;
 		FRole Role = DefaultRole;
 		const FAnimNode_PoseSearchHistoryCollector_Base* HistoryCollector = nullptr;
-		const UPoseSearchDatabase* Database;
+		const UPoseSearchDatabase* Database = nullptr;
 		float MaxDistanceSquared = 0.f;
 	};
 	struct FRoledAnimContextInfos : TArray<FRoledAnimContextInfo, TInlineAllocator<PreallocatedRolesNum, TMemStackAllocator<>>>
@@ -874,7 +874,7 @@ void UPoseSearchInteractionSubsystem::GenerateSearchContexts(UE::PoseSearch::FIn
 								const FRoledAnimContextInfo& OtherRoledAnimContextInfo = RoledAnimContextInfos[OtherRoledAnimContextIndex];
 
 								const float DistanceSquared = RoledAnimContextInfo.AnimContextInfo->GetCachedDistanceSquaredFrom(*OtherRoledAnimContextInfo.AnimContextInfo);
-								const float MinDistanceSquared = FMath::Min(RoledAnimContextInfo.MaxDistanceSquared, RoledAnimContextInfo.MaxDistanceSquared);
+								const float MinDistanceSquared = FMath::Min(RoledAnimContextInfo.MaxDistanceSquared, OtherRoledAnimContextInfo.MaxDistanceSquared);
 								if (DistanceSquared > MinDistanceSquared)
 								{
 									return false;
diff --git a/Engine/Plugins/Animation/PoseSearch/Source/Runtime/Public/PoseSearch/MultiAnimAsset.h b/Engine/Plugins/Animation/PoseSearch/Source/Runtime/Public/PoseSearch/MultiAnimAsset.h
--- a/Engine/Plugins/Animation/PoseSearch/Source/Runtime/Public/PoseSearch/MultiAnimAsset.h
+++ b/Engine/Plugins/Animation/PoseSearch/Source/Runtime/Public/PoseSearch/MultiAnimAsset.h
@@ -30,7 +30,7 @@ public:
 	[[nodiscard]] virtual UAnimationAsset* GetAnimationAsset(const FName& Role) const PURE_VIRTUAL(UMultiAnimAsset::GetAnimationAsset, return nullptr;);
 	[[nodiscard]] virtual FTransform GetOrigin(const FName& Role) const PURE_VIRTUAL(UMultiAnimAsset::GetOrigin, return FTransform::Identity;);
 
-	virtual void CalculateWarpTransforms(float Time, const TArrayView<const FTransform> ActorRootBoneTransforms, TArrayView<FTransform> FullAlignedActorRootBoneTransforms) const PURE_VIRTUAL(UMultiAnimAsset::CalculateWarpTransforms, );
+	virtual void CalculateWarpTransforms(float Time, const TConstArrayView<const FTransform> ActorRootBoneTransforms, TArrayView<FTransform> FullAlignedActorRootBoneTransforms, const TConstArrayView<bool> RelevantRoleIndexes = TArrayView<bool>()) const PURE_VIRTUAL(UMultiAnimAsset::CalculateWarpTransforms, );
 
 	UFUNCTION(BlueprintPure, Category = "Animation", meta=(BlueprintThreadSafe, DisplayName = "Get Animation Asset"))
 	UAnimationAsset* BP_GetAnimationAsset(const FName& Role) const { return GetAnimationAsset(Role); }
diff --git a/Engine/Plugins/Animation/PoseSearch/Source/Runtime/Public/PoseSearch/PoseSearchInteractionAsset.h b/Engine/Plugins/Animation/PoseSearch/Source/Runtime/Public/PoseSearch/PoseSearchInteractionAsset.h
--- a/Engine/Plugins/Animation/PoseSearch/Source/Runtime/Public/PoseSearch/PoseSearchInteractionAsset.h
+++ b/Engine/Plugins/Animation/PoseSearch/Source/Runtime/Public/PoseSearch/PoseSearchInteractionAsset.h
@@ -92,7 +92,7 @@ public:
 	virtual USkeletalMesh* GetPreviewMesh(const UE::PoseSearch::FRole& Role) const override;
 #endif // WITH_EDITOR
 
-	virtual void CalculateWarpTransforms(float Time, const TArrayView<const FTransform> ActorRootBoneTransforms, TArrayView<FTransform> FullAlignedActorRootBoneTransforms) const override;
+	virtual void CalculateWarpTransforms(float Time, const TConstArrayView<const FTransform> ActorRootBoneTransforms, TArrayView<FTransform> FullAlignedActorRootBoneTransforms, const TConstArrayView<bool> RelevantRoleIndexes = TConstArrayView<bool>()) const override;
 	
 	FVector FindReferencePosition(const TArrayView<const FTransform> Transforms, const TArrayView<float> NormalizedWarpingWeightTranslation) const;
 	FQuat FindReferenceOrientation(const TArrayView<const FTransform> Transforms, const TArrayView<int32> SortedByWarpingWeightRotationItemIndex, const TArrayView<float> NormalizedWarpingWeightRotation) const;
