posesearch - adding new experimental distance channel

[RN] skip

diff --git a/Engine/Plugins/Animation/PoseSearch/Source/Runtime/Private/PoseSearchFeatureChannel_Distance.cpp b/Engine/Plugins/Animation/PoseSearch/Source/Runtime/Private/PoseSearchFeatureChannel_Distance.cpp
new file mode 100644
--- /dev/null
+++ b/Engine/Plugins/Animation/PoseSearch/Source/Runtime/Private/PoseSearchFeatureChannel_Distance.cpp
@@ -0,0 +1,290 @@
+// Copyright Epic Games, Inc. All Rights Reserved.
+
+#include "PoseSearch/PoseSearchFeatureChannel_Distance.h"
+#include "Engine/BlueprintGeneratedClass.h"
+#include "PoseSearch/PoseSearchAssetIndexer.h"
+#include "PoseSearch/PoseSearchAssetSampler.h"
+#include "PoseSearch/PoseSearchContext.h"
+#include "PoseSearch/PoseSearchDatabase.h"
+#include "PoseSearch/PoseSearchFeatureChannel_Position.h"
+#include "PoseSearch/PoseSearchSchema.h"
+
+#if WITH_EDITOR
+#include "PropertyHandle.h"
+#endif // WITH_EDITOR
+
+UPoseSearchFeatureChannel_Distance::UPoseSearchFeatureChannel_Distance()
+{
+	bUseBlueprintQueryOverride = Cast<UBlueprintGeneratedClass>(GetClass()) != nullptr;
+}
+
+bool UPoseSearchFeatureChannel_Distance::Finalize(UPoseSearchSchema* Schema)
+{
+	using namespace UE::PoseSearch;
+
+	ChannelDataOffset = Schema->SchemaCardinality;
+	ChannelCardinality = 1;
+	Schema->SchemaCardinality += ChannelCardinality;
+
+	SchemaBoneIdx = Schema->AddBoneReference(Bone, SampleRole, bDefaultWithRootBone);
+	SchemaOriginBoneIdx = Schema->AddBoneReference(OriginBone, OriginRole, bDefaultWithRootBone);
+	
+	return SchemaBoneIdx != InvalidSchemaBoneIdx && SchemaOriginBoneIdx != InvalidSchemaBoneIdx;
+}
+
+void UPoseSearchFeatureChannel_Distance::AddDependentChannels(UPoseSearchSchema* Schema) const
+{
+	using namespace UE::PoseSearch;
+
+	if (Schema->bInjectAdditionalDebugChannels)
+	{
+		if (SchemaOriginBoneIdx != RootSchemaBoneIdx || PermutationTimeType == EPermutationTimeType::UsePermutationTime)
+		{
+			if (bDefaultWithRootBone)
+			{
+				const EPermutationTimeType DependentChannelsPermutationTimeType = PermutationTimeType == EPermutationTimeType::UsePermutationTime ? EPermutationTimeType::UseSampleToPermutationTime : EPermutationTimeType::UseSampleTime;
+				UPoseSearchFeatureChannel_Position::FindOrAddToSchema(Schema, 0.f, OriginBone.BoneName, OriginRole, DependentChannelsPermutationTimeType);
+			}
+			else
+			{
+				// @todo: add bInjectAdditionalDebugChannels support for !bDefaultWithRootBone UPoseSearchFeatureChannel_Distance
+			}
+		}
+	}
+}
+
+void UPoseSearchFeatureChannel_Distance::BuildQuery(UE::PoseSearch::FSearchContext& SearchContext) const
+{
+	using namespace UE::PoseSearch;
+
+	if (bUseBlueprintQueryOverride)
+	{
+		if (const UAnimInstance* AnimInstance = Cast<UAnimInstance>(SearchContext.GetAnimContext(SampleRole)))
+		{
+			const float Distance = BP_GetDistance(AnimInstance);
+  			FFeatureVectorHelper::EncodeFloat(SearchContext.EditFeatureVector(), ChannelDataOffset, Distance);
+		}
+		else
+		{
+			// @todo: support non UAnimInstance anim contexts for AnimNext
+			UE_LOG(LogPoseSearch, Warning, TEXT("UPoseSearchFeatureChannel_Distance::BuildQuery - unsupported null UAnimInstance: WIP support for AnimNext!"));
+		}
+		return;
+	}
+
+	// trying to get the BuildQuery data from another schema UPoseSearchFeatureChannel_Distance already cached in the SearchContext
+	if (SearchContext.IsUseCachedChannelData())
+	{
+		// composing a unique identifier to specify this channel with all the required properties to be able to share the query data with other channels of the same type
+		uint32 UniqueIdentifier = GetClass()->GetUniqueID();
+		UniqueIdentifier = HashCombineFast(UniqueIdentifier, GetTypeHash(SampleRole));
+		UniqueIdentifier = HashCombineFast(UniqueIdentifier, GetTypeHash(OriginRole));
+		UniqueIdentifier = HashCombineFast(UniqueIdentifier, GetTypeHash(SamplingAttributeId));
+		UniqueIdentifier = HashCombineFast(UniqueIdentifier, GetTypeHash(SampleTimeOffset));
+		UniqueIdentifier = HashCombineFast(UniqueIdentifier, GetTypeHash(OriginTimeOffset));
+		UniqueIdentifier = HashCombineFast(UniqueIdentifier, GetTypeHash(SchemaBoneIdx));
+		UniqueIdentifier = HashCombineFast(UniqueIdentifier, GetTypeHash(SchemaOriginBoneIdx));
+		UniqueIdentifier = HashCombineFast(UniqueIdentifier, GetTypeHash(InputQueryPose));
+		UniqueIdentifier = HashCombineFast(UniqueIdentifier, GetTypeHash(PermutationTimeType));
+
+		TConstArrayView<float> CachedChannelData;
+		if (const UPoseSearchFeatureChannel* CachedChannel = SearchContext.GetCachedChannelData(UniqueIdentifier, this, CachedChannelData))
+		{
+#if DO_CHECK
+			const UPoseSearchFeatureChannel_Distance* CachedDistanceChannel = Cast<UPoseSearchFeatureChannel_Distance>(CachedChannel);
+			check(CachedDistanceChannel);
+			check(CachedDistanceChannel->GetChannelCardinality() == ChannelCardinality);
+			check(CachedChannelData.Num() == ChannelCardinality);
+
+			// making sure there were no hash collisions
+			check(CachedDistanceChannel->SampleRole == SampleRole);
+			check(CachedDistanceChannel->OriginRole == OriginRole);
+			check(CachedDistanceChannel->SamplingAttributeId == SamplingAttributeId);
+			check(CachedDistanceChannel->SampleTimeOffset == SampleTimeOffset);
+			check(CachedDistanceChannel->OriginTimeOffset == OriginTimeOffset);
+			check(CachedDistanceChannel->SchemaBoneIdx == SchemaBoneIdx);
+			check(CachedDistanceChannel->SchemaOriginBoneIdx == SchemaOriginBoneIdx);
+			check(CachedDistanceChannel->InputQueryPose == InputQueryPose);
+			check(CachedDistanceChannel->PermutationTimeType == PermutationTimeType);
+#endif //DO_CHECK
+
+			// copying the CachedChannelData into this channel portion of the FeatureVectorBuilder
+			FFeatureVectorHelper::Copy(SearchContext.EditFeatureVector().Slice(ChannelDataOffset, ChannelCardinality), 0, ChannelCardinality, CachedChannelData);
+			return;
+		}
+	}
+
+	const bool bCanUseCurrentResult = SearchContext.CanUseCurrentResult();
+	const bool bSkip = InputQueryPose != EInputQueryPose::UseCharacterPose && bCanUseCurrentResult && SampleRole == OriginRole;
+	const bool bIsRootBone = SchemaBoneIdx == RootSchemaBoneIdx;
+	if (bSkip || (!SearchContext.ArePoseHistoriesValid() && !bIsRootBone))
+	{
+		if (bCanUseCurrentResult)
+		{
+			FFeatureVectorHelper::Copy(SearchContext.EditFeatureVector(), ChannelDataOffset, ChannelCardinality, SearchContext.GetCurrentResultPoseVector());
+			return;
+		}
+
+		// we leave the SearchContext.EditFeatureVector() set to zero since the SearchContext.PoseHistory is invalid and it'll fail if we continue
+		UE_LOG(LogPoseSearch, Error, TEXT("UPoseSearchFeatureChannel_Distance::BuildQuery - Failed because Pose History Node is missing."));
+		return;
+	}
+	
+	// calculating the distance
+	const FVector BonePosition = SearchContext.GetSamplePosition(SampleTimeOffset, OriginTimeOffset, SchemaBoneIdx, SchemaOriginBoneIdx, SampleRole, OriginRole, PermutationTimeType);
+	const float Distance = BonePosition.Length();
+	FFeatureVectorHelper::EncodeFloat(SearchContext.EditFeatureVector(), ChannelDataOffset, Distance);
+}
+
+#if ENABLE_DRAW_DEBUG
+void UPoseSearchFeatureChannel_Distance::DebugDraw(const UE::PoseSearch::FDebugDrawParams& DrawParams, TConstArrayView<float> PoseVector) const
+{
+	using namespace UE::PoseSearch;
+
+	bool bDrawInjectAdditionalDebugChannels = false;
+#if WITH_EDITORONLY_DATA
+	if (const UPoseSearchSchema* Schema = GetSchema())
+	{
+		bDrawInjectAdditionalDebugChannels = Schema->bDrawInjectAdditionalDebugChannels;
+	}
+#endif // WITH_EDITORONLY_DATA
+
+	if (bDrawInjectAdditionalDebugChannels || DrawParams.IsAnyWeightRelevant(this))
+	{
+		FColor Color;
+#if WITH_EDITORONLY_DATA
+		Color = DebugColor.ToFColor(true);
+#else // WITH_EDITORONLY_DATA
+		Color = FLinearColor::Blue.ToFColor(true);
+#endif // WITH_EDITORONLY_DATA
+
+		float PermutationSampleTimeOffset = 0.f;
+		float PermutationOriginTimeOffset = 0.f;
+		UPoseSearchFeatureChannel::GetPermutationTimeOffsets(PermutationTimeType, DrawParams.ExtractPermutationTime(PoseVector), PermutationSampleTimeOffset, PermutationOriginTimeOffset);
+		const EPermutationTimeType OriginPermutationTimeType = PermutationTimeType == EPermutationTimeType::UsePermutationTime ? EPermutationTimeType::UseSampleToPermutationTime : EPermutationTimeType::UseSampleTime;
+
+		const float Distance = FFeatureVectorHelper::DecodeFloat(PoseVector, ChannelDataOffset);
+		const FVector OriginBonePos = DrawParams.ExtractPosition(PoseVector, OriginTimeOffset, SchemaOriginBoneIdx, OriginRole, OriginPermutationTimeType, INDEX_NONE, PermutationOriginTimeOffset);
+
+		static const int32 Segments = 32;
+		DrawParams.DrawSphere(OriginBonePos, Distance, Segments, Color);
+	}
+}
+#endif // ENABLE_DRAW_DEBUG
+
+#if WITH_EDITOR
+void UPoseSearchFeatureChannel_Distance::FillWeights(TArrayView<float> Weights) const
+{
+	for (int32 i = 0; i < ChannelCardinality; ++i)
+	{
+		Weights[ChannelDataOffset + i] = Weight;
+	}
+}
+
+bool UPoseSearchFeatureChannel_Distance::IndexAsset(UE::PoseSearch::FAssetIndexer& Indexer) const
+{
+	using namespace UE::PoseSearch;
+
+	FVector BonePosition;
+	for (int32 SampleIdx = Indexer.GetBeginSampleIdx(); SampleIdx != Indexer.GetEndSampleIdx(); ++SampleIdx)
+	{
+		if (Indexer.GetSamplePosition(BonePosition, SampleTimeOffset, OriginTimeOffset, SampleIdx, SchemaBoneIdx, SchemaOriginBoneIdx, SampleRole, OriginRole, PermutationTimeType, SamplingAttributeId))
+		{
+			FFeatureVectorHelper::EncodeFloat(Indexer.GetPoseVector(SampleIdx), ChannelDataOffset, BonePosition.Length());
+		}
+		else
+		{
+			return false;
+		}
+	}
+	return true;
+}
+
+UE::PoseSearch::TLabelBuilder& UPoseSearchFeatureChannel_Distance::GetLabel(UE::PoseSearch::TLabelBuilder& LabelBuilder, UE::PoseSearch::ELabelFormat LabelFormat) const
+{
+	using namespace UE::PoseSearch;
+
+	GetOuterLabel(LabelBuilder, LabelFormat);
+	AppendLabelSeparator(LabelBuilder, LabelFormat);
+
+	LabelBuilder.Append(TEXT("Dist"));
+
+	const UPoseSearchSchema* Schema = GetSchema();
+	check(Schema);
+	if (SchemaBoneIdx > RootSchemaBoneIdx)
+	{
+		LabelBuilder.Append(TEXT("_"));
+		LabelBuilder.Append(Schema->GetBoneReferences(SampleRole)[SchemaBoneIdx].BoneName.ToString());
+	}
+	else if (SchemaBoneIdx == TrajectorySchemaBoneIdx)
+	{
+		LabelBuilder.Append(TEXT("_Trj"));
+	}
+	
+	if (SampleRole != DefaultRole)
+	{
+		LabelBuilder.Append(TEXT("["));
+		LabelBuilder.Append(SampleRole.ToString());
+		LabelBuilder.Append(TEXT("]"));
+	}
+
+	if (SchemaOriginBoneIdx > RootSchemaBoneIdx)
+	{
+		LabelBuilder.Append(TEXT("_"));
+		LabelBuilder.Append(Schema->GetBoneReferences(OriginRole)[SchemaOriginBoneIdx].BoneName.ToString());
+	}
+	else if (SchemaOriginBoneIdx == TrajectorySchemaBoneIdx)
+	{
+		LabelBuilder.Append(TEXT("_Trj"));
+	}
+	
+	if (OriginRole != DefaultRole)
+	{
+		LabelBuilder.Append(TEXT("["));
+		LabelBuilder.Append(OriginRole.ToString());
+		LabelBuilder.Append(TEXT("]"));
+	}
+
+	if (PermutationTimeType == EPermutationTimeType::UsePermutationTime)
+	{
+		LabelBuilder.Append(TEXT("_PT"));
+	}
+	else if (PermutationTimeType == EPermutationTimeType::UseSampleToPermutationTime)
+	{
+		LabelBuilder.Append(TEXT("_SPT"));
+	}
+	
+	AppendLabelSeparator(LabelBuilder, LabelFormat, true);
+
+	LabelBuilder.Appendf(TEXT("%.2f"), SampleTimeOffset);
+
+	if (!FMath::IsNearlyZero(OriginTimeOffset))
+	{
+		LabelBuilder.Appendf(TEXT("-%.2f"), OriginTimeOffset);
+	}
+
+	return LabelBuilder;
+}
+
+USkeleton* UPoseSearchFeatureChannel_Distance::GetSkeleton(bool& bInvalidSkeletonIsError, const IPropertyHandle* PropertyHandle)
+{
+	// blueprint generated classes don't have a schema, until they're instanced by the schema
+	if (const UPoseSearchSchema* Schema = GetSchema())
+	{
+		bInvalidSkeletonIsError = false;
+		if (PropertyHandle)
+		{
+			if (PropertyHandle->GetProperty()->GetFName() == GET_MEMBER_NAME_CHECKED(UPoseSearchFeatureChannel_Distance, Bone))
+			{
+				return Schema->GetSkeleton(SampleRole);
+			}
+			if (PropertyHandle->GetProperty()->GetFName() == GET_MEMBER_NAME_CHECKED(UPoseSearchFeatureChannel_Distance, OriginBone))
+			{
+				return Schema->GetSkeleton(OriginRole);
+			}
+		}
+	}
+
+	return Super::GetSkeleton(bInvalidSkeletonIsError, PropertyHandle);
+}
+#endif
\ No newline at end of file
diff --git a/Engine/Plugins/Animation/PoseSearch/Source/Runtime/Public/PoseSearch/PoseSearchFeatureChannel_Distance.h b/Engine/Plugins/Animation/PoseSearch/Source/Runtime/Public/PoseSearch/PoseSearchFeatureChannel_Distance.h
new file mode 100644
--- /dev/null
+++ b/Engine/Plugins/Animation/PoseSearch/Source/Runtime/Public/PoseSearch/PoseSearchFeatureChannel_Distance.h
@@ -0,0 +1,108 @@
+// Copyright Epic Games, Inc. All Rights Reserved.
+
+#pragma once
+
+#include "PoseSearch/PoseSearchFeatureChannel.h"
+#include "BoneContainer.h"
+#include "PoseSearchFeatureChannel_Distance.generated.h"
+
+UCLASS(Experimental, EditInlineNew, Blueprintable, meta = (DisplayName = "Distance Channel"), CollapseCategories)
+class POSESEARCH_API UPoseSearchFeatureChannel_Distance : public UPoseSearchFeatureChannel
+{
+	GENERATED_BODY()
+
+public:
+	// Experimental, this feature might be removed without warning, not for production use
+	// if bDefaultWithRootBone is true Bone.BoneName and OriginBone.BoneName get initialized to the associated
+	// skeleton root bone if equals to NAME_None, otherwise transforms from the trajectory will be used
+	UPROPERTY(EditAnywhere, Category = "Experimental")
+	bool bDefaultWithRootBone = true;
+
+	UPROPERTY(EditAnywhere, Category = "Settings")
+	FBoneReference Bone;
+
+	UPROPERTY(EditAnywhere, Category = "Settings")
+	FName SampleRole = UE::PoseSearch::DefaultRole;
+
+	UPROPERTY(EditAnywhere, Category = "Settings")
+	FBoneReference OriginBone;
+
+	UPROPERTY(EditAnywhere, Category = "Settings")
+	FName OriginRole = UE::PoseSearch::DefaultRole;
+
+#if WITH_EDITORONLY_DATA
+	UPROPERTY(EditAnywhere, Category = "Settings")
+	float Weight = 1.f;
+#endif // WITH_EDITORONLY_DATA
+
+	// if SamplingAttributeId >= 0, ALL the animations contained in the pose search database referencing the schema containing this channel are expected to have 
+	// UAnimNotifyState_PoseSearchSamplingAttribute notify state with a matching SamplingAttributeId, and the UAnimNotifyState_PoseSearchSamplingAttribute properties
+	// will be used as source of data instead of this channel "Bone". UAnimNotifyState_PoseSearchSamplingAttribute properties will be then converted into OriginBone space
+	UPROPERTY(EditAnywhere, Category = "Settings")
+	int32 SamplingAttributeId = INDEX_NONE;
+	
+	// the data relative to the sampling time associated to this channel will be offsetted by SampleTimeOffset seconds.
+	// For example, if Bone is the head bone, and SampleTimeOffset is 0.5, this channel will try to match the future distance of the character head bone 0.5 seconds ahead
+	UPROPERTY(EditAnywhere, Category = "Settings")
+	float SampleTimeOffset = 0.f;
+
+	// the data relative to the sampling time associated to this channel origin (root / trajectory bone) will be offsetted by OriginTimeOffset seconds.
+	// For example, if Bone is the head bone, SampleTimeOffset is 0.5, and OriginTimeOffset is 0.5, this channel will try to match 
+	// the future distance of the character head bone 0.5 seconds ahead, relative to the future root bone 0.5 seconds ahead
+	UPROPERTY(EditAnywhere, Category = "Settings")
+	float OriginTimeOffset = 0.f;
+
+	// index referencing the associated bone in UPoseSearchSchema::BoneReferences
+	UPROPERTY(Transient)
+	int8 SchemaBoneIdx = 0;
+
+	UPROPERTY(Transient)
+	int8 SchemaOriginBoneIdx = 0;
+
+#if WITH_EDITORONLY_DATA
+	UPROPERTY(EditAnywhere, Category = "Settings", meta = (ExcludeFromHash, DisplayPriority = 0))
+	FLinearColor DebugColor = FLinearColor::Blue;
+#endif // WITH_EDITORONLY_DATA
+
+	UPROPERTY(EditAnywhere, Category = "Settings")
+	EInputQueryPose InputQueryPose = EInputQueryPose::UseContinuingPose;
+
+	UPROPERTY(EditAnywhere, Category = "Settings")
+	EPermutationTimeType PermutationTimeType = EPermutationTimeType::UseSampleTime;
+
+#if WITH_EDITORONLY_DATA
+	// if set, all the channels of the same class with the same cardinality, and the same NormalizationGroup, will be normalized together.
+	// for example in a locomotion database of a character holding a weapon, containing non mirrorable animations, you'd still want to normalize together 
+	// all the distances
+	UPROPERTY(EditAnywhere, Category = "Settings")
+	FName NormalizationGroup;
+#endif //WITH_EDITORONLY_DATA
+
+	UFUNCTION(BlueprintPure, BlueprintImplementableEvent, meta=(BlueprintThreadSafe, DisplayName = "Get Distance"), Category = "Settings")
+	float BP_GetDistance(const UAnimInstance* AnimInstance) const;
+
+	bool bUseBlueprintQueryOverride = false;
+
+	UPoseSearchFeatureChannel_Distance();
+
+	// UPoseSearchFeatureChannel interface
+	virtual bool Finalize(UPoseSearchSchema* Schema) override;
+	virtual void BuildQuery(UE::PoseSearch::FSearchContext& SearchContext) const override;
+
+	virtual EPermutationTimeType GetPermutationTimeType() const override { return PermutationTimeType; }
+	virtual void AddDependentChannels(UPoseSearchSchema* Schema) const override;
+
+#if ENABLE_DRAW_DEBUG
+	virtual void DebugDraw(const UE::PoseSearch::FDebugDrawParams& DrawParams, TConstArrayView<float> PoseVector) const override;
+#endif // ENABLE_DRAW_DEBUG
+
+#if WITH_EDITOR
+	virtual void FillWeights(TArrayView<float> Weights) const override;
+	virtual bool IndexAsset(UE::PoseSearch::FAssetIndexer& Indexer) const override;
+	virtual UE::PoseSearch::TLabelBuilder& GetLabel(UE::PoseSearch::TLabelBuilder& LabelBuilder, UE::PoseSearch::ELabelFormat LabelFormat = UE::PoseSearch::ELabelFormat::Full_Horizontal) const override;
+	virtual FName GetNormalizationGroup() const override { return NormalizationGroup; }
+
+	// IBoneReferenceSkeletonProvider interface
+	USkeleton* GetSkeleton(bool& bInvalidSkeletonIsError, const IPropertyHandle* PropertyHandle) override;
+#endif
+};
