UI/UX improvements for pose search debugger.
- Fixed inproper scaling of columns
- Fixed text alignment

#rnx

diff --git a/Engine/Plugins/Animation/PoseSearch/Source/Editor/Private/PoseSearchDebuggerDatabaseColumns.h b/Engine/Plugins/Animation/PoseSearch/Source/Editor/Private/PoseSearchDebuggerDatabaseColumns.h
--- a/Engine/Plugins/Animation/PoseSearch/Source/Editor/Private/PoseSearchDebuggerDatabaseColumns.h
+++ b/Engine/Plugins/Animation/PoseSearch/Source/Editor/Private/PoseSearchDebuggerDatabaseColumns.h
@@ -23,6 +23,8 @@
 
 namespace UE::PoseSearch::DebuggerDatabaseColumns
 {
+constexpr float ColumnWidthPadding = 10.f;
+constexpr float ColumnSortWidthPadding = 20.f;
 	
 /** Column struct to represent each column in the debugger database */
 struct IColumn : TSharedFromThis<IColumn>
@@ -95,13 +97,12 @@ protected:
 	{
 		const TSharedRef<FSlateFontMeasure> FontMeasure = FSlateApplication::Get().GetRenderer()->GetFontMeasureService();
 		const FSlateFontInfo & CellTextStyle = FAppStyle::GetFontStyle("SmallText");
-		const float CellPadding = 10.0f;
 
-		float MinWidth = FontMeasure->Measure(GetLabel(), CellTextStyle).X + CellPadding;
+		float MinWidth = FontMeasure->Measure(GetLabel(), CellTextStyle).X + ColumnWidthPadding;
 		
 		for (const TSharedRef<FDebuggerDatabaseRowData>& Row : Rows)
 		{
-			MinWidth = FMath::Max((FontMeasure->Measure(GetRowText(Row), CellTextStyle).X + CellPadding), MinWidth);
+			MinWidth = FMath::Max((FontMeasure->Measure(GetRowText(Row), CellTextStyle).X + ColumnWidthPadding), MinWidth);
 		}
 		
 		return MinWidth;
@@ -242,13 +243,12 @@ struct FDatabaseName : IColumn
 	{
 		const TSharedRef<FSlateFontMeasure> FontMeasure = FSlateApplication::Get().GetRenderer()->GetFontMeasureService();
 		const FSlateFontInfo & CellTextStyle = FAppStyle::GetFontStyle("SmallText");
-		const float CellPadding = 10.0f;
 
-		float MinWidth = FontMeasure->Measure(GetLabel(), CellTextStyle).X + CellPadding;
+		float MinWidth = FontMeasure->Measure(GetLabel(), CellTextStyle).X + ColumnWidthPadding;
 		
 		for (const TSharedRef<FDebuggerDatabaseRowData>& Row : Rows)
 		{
-			MinWidth = FMath::Max(FontMeasure->Measure(Row->SharedData->DatabaseName, CellTextStyle).X + CellPadding, MinWidth);
+			MinWidth = FMath::Max(FontMeasure->Measure(Row->SharedData->DatabaseName, CellTextStyle).X + ColumnWidthPadding, MinWidth);
 		}
 		
 		return MinWidth;
@@ -332,9 +332,8 @@ struct FAssetName : IColumn
 	{
 		const TSharedRef<FSlateFontMeasure> FontMeasure = FSlateApplication::Get().GetRenderer()->GetFontMeasureService();
 		const FSlateFontInfo & CellTextStyle = FAppStyle::GetFontStyle("SmallText");
-		const float CellPadding = 10.0f;
 
-		float MinWidth = FontMeasure->Measure(GetLabel(), CellTextStyle).X + CellPadding;
+		float MinWidth = FontMeasure->Measure(GetLabel(), CellTextStyle).X + ColumnWidthPadding;
 		
 		for (const TSharedRef<FDebuggerDatabaseRowData>& Row : Rows)
 		{
diff --git a/Engine/Plugins/Animation/PoseSearch/Source/Editor/Private/PoseSearchDebuggerDatabaseRow.h b/Engine/Plugins/Animation/PoseSearch/Source/Editor/Private/PoseSearchDebuggerDatabaseRow.h
--- a/Engine/Plugins/Animation/PoseSearch/Source/Editor/Private/PoseSearchDebuggerDatabaseRow.h
+++ b/Engine/Plugins/Animation/PoseSearch/Source/Editor/Private/PoseSearchDebuggerDatabaseRow.h
@@ -53,8 +53,8 @@ class SDebuggerDatabaseRow : public SMultiColumnTableRow<TSharedRef<FDebuggerDat
 		
 		return
 			SNew(SBorder)
-			.HAlign(HAlign_Center)
-			.VAlign(VAlign_Center)
+			.HAlign(HAlign_Fill)
+			.VAlign(VAlign_Fill)
 			.BorderImage(RowBrush)
 			[
 				SNew(SHorizontalBox)
diff --git a/Engine/Plugins/Animation/PoseSearch/Source/Editor/Private/PoseSearchDebuggerDatabaseView.cpp b/Engine/Plugins/Animation/PoseSearch/Source/Editor/Private/PoseSearchDebuggerDatabaseView.cpp
--- a/Engine/Plugins/Animation/PoseSearch/Source/Editor/Private/PoseSearchDebuggerDatabaseView.cpp
+++ b/Engine/Plugins/Animation/PoseSearch/Source/Editor/Private/PoseSearchDebuggerDatabaseView.cpp
@@ -456,6 +456,12 @@ void SDebuggerDatabaseView::Update(const FTraceMotionMatchingStateMessage& State
 		for (TPair<FName, TSharedRef<IColumn>>& ColumnPair : Columns)
 		{
 			ColumnPair.Value->Width = ColumnPair.Value->ComputeColumnWidth(UnfilteredDatabaseRows);
+
+			// Add spacing for sorting arrow
+			if (ColumnPair.Value->ColumnId == SortColumn && SortMode != EColumnSortMode::None)
+			{
+				ColumnPair.Value->Width += DebuggerDatabaseColumns::ColumnSortWidthPadding;
+			}
 		}
 		
 		// Add columns from map to header row
@@ -473,15 +479,15 @@ void SDebuggerDatabaseView::Update(const FTraceMotionMatchingStateMessage& State
 					.OnWidthChanged(this, &SDebuggerDatabaseView::OnColumnWidthChanged, Column.ColumnId)
 					.VAlignHeader(VAlign_Center)
 					.VAlignCell(VAlign_Center)
-					.HAlignHeader(HAlign_Fill)
+					.HAlignHeader(HAlign_Center)
 					.HAlignCell(HAlign_Fill)
 					.HeaderContent()
 					[
-						SNew( SBox )
+						SNew(SBox)
 						.VAlign( VAlign_Fill )
 						.HAlign( HAlign_Fill )
 						[
-							SNew( STextBlock )
+							SNew(STextBlock)
 							.TextStyle(FAppStyle::Get(), "SmallText")
 							.Text(Column.GetLabel())
 							.Margin(FMargin(5.f, 5.f, 5.f, 5.f))
@@ -522,10 +528,24 @@ float SDebuggerDatabaseView::GetColumnWidth(const FName ColumnId) const
 
 void SDebuggerDatabaseView::OnColumnSortModeChanged(const EColumnSortPriority::Type SortPriority, const FName & ColumnId, const EColumnSortMode::Type InSortMode)
 {
+	const bool bHasSortColumnChanged = SortColumn != ColumnId;
+	
+	// Restore previous sort column to its default width
+	if (bHasSortColumnChanged)
+	{
+		Columns[SortColumn]->Width -= DebuggerDatabaseColumns::ColumnSortWidthPadding;
+	}
+	
 	SortColumn = ColumnId;
 	SortMode = InSortMode;
 	SortDatabaseRows();
 	PopulateViewRows();
+
+	// Take into account sorting arrow for the column width
+	if (bHasSortColumnChanged && SortMode != EColumnSortMode::None)
+	{
+		Columns[SortColumn]->Width += DebuggerDatabaseColumns::ColumnSortWidthPadding;
+	}
 }
 
 void SDebuggerDatabaseView::OnColumnWidthChanged(const float NewWidth, FName ColumnId) const
