posesearch - disable interacting characters collisions via availabilities

[RN] skip

diff --git a/Engine/Plugins/Animation/PoseSearch/Source/Runtime/Private/PoseSearchInteractionIsland.cpp b/Engine/Plugins/Animation/PoseSearch/Source/Runtime/Private/PoseSearchInteractionIsland.cpp
--- a/Engine/Plugins/Animation/PoseSearch/Source/Runtime/Private/PoseSearchInteractionIsland.cpp
+++ b/Engine/Plugins/Animation/PoseSearch/Source/Runtime/Private/PoseSearchInteractionIsland.cpp
@@ -14,12 +14,29 @@
 #include "PoseSearch/PoseSearchLibrary.h"
 #include "PoseSearch/PoseSearchDatabase.h"
 #include "PoseSearch/PoseSearchSchema.h"
+#include "VisualLogger/VisualLogger.h"
 
 namespace UE::PoseSearch
 {
 
 // Utils functions
 ///////////////////////////////////////////////////////////
+const FTransform& GetContextTransform(const UObject* AnimContext)
+{
+	if (const UAnimInstance* AnimInstance = Cast<UAnimInstance>(AnimContext))
+	{
+		return AnimInstance->GetSkelMeshComponent()->GetComponentTransform();
+	}
+
+	unimplemented();
+	return FTransform::Identity;
+}
+
+FVector GetContextLocation(const UObject* AnimContext)
+{
+	return GetContextTransform(AnimContext).GetLocation();
+}
+
 static FInteractionSearchResult InitSearchResult(const FSearchResult& SearchResult, const FInteractionSearchContext& SearchContext, int32 SearchIndex)
 {
 	FInteractionSearchResult InteractionSearchResult;
@@ -163,9 +180,9 @@ static void InitSearchResults(TArray<FInteractionSearchResult>& SearchResults, T
 	}
 }
 
-// FInteractionSearchContext
+// FInteractionSearchContextBase
 ///////////////////////////////////////////////////////////
-bool FInteractionSearchContext::IsEquivalent(const FInteractionSearchContext& Other) const
+bool FInteractionSearchContextBase::IsEquivalent(const FInteractionSearchContextBase& Other) const
 {
 	if (Database != Other.Database)
 	{
@@ -201,8 +218,45 @@ bool FInteractionSearchContext::IsEquivalent(const FInteractionSearchContext& Ot
 	return false;
 }
 
+#if ENABLE_VISUAL_LOG
+void FInteractionSearchContextBase::VLogContext(const FColor& Color) const
+{
+	using namespace UE::PoseSearch;
+
+	if (FVisualLogger::IsRecording())
+	{
+		static const TCHAR* LogName = TEXT("PoseSearchInteraction");
+
+		const int32 AnimContextsNum = AnimContexts.Num();
+		TArray<FVector, TInlineAllocator<PreallocatedRolesNum, TMemStackAllocator<>>> Locations;
+		Locations.SetNum(AnimContextsNum);
+		for (int32 Index = 0; Index < AnimContextsNum; ++Index)
+		{
+			if (const UObject* AnimContext = AnimContexts[Index].Get())
+			{
+				Locations[Index] = GetContextLocation(AnimContext);
+			}
+		}
+
+		for (int32 IndexA = 0; IndexA < AnimContextsNum; ++IndexA)
+		{
+			for (int32 IndexB = IndexA + 1; IndexB < AnimContextsNum; ++IndexB)
+			{
+				for (int32 IndexAll = 0; IndexAll < AnimContextsNum; ++IndexAll)
+				{
+					if (const UObject* AnimContext = AnimContexts[IndexAll].Get())
+					{
+						UE_VLOG_SEGMENT(AnimContext, LogName, Display, Locations[IndexA], Locations[IndexB], Color, TEXT(""));
+					}
+				}
+			}
+		}
+	}
+}
+#endif // ENABLE_VISUAL_LOG
+
 #if DO_CHECK
-bool FInteractionSearchContext::CheckForConsistency() const
+bool FInteractionSearchContextBase::CheckForConsistency() const
 {
 	if (Database == nullptr)
 	{
@@ -271,7 +325,7 @@ bool FInteractionSearchContext::CheckForConsistency() const
 	return true;
 }
 
-void FInteractionSearchContext::TestHistoryCollectorsThreadingAccess() const
+void FInteractionSearchContextBase::TestHistoryCollectorsThreadingAccess() const
 {
 	for (int32 Index = 0; Index < AnimContexts.Num(); ++Index)
 	{
diff --git a/Engine/Plugins/Animation/PoseSearch/Source/Runtime/Private/PoseSearchInteractionSubsystem.cpp b/Engine/Plugins/Animation/PoseSearch/Source/Runtime/Private/PoseSearchInteractionSubsystem.cpp
--- a/Engine/Plugins/Animation/PoseSearch/Source/Runtime/Private/PoseSearchInteractionSubsystem.cpp
+++ b/Engine/Plugins/Animation/PoseSearch/Source/Runtime/Private/PoseSearchInteractionSubsystem.cpp
@@ -22,22 +22,6 @@ static bool GVarPoseSearchInteractionCacheIslands = true;
 static FAutoConsoleVariableRef CVarPoseSearchInteractionCacheIslands(TEXT("a.PoseSearchInteraction.CacheIslands"), GVarPoseSearchInteractionCacheIslands, TEXT("Cache Pose Search Interaction Islands for future reuse instead of destrying them"));
 #endif // !NO_CVARS
 
-static const FTransform& GetContextTransform(const UObject* AnimContext)
-{
-	if (const UAnimInstance* AnimInstance = Cast<UAnimInstance>(AnimContext))
-	{
-		return AnimInstance->GetSkelMeshComponent()->GetComponentTransform();
-	}
-
-	unimplemented();
-	return FTransform::Identity;
-}
-
-static FVector GetContextLocation(const UObject* AnimContext)
-{
-	return GetContextTransform(AnimContext).GetLocation();
-}
-
 FString FInteractionAvailability::GetPoseHistoryName() const
 {
 	if (HistoryCollector)
@@ -471,6 +455,8 @@ void UPoseSearchInteractionSubsystem::Initialize(FSubsystemCollectionBase& Colle
 
 void UPoseSearchInteractionSubsystem::Deinitialize()
 {
+	UpdateValidInteractionSearches();
+
 	DestroyAllIslands();
 
 	Super::Deinitialize();
@@ -636,12 +622,13 @@ void UPoseSearchInteractionSubsystem::GenerateSearchContexts(UE::PoseSearch::FIn
 
 	struct FRoledAnimContextInfo
 	{
-		FRoledAnimContextInfo(const FAnimContextInfo& InAnimContextInfo, const FRole InRole, const FAnimNode_PoseSearchHistoryCollector_Base& InHistoryCollector, const UPoseSearchDatabase& InDatabase, float InMaxDistanceSquared)
+		FRoledAnimContextInfo(const FAnimContextInfo& InAnimContextInfo, const FRole InRole, const FAnimNode_PoseSearchHistoryCollector_Base& InHistoryCollector, const UPoseSearchDatabase& InDatabase, float InMaxDistanceSquared, bool bInDisableCollisions)
 			: AnimContextInfo(&InAnimContextInfo)
 			, Role(InRole)
 			, HistoryCollector(&InHistoryCollector)
 			, Database(&InDatabase)
 			, MaxDistanceSquared(InMaxDistanceSquared)
+			, bDisableCollisions(bInDisableCollisions)
 		{
 			check(MaxDistanceSquared > UE_SMALL_NUMBER);
 		}
@@ -653,10 +640,11 @@ void UPoseSearchInteractionSubsystem::GenerateSearchContexts(UE::PoseSearch::FIn
 		const FAnimNode_PoseSearchHistoryCollector_Base* HistoryCollector = nullptr;
 		const UPoseSearchDatabase* Database = nullptr;
 		float MaxDistanceSquared = 0.f;
+		bool bDisableCollisions = false;
 	};
 	struct FRoledAnimContextInfos : TArray<FRoledAnimContextInfo, TInlineAllocator<PreallocatedRolesNum, TMemStackAllocator<>>>
 	{
-		void AddRoledAnimContextInfos(const FAnimContextInfo& AnimContextInfo, const TArray<FName>& RolesFilter, const FAnimNode_PoseSearchHistoryCollector_Base& HistoryCollector, const UPoseSearchDatabase& Database, float MaxDistanceSquared)
+		void AddRoledAnimContextInfos(const FAnimContextInfo& AnimContextInfo, const TArray<FName>& RolesFilter, const FAnimNode_PoseSearchHistoryCollector_Base& HistoryCollector, const UPoseSearchDatabase& Database, float MaxDistanceSquared, bool bDisableCollisions)
 		{
 			const UPoseSearchSchema* Schema = Database.Schema;
 			check(Schema);
@@ -665,7 +653,7 @@ void UPoseSearchInteractionSubsystem::GenerateSearchContexts(UE::PoseSearch::FIn
 				// adding ALL the possible roles from the database:
 				for (const FPoseSearchRoledSkeleton& RoledSkeleton : Schema->GetRoledSkeletons())
 				{
-					AddUnique(FRoledAnimContextInfo(AnimContextInfo, RoledSkeleton.Role, HistoryCollector, Database, MaxDistanceSquared));
+					AddUnique(FRoledAnimContextInfo(AnimContextInfo, RoledSkeleton.Role, HistoryCollector, Database, MaxDistanceSquared, bDisableCollisions));
 				}
 			}
 			else
@@ -674,7 +662,7 @@ void UPoseSearchInteractionSubsystem::GenerateSearchContexts(UE::PoseSearch::FIn
 				{
 					if (Schema->GetRoledSkeleton(Role))
 					{
-						AddUnique(FRoledAnimContextInfo(AnimContextInfo, Role, HistoryCollector, Database, MaxDistanceSquared));
+						AddUnique(FRoledAnimContextInfo(AnimContextInfo, Role, HistoryCollector, Database, MaxDistanceSquared, bDisableCollisions));
 					}
 					else
 					{
@@ -768,7 +756,7 @@ void UPoseSearchInteractionSubsystem::GenerateSearchContexts(UE::PoseSearch::FIn
 						check(Database->Schema);
 						FRoledAnimContextInfos& RoledAnimContextInfos = DatabaseToRoledAnimContextInfos.FindOrAdd(Database);
 						const float MaxDistanceSquared = Availability.BroadPhaseRadius * Availability.BroadPhaseRadius;
-						RoledAnimContextInfos.AddRoledAnimContextInfos(AnimContextInfo, Availability.RolesFilter, *HistoryCollector, *Database, MaxDistanceSquared);
+						RoledAnimContextInfos.AddRoledAnimContextInfos(AnimContextInfo, Availability.RolesFilter, *HistoryCollector, *Database, MaxDistanceSquared, Availability.bDisableCollisions);
 					}
 					else if (Availability.IsTagValid())
 					{
@@ -781,7 +769,7 @@ void UPoseSearchInteractionSubsystem::GenerateSearchContexts(UE::PoseSearch::FIn
 							{
 								check(DatabaseFromTag && DatabaseFromTag->Schema);
 								FRoledAnimContextInfos& RoledAnimContextInfos = DatabaseToRoledAnimContextInfos.FindOrAdd(DatabaseFromTag);
-								RoledAnimContextInfos.AddRoledAnimContextInfos(AnimContextInfo, Availability.RolesFilter, *HistoryCollector, *DatabaseFromTag, MaxDistanceSquared);
+								RoledAnimContextInfos.AddRoledAnimContextInfos(AnimContextInfo, Availability.RolesFilter, *HistoryCollector, *DatabaseFromTag, MaxDistanceSquared, Availability.bDisableCollisions);
 							}
 						}
 						else
@@ -884,6 +872,7 @@ void UPoseSearchInteractionSubsystem::GenerateSearchContexts(UE::PoseSearch::FIn
 							SearchContext.AnimContexts.Add(RoledAnimContextInfo.AnimContextInfo->AnimContext);
 							SearchContext.Roles.Add(RoledAnimContextInfo.Role);
 							SearchContext.HistoryCollectors.Add(RoledAnimContextInfo.HistoryCollector);
+							SearchContext.bDisableCollisions |= RoledAnimContextInfo.bDisableCollisions;
 						}
 
 						// does SearchContext cover all the roles required by this interaction?
@@ -911,6 +900,152 @@ void UPoseSearchInteractionSubsystem::GenerateSearchContexts(UE::PoseSearch::FIn
 		});
 }
 
+void UPoseSearchInteractionSubsystem::OnInteractionStart(UE::PoseSearch::FValidInteractionSearch& ValidInteractionSearch)
+{
+	using namespace UE::PoseSearch;
+
+#if ENABLE_VISUAL_LOG
+	ValidInteractionSearch.SearchContext.VLogContext(FColor::Blue);
+#endif
+
+	check(ValidInteractionSearch.DisabledCollisions.IsEmpty());
+	if (ValidInteractionSearch.SearchContext.bDisableCollisions)
+	{
+		TArray<AActor*, TInlineAllocator<PreallocatedRolesNum, TMemStackAllocator<>>> Actors;
+		TArray<UPrimitiveComponent*, TInlineAllocator<PreallocatedRolesNum, TMemStackAllocator<>>> PrimitiveComponents;
+
+		for (const TWeakObjectPtr<const UObject>& AnimContextPtr : ValidInteractionSearch.SearchContext.AnimContexts)
+		{
+			if (const UObject* AnimContext = AnimContextPtr.Get())
+			{
+				if (const UAnimInstance* AnimInstance = Cast<UAnimInstance>(AnimContext))
+				{
+					AActor* Actor = AnimInstance->GetOwningActor();
+					check(Actor);
+					Actors.Add(Actor);
+					PrimitiveComponents.Add(Cast<UPrimitiveComponent>(Actor->GetRootComponent()));
+				}
+				else
+				{
+					unimplemented();
+				}
+			}
+		}
+
+		for (int32 IndexA = 0; IndexA < Actors.Num(); ++IndexA)
+		{
+			for (int32 IndexB = IndexA + 1; IndexB < Actors.Num(); ++IndexB)
+			{
+				AActor* ActorA = Actors[IndexA];
+				AActor* ActorB = Actors[IndexB];
+				check(ActorA && ActorB);
+
+				UPrimitiveComponent* PrimitiveComponentA = PrimitiveComponents[IndexA];
+				UPrimitiveComponent* PrimitiveComponentB = PrimitiveComponents[IndexB];
+
+				if (PrimitiveComponentA && !PrimitiveComponentA->GetMoveIgnoreActors().Contains(ActorB))
+				{
+					ValidInteractionSearch.DisabledCollisions.Add({ ActorA, ActorB });
+					PrimitiveComponentA->IgnoreActorWhenMoving(ActorB, true);
+				}
+
+				if (PrimitiveComponentB && !PrimitiveComponentB->GetMoveIgnoreActors().Contains(ActorA))
+				{
+					ValidInteractionSearch.DisabledCollisions.Add({ ActorB, ActorA });
+					PrimitiveComponentB->IgnoreActorWhenMoving(ActorA, true);
+				}
+			}
+		}
+	}
+}
+
+void UPoseSearchInteractionSubsystem::OnInteractionContinuing(UE::PoseSearch::FValidInteractionSearch& ValidInteractionSearch)
+{
+#if ENABLE_VISUAL_LOG
+	ValidInteractionSearch.SearchContext.VLogContext(FColor::Green);
+#endif
+}
+
+void UPoseSearchInteractionSubsystem::OnInteractionEnd(UE::PoseSearch::FValidInteractionSearch& ValidInteractionSearch)
+{
+	using namespace UE::PoseSearch;
+
+#if ENABLE_VISUAL_LOG
+	ValidInteractionSearch.SearchContext.VLogContext(FColor::Black);
+#endif
+
+	for (const FDisabledCollisions::ElementType& DisabledCollision : ValidInteractionSearch.DisabledCollisions)
+	{
+		if (AActor* ActorA = DisabledCollision.Key.Get())
+		{
+			if (AActor* ActorB = DisabledCollision.Value.Get())
+			{
+				if (UPrimitiveComponent* PrimitiveComponentA = Cast<UPrimitiveComponent>(ActorA->GetRootComponent()))
+				{
+					PrimitiveComponentA->IgnoreActorWhenMoving(ActorB, false);
+				}
+			}
+		}
+	}
+}
+
+void UPoseSearchInteractionSubsystem::UpdateValidInteractionSearches()
+{
+	using namespace UE::PoseSearch;
+
+	const int32 ValidInteractionSearchesNum = ValidInteractionSearches.Num();
+	TArray<bool, TInlineAllocator<PreallocatedSearchesNum * 4, TMemStackAllocator<>>> Visited;
+	Visited.SetNum(ValidInteractionSearchesNum);
+
+	TArray<FValidInteractionSearch, TInlineAllocator<PreallocatedSearchesNum * 4, TMemStackAllocator<>>> NewValidInteractionSearches;
+	for (FInteractionIsland* Island : Islands)
+	{
+		if (Island->IsInitialized())
+		{
+			// analyzing ALL current tick interaction results
+			for (const FInteractionSearchResult& SearchResult : Island->GetSearchResults())
+			{
+				const FInteractionSearchContext& SearchContext = Island->GetSearchContexts()[SearchResult.SearchIndex];
+
+				int32 Index = 0;
+				for (; Index < ValidInteractionSearchesNum; ++Index)
+				{
+					if (ValidInteractionSearches[Index].SearchContext.IsEquivalent(SearchContext))
+					{
+						check(!Visited[Index]);
+						Visited[Index] = true;
+						OnInteractionContinuing(ValidInteractionSearches[Index]);
+						break;
+					}
+				}
+
+				if (Index == ValidInteractionSearchesNum)
+				{
+					// we haven't found an equivalent SearchContext in ValidInteractionSearches, so it's a new interaction!
+					FValidInteractionSearch& NewValidInteractionSearch = NewValidInteractionSearches.AddDefaulted_GetRef();
+					NewValidInteractionSearch.SearchContext = SearchContext;
+					OnInteractionStart(NewValidInteractionSearch);
+				}
+			}
+		}
+	}
+
+	// checking for leftover unvisited ValidInteractionSearches. Those are interactions that just ended
+	for (int32 Index = 0; Index < ValidInteractionSearchesNum; ++Index)
+	{
+		if (Visited[Index])
+		{
+			NewValidInteractionSearches.Add(ValidInteractionSearches[Index]);
+		}
+		else
+		{
+			OnInteractionEnd(ValidInteractionSearches[Index]);
+		}
+	}
+
+	ValidInteractionSearches = NewValidInteractionSearches;
+}
+
 void UPoseSearchInteractionSubsystem::Tick(float DeltaSeconds)
 {
 	QUICK_SCOPE_CYCLE_COUNTER(STAT_UPoseSearchInteractionSubsystem_Tick);
@@ -919,6 +1054,9 @@ void UPoseSearchInteractionSubsystem::Tick(float DeltaSeconds)
 
 	Super::Tick(DeltaSeconds);
 
+	FMemMark Mark(FMemStack::Get());
+	UpdateValidInteractionSearches();
+
 	if (!IsThereAnyAvailability())
 	{
 		bool bAllUninjected = true;
@@ -940,7 +1078,6 @@ void UPoseSearchInteractionSubsystem::Tick(float DeltaSeconds)
 	check(IsInGameThread());
 
 	// FScopeLock Lock(&AvailabilitiesMapMutex); is not necessary since UPoseSearchInteractionSubsystem gets ticked outside the parallel animation jobs
-	FMemMark Mark(FMemStack::Get());
 
 	// generating all the possible interaction tuples of AnimContext(s) with roles and pose histories (defined in FInteractionSearchContext)
 	FInteractionSearchContexts SearchContexts;
diff --git a/Engine/Plugins/Animation/PoseSearch/Source/Runtime/Public/PoseSearch/PoseSearchInteractionIsland.h b/Engine/Plugins/Animation/PoseSearch/Source/Runtime/Public/PoseSearch/PoseSearchInteractionIsland.h
--- a/Engine/Plugins/Animation/PoseSearch/Source/Runtime/Public/PoseSearch/PoseSearchInteractionIsland.h
+++ b/Engine/Plugins/Animation/PoseSearch/Source/Runtime/Public/PoseSearch/PoseSearchInteractionIsland.h
@@ -16,7 +16,12 @@ namespace UE::PoseSearch
 struct FInteractionIsland;
 
 // Experimental, this feature might be removed without warning, not for production use
-struct FInteractionSearchContext
+const FTransform& GetContextTransform(const UObject* AnimContext);
+// Experimental, this feature might be removed without warning, not for production use
+FVector GetContextLocation(const UObject* AnimContext);
+
+// Experimental, this feature might be removed without warning, not for production use
+struct FInteractionSearchContextBase
 {
 	// @todo: since AnimContexts, HistoryCollectors, Roles, and BroadPhaseRadiuses have the same cardinality,
 	// add an array of structs instead of many arrays
@@ -26,10 +31,14 @@ struct FInteractionSearchContext
 	TArray<UE::PoseSearch::FRole, TInlineAllocator<PreallocatedRolesNum>> Roles;
 
 	TWeakObjectPtr<const UPoseSearchDatabase> Database;
-	FPoseSearchContinuingProperties ContinuingProperties;
+	bool bDisableCollisions = false;
 
 	bool IsValid() const { return !AnimContexts.IsEmpty(); }
-	bool IsEquivalent(const FInteractionSearchContext& Other) const;
+	bool IsEquivalent(const FInteractionSearchContextBase& Other) const;
+
+#if ENABLE_VISUAL_LOG
+	void VLogContext(const FColor& Color) const;
+#endif // ENABLE_VISUAL_LOG
 
 #if DO_CHECK
 	// method that queries the HistoryCollectors. the PoseHistory will complain if accessed in a non thread safe manner!
@@ -37,6 +46,13 @@ struct FInteractionSearchContext
 	bool CheckForConsistency() const;
 #endif // DO_CHECK
 };
+
+// Experimental, this feature might be removed without warning, not for production use
+struct FInteractionSearchContext : FInteractionSearchContextBase
+{
+	FPoseSearchContinuingProperties ContinuingProperties;
+};
+// Experimental, this feature might be removed without warning, not for production use
 typedef TArray<FInteractionSearchContext, TInlineAllocator<PreallocatedSearchesNum, TMemStackAllocator<>>> FInteractionSearchContexts;
 
 // Experimental, this feature might be removed without warning, not for production use
diff --git a/Engine/Plugins/Animation/PoseSearch/Source/Runtime/Public/PoseSearch/PoseSearchInteractionLibrary.h b/Engine/Plugins/Animation/PoseSearch/Source/Runtime/Public/PoseSearch/PoseSearchInteractionLibrary.h
--- a/Engine/Plugins/Animation/PoseSearch/Source/Runtime/Public/PoseSearch/PoseSearchInteractionLibrary.h
+++ b/Engine/Plugins/Animation/PoseSearch/Source/Runtime/Public/PoseSearch/PoseSearchInteractionLibrary.h
@@ -39,6 +39,10 @@ struct POSESEARCH_API FPoseSearchInteractionAvailability
 	UPROPERTY(EditAnywhere, BlueprintReadWrite, Category=Settings)
 	float BroadPhaseRadius = 500.f;
 
+	// if true, the system will disable collsions between interacting characters
+	UPROPERTY(EditAnywhere, BlueprintReadWrite, Category=Settings)
+	bool bDisableCollisions = false;
+
 	bool IsTagValid() const { return !Tag.IsNone(); }
 };
 
diff --git a/Engine/Plugins/Animation/PoseSearch/Source/Runtime/Public/PoseSearch/PoseSearchInteractionSubsystem.h b/Engine/Plugins/Animation/PoseSearch/Source/Runtime/Public/PoseSearch/PoseSearchInteractionSubsystem.h
--- a/Engine/Plugins/Animation/PoseSearch/Source/Runtime/Public/PoseSearch/PoseSearchInteractionSubsystem.h
+++ b/Engine/Plugins/Animation/PoseSearch/Source/Runtime/Public/PoseSearch/PoseSearchInteractionSubsystem.h
@@ -32,7 +32,9 @@ private:
 	const FAnimNode_PoseSearchHistoryCollector_Base* HistoryCollector = nullptr;
 };
 
+// Experimental, this feature might be removed without warning, not for production use
 typedef TArray<FInteractionAvailability> FInteractionAvailabilities;
+// Experimental, this feature might be removed without warning, not for production use
 struct FAvailabilitiesMap : protected TMap<TWeakObjectPtr<const UObject>, FInteractionAvailabilities>
 {
 	using Super = TMap<TWeakObjectPtr<const UObject>, FInteractionAvailabilities>;
@@ -61,6 +63,16 @@ struct FAvailabilitiesMap : protected TMap<TWeakObjectPtr<const UObject>, FInter
 	}
 };
 
+// Experimental, this feature might be removed without warning, not for production use
+typedef TArray<TPair<TWeakObjectPtr<AActor>, TWeakObjectPtr<AActor>>> FDisabledCollisions;
+
+// Experimental, this feature might be removed without warning, not for production use
+struct FValidInteractionSearch
+{
+	FInteractionSearchContextBase SearchContext;
+	FDisabledCollisions DisabledCollisions;
+};
+
 struct FAnimContextInfos;
 struct FTagToDatabases;
 
@@ -163,6 +175,11 @@ private:
 	void GenerateAnimContextInfosAndTagToDatabases(UE::PoseSearch::FAnimContextInfos& AnimContextInfos, UE::PoseSearch::FTagToDatabases& TagToDatabases) const;
 	void GenerateSearchContexts(UE::PoseSearch::FInteractionSearchContexts& SearchContexts) const;
 
+	void UpdateValidInteractionSearches();
+	static void OnInteractionStart(UE::PoseSearch::FValidInteractionSearch& ValidInteractionSearch);
+	static void OnInteractionContinuing(UE::PoseSearch::FValidInteractionSearch& ValidInteractionSearch);
+	static void OnInteractionEnd(UE::PoseSearch::FValidInteractionSearch& ValidInteractionSearch);
+
 	// @todo: use a thread local storage data structure to make publishing availabilities lockless to avoid the need of AvailabilitiesMapMutex
 	UE::PoseSearch::FAvailabilitiesMap AvailabilitiesMap;
 	FCriticalSection AvailabilitiesMapMutex;
@@ -170,6 +187,8 @@ private:
 	// array of groups of characters that needs to be anaylzed together for possible interactions
 	TArray<UE::PoseSearch::FInteractionIsland*> Islands;
 
+	TArray<UE::PoseSearch::FValidInteractionSearch> ValidInteractionSearches;
+
 #if UE_POSE_SEARCH_ENABLE_INTERACTION_SUBSYSTEM_CRITICAL_SECTION
 	// critical section to retrieve the subsystem in a thread safe manner
 	static FCriticalSection RetrieveSubsystemMutex;
