posesearch - fix for FPoseSearchDatabaseAnimationAssetBase::GetPlayLength when GetAnimationAsset() returns nullptr

diff --git a/Engine/Plugins/Animation/PoseSearch/Source/Runtime/Private/PoseSearchDatabase.cpp b/Engine/Plugins/Animation/PoseSearch/Source/Runtime/Private/PoseSearchDatabase.cpp
--- a/Engine/Plugins/Animation/PoseSearch/Source/Runtime/Private/PoseSearchDatabase.cpp
+++ b/Engine/Plugins/Animation/PoseSearch/Source/Runtime/Private/PoseSearchDatabase.cpp
@@ -233,8 +233,6 @@ float FPoseSearchDatabaseAnimationAssetBase::GetPlayLength() const
 	{
 		return AnimationAsset->GetPlayLength();
 	}
-
-	checkNoEntry();
 	return 0;
 }
 
