pose search - UPoseSearchInteractionLibrary::GetInteractingPawns

[RN] skip

diff --git a/Engine/Plugins/Animation/PoseSearch/Source/Editor/Private/PoseSearchDatabaseViewModel.cpp b/Engine/Plugins/Animation/PoseSearch/Source/Editor/Private/PoseSearchDatabaseViewModel.cpp
--- a/Engine/Plugins/Animation/PoseSearch/Source/Editor/Private/PoseSearchDatabaseViewModel.cpp
+++ b/Engine/Plugins/Animation/PoseSearch/Source/Editor/Private/PoseSearchDatabaseViewModel.cpp
@@ -684,8 +684,11 @@ bool FDatabaseViewModel::DeleteFromDatabase(int32 AnimationAssetIndex)
 					for (int32 NotifyIndex = AnimSequenceBase->Notifies.Num() - 1; NotifyIndex >= 0; --NotifyIndex)
 					{
 						const FAnimNotifyEvent& NotifyEvent = AnimSequenceBase->Notifies[NotifyIndex];
-						if (const UAnimNotifyState_PoseSearchBranchIn* PoseSearchBranchIn = Cast<UAnimNotifyState_PoseSearchBranchIn>(NotifyEvent.NotifyStateClass))
+						if (NotifyEvent.NotifyStateClass && NotifyEvent.NotifyStateClass->GetClass()->IsChildOf<UAnimNotifyState_PoseSearchBranchIn>())
 						{
+							const UAnimNotifyState_PoseSearchBranchIn* PoseSearchBranchIn = Cast<UAnimNotifyState_PoseSearchBranchIn>(NotifyEvent.NotifyStateClass);
+							check(PoseSearchBranchIn);
+
 							if (PoseSearchBranchIn->Database == Database && PoseSearchBranchIn->GetBranchInId() == DatabaseAnimationAssetBase->BranchInId)
 							{
 								if (!bModified)
diff --git a/Engine/Plugins/Animation/PoseSearch/Source/Runtime/Private/AnimNode_MotionMatchingInteraction.cpp b/Engine/Plugins/Animation/PoseSearch/Source/Runtime/Private/AnimNode_MotionMatchingInteraction.cpp
--- a/Engine/Plugins/Animation/PoseSearch/Source/Runtime/Private/AnimNode_MotionMatchingInteraction.cpp
+++ b/Engine/Plugins/Animation/PoseSearch/Source/Runtime/Private/AnimNode_MotionMatchingInteraction.cpp
@@ -245,7 +245,7 @@ void FAnimNode_MotionMatchingInteraction::Evaluate_AnyThread(FPoseContext& Outpu
 		// using UAnimNotifyState_PoseSearchWarping to change the WarpingAlignmentExtractionTime (time we sample the full alignment transforms at)
 		for (const FAnimNotifyEvent& NotifyEvent : Sampler.GetAllAnimNotifyEvents())
 		{
-			if (NotifyEvent.NotifyStateClass->GetClass()->IsChildOf<UAnimNotifyState_PoseSearchWarping>())
+			if (NotifyEvent.NotifyStateClass && NotifyEvent.NotifyStateClass->GetClass()->IsChildOf<UAnimNotifyState_PoseSearchWarping>())
 			{
 				const UAnimNotifyState_PoseSearchWarping* NotifyStatePoseSearchWarping = Cast<UAnimNotifyState_PoseSearchWarping>(NotifyEvent.NotifyStateClass);
 				check(NotifyStatePoseSearchWarping);
diff --git a/Engine/Plugins/Animation/PoseSearch/Source/Runtime/Private/PoseSearchInteractionLibrary.cpp b/Engine/Plugins/Animation/PoseSearch/Source/Runtime/Private/PoseSearchInteractionLibrary.cpp
--- a/Engine/Plugins/Animation/PoseSearch/Source/Runtime/Private/PoseSearchInteractionLibrary.cpp
+++ b/Engine/Plugins/Animation/PoseSearch/Source/Runtime/Private/PoseSearchInteractionLibrary.cpp
@@ -2,6 +2,7 @@
 
 #include "PoseSearch/PoseSearchInteractionLibrary.h"
 #include "Animation/AnimMontage.h"
+#include "GameFramework/Pawn.h"
 #include "PoseSearch/PoseSearchInteractionSubsystem.h"
 #include "Engine/World.h"
 
@@ -39,4 +40,36 @@ FPoseSearchContinuingProperties UPoseSearchInteractionLibrary::GetMontageContinu
 		ContinuingProperties.PlayingAssetAccumulatedTime = AnimMontageInstance->DeltaTimeRecord.GetPrevious();
 	}
 	return ContinuingProperties;
-}
\ No newline at end of file
+}
+
+TArray<APawn*> UPoseSearchInteractionLibrary::GetInteractingPawns(APawn* Pawn)
+{
+	using namespace UE::PoseSearch;
+
+	TArray<APawn*> InteractingPawns;
+	if (const USkeletalMeshComponent* SkeletalMeshComponent = Pawn->GetComponentByClass<USkeletalMeshComponent>())
+	{
+		if (const UAnimInstance* AnimInstance = SkeletalMeshComponent->GetAnimInstance())
+		{
+			if (const UPoseSearchInteractionSubsystem* InteractionSubsystem = UPoseSearchInteractionSubsystem::GetSubsystem_AnyThread(AnimInstance))
+			{
+				for (const UObject* InteractingAnimContext : InteractionSubsystem->GetInteractingAnimContextsWith(AnimInstance))
+				{
+					check(InteractingAnimContext);
+					if (const UAnimInstance* InteractingAnimInstance = Cast<UAnimInstance>(InteractingAnimContext))
+					{
+						if (APawn* InteractingPawn = Cast<APawn>(InteractingAnimInstance->GetOwningActor()))
+						{
+							InteractingPawns.Add(InteractingPawn);
+						}
+					}
+					else
+					{
+						unimplemented();
+					}
+				}
+			}
+		}
+	}
+	return InteractingPawns;
+}
diff --git a/Engine/Plugins/Animation/PoseSearch/Source/Runtime/Private/PoseSearchInteractionSubsystem.cpp b/Engine/Plugins/Animation/PoseSearch/Source/Runtime/Private/PoseSearchInteractionSubsystem.cpp
--- a/Engine/Plugins/Animation/PoseSearch/Source/Runtime/Private/PoseSearchInteractionSubsystem.cpp
+++ b/Engine/Plugins/Animation/PoseSearch/Source/Runtime/Private/PoseSearchInteractionSubsystem.cpp
@@ -1046,6 +1046,32 @@ void UPoseSearchInteractionSubsystem::UpdateValidInteractionSearches()
 	ValidInteractionSearches = NewValidInteractionSearches;
 }
 
+// @todo: add a debug critical section to avoid overlapping of GetInteractingAnimContextsWith with UpdateValidInteractionSearches to avoid / detecting threading issues!
+TArray<const UObject*> UPoseSearchInteractionSubsystem::GetInteractingAnimContextsWith(const UObject* AnimContext) const
+{
+	using namespace UE::PoseSearch;
+
+	TArray<const UObject*> InteractingAnimContexts;
+	for (const FValidInteractionSearch& ValidInteractionSearch : ValidInteractionSearches)
+	{
+		const int32 AnimContextIndex = ValidInteractionSearch.SearchContext.AnimContexts.IndexOfByKey(AnimContext);
+		if (AnimContextIndex != INDEX_NONE)
+		{
+			for (int32 Index = 0; Index < ValidInteractionSearch.SearchContext.AnimContexts.Num(); ++Index)
+			{
+				if (Index != AnimContextIndex)
+				{
+					if (const UObject* InteractingAnimContext = ValidInteractionSearch.SearchContext.AnimContexts[Index].Get())
+					{
+						InteractingAnimContexts.Add(InteractingAnimContext);
+					}
+				}
+			}
+		}
+	}
+	return InteractingAnimContexts;
+}
+
 void UPoseSearchInteractionSubsystem::Tick(float DeltaSeconds)
 {
 	QUICK_SCOPE_CYCLE_COUNTER(STAT_UPoseSearchInteractionSubsystem_Tick);
diff --git a/Engine/Plugins/Animation/PoseSearch/Source/Runtime/Public/PoseSearch/PoseSearchInteractionLibrary.h b/Engine/Plugins/Animation/PoseSearch/Source/Runtime/Public/PoseSearch/PoseSearchInteractionLibrary.h
--- a/Engine/Plugins/Animation/PoseSearch/Source/Runtime/Public/PoseSearch/PoseSearchInteractionLibrary.h
+++ b/Engine/Plugins/Animation/PoseSearch/Source/Runtime/Public/PoseSearch/PoseSearchInteractionLibrary.h
@@ -132,4 +132,8 @@ public:
 
 	UFUNCTION(Experimental, BlueprintPure, Category = "Animation|Pose Search", meta = (BlueprintThreadSafe))
 	static FPoseSearchContinuingProperties GetMontageContinuingProperties(UAnimInstance* AnimInstance);
+
+	// @todo: make it BlueprintThreadSafe
+	UFUNCTION(Experimental, BlueprintPure, Category = "Animation|Pose Search")
+	static TArray<APawn*> GetInteractingPawns(APawn* Pawn);
 };
\ No newline at end of file
diff --git a/Engine/Plugins/Animation/PoseSearch/Source/Runtime/Public/PoseSearch/PoseSearchInteractionSubsystem.h b/Engine/Plugins/Animation/PoseSearch/Source/Runtime/Public/PoseSearch/PoseSearchInteractionSubsystem.h
--- a/Engine/Plugins/Animation/PoseSearch/Source/Runtime/Public/PoseSearch/PoseSearchInteractionSubsystem.h
+++ b/Engine/Plugins/Animation/PoseSearch/Source/Runtime/Public/PoseSearch/PoseSearchInteractionSubsystem.h
@@ -145,6 +145,9 @@ public:
 		const FPoseSearchContinuingProperties& ContinuingProperties, FPoseSearchInteractionBlueprintResult& Result,
 		FName PoseHistoryName, const FAnimNode_PoseSearchHistoryCollector_Base* HistoryCollector, bool bValidateResultAgainstAvailabilities);
 
+	// returns and array containing all the other AnimContext(s) that the input AnimContext is currently interacting with
+	TArray<const UObject*> GetInteractingAnimContextsWith(const UObject* AnimContext) const;
+
 private:
 #if ENABLE_ANIM_DEBUG
 	friend struct UE::PoseSearch::FInteractionValidator;
