[AutoRTFM] Update thread-safety access detection in PoseSearch to use MT_ACCESS_DETECTOR.

These perform the same function and are transactionally safe.

diff --git a/Engine/Plugins/Animation/PoseSearch/Source/Runtime/Private/PoseSearchHistory.cpp b/Engine/Plugins/Animation/PoseSearch/Source/Runtime/Private/PoseSearchHistory.cpp
--- a/Engine/Plugins/Animation/PoseSearch/Source/Runtime/Private/PoseSearchHistory.cpp
+++ b/Engine/Plugins/Animation/PoseSearch/Source/Runtime/Private/PoseSearchHistory.cpp
@@ -621,9 +621,8 @@ FPoseHistory::FPoseHistory(FPoseHistory&& Other)
 FPoseHistory& FPoseHistory::operator=(const FPoseHistory& Other)
 {
 #if ENABLE_ANIM_DEBUG
-	FThreadSafeCounter& OtherPoseDataThreadSafeCounter = Other.PoseDataThreadSafeCounter;
-	CheckThreadSafetyWrite(OtherPoseDataThreadSafeCounter);
-	CheckThreadSafetyWrite(PoseDataThreadSafeCounter);
+	UE_MT_SCOPED_WRITE_ACCESS(Other.PoseDataThreadSafeCounter);
+	UE_MT_SCOPED_WRITE_ACCESS(PoseDataThreadSafeCounter);
 #endif // ENABLE_ANIM_DEBUG
 
 	MaxNumPoses = Other.MaxNumPoses;
@@ -640,9 +639,8 @@ FPoseHistory& FPoseHistory::operator=(const FPoseHistory& Other)
 FPoseHistory& FPoseHistory::operator=(FPoseHistory&& Other)
 {
 #if ENABLE_ANIM_DEBUG
-	FThreadSafeCounter& OtherPoseDataThreadSafeCounter = Other.PoseDataThreadSafeCounter;
-	CheckThreadSafetyWrite(OtherPoseDataThreadSafeCounter);
-	CheckThreadSafetyWrite(PoseDataThreadSafeCounter);
+	UE_MT_SCOPED_WRITE_ACCESS(Other.PoseDataThreadSafeCounter);
+	UE_MT_SCOPED_WRITE_ACCESS(PoseDataThreadSafeCounter);
 #endif // ENABLE_ANIM_DEBUG
 
 	MaxNumPoses = MoveTemp(Other.MaxNumPoses);
@@ -658,7 +656,7 @@ FPoseHistory& FPoseHistory::operator=(FPoseHistory&& Other)
 
 void FPoseHistory::Initialize_AnyThread(int32 InNumPoses, float InSamplingInterval)
 {
-	CheckThreadSafetyWrite(PoseDataThreadSafeCounter);
+	UE_MT_SCOPED_WRITE_ACCESS(PoseDataThreadSafeCounter);
 	check(InNumPoses >= 2);
 
 	MaxNumPoses = InNumPoses;
@@ -673,7 +671,7 @@ void FPoseHistory::Initialize_AnyThread(int32 InNumPoses, float InSamplingInterv
 
 bool FPoseHistory::GetTransformAtTime(float Time, FTransform& OutBoneTransform, const USkeleton* BoneIndexSkeleton, FBoneIndexType BoneIndexType, FBoneIndexType ReferenceBoneIndexType, bool bExtrapolate) const
 {
-	CheckThreadSafetyRead(PoseDataThreadSafeCounter);
+	UE_MT_SCOPED_READ_ACCESS(PoseDataThreadSafeCounter);
 
 	static_assert(RootBoneIndexType == 0 && ComponentSpaceIndexType == FBoneIndexType(-1) && WorldSpaceIndexType == FBoneIndexType(-2)); // some assumptions
 	
@@ -731,7 +729,7 @@ bool FPoseHistory::GetTransformAtTime(float Time, FTransform& OutBoneTransform,
 
 bool FPoseHistory::GetCurveValueAtTime(float Time, const FName& CurveName, float& OutCurveValue, bool bExtrapolate /*= false*/) const
 {
-	CheckThreadSafetyRead(PoseDataThreadSafeCounter);
+	UE_MT_SCOPED_READ_ACCESS(PoseDataThreadSafeCounter);
 
 	const int32 NumEntries = PoseData.Entries.Num();
 	if (NumEntries > 0)
@@ -758,43 +756,43 @@ bool FPoseHistory::GetCurveValueAtTime(float Time, const FName& CurveName, float
 
 const FPoseSearchQueryTrajectory& FPoseHistory::GetTrajectory() const
 {
-	CheckThreadSafetyRead(PoseDataThreadSafeCounter);
+	UE_MT_SCOPED_READ_ACCESS(PoseDataThreadSafeCounter);
 	return Trajectory;
 }
 
 float FPoseHistory::GetTrajectorySpeedMultiplier() const
 {
-	CheckThreadSafetyRead(PoseDataThreadSafeCounter);
+	UE_MT_SCOPED_READ_ACCESS(PoseDataThreadSafeCounter);
 	return TrajectorySpeedMultiplier;
 }
 
 bool FPoseHistory::IsEmpty() const
 {
-	CheckThreadSafetyRead(PoseDataThreadSafeCounter);
+	UE_MT_SCOPED_READ_ACCESS(PoseDataThreadSafeCounter);
 	return PoseData.Entries.IsEmpty();
 }
 
 const FBoneToTransformMap& FPoseHistory::GetBoneToTransformMap() const
 {
-	CheckThreadSafetyRead(PoseDataThreadSafeCounter);
+	UE_MT_SCOPED_READ_ACCESS(PoseDataThreadSafeCounter);
 	return PoseData.BoneToTransformMap;
 }
 
 const TConstArrayView<FName> FPoseHistory::GetCollectedCurves() const
 {
-	CheckThreadSafetyRead(PoseDataThreadSafeCounter);
+	UE_MT_SCOPED_READ_ACCESS(PoseDataThreadSafeCounter);
 	return MakeConstArrayView(PoseData.CollectedCurves);
 }
 
 int32 FPoseHistory::GetNumEntries() const
 {
-	CheckThreadSafetyRead(PoseDataThreadSafeCounter);
+	UE_MT_SCOPED_READ_ACCESS(PoseDataThreadSafeCounter);
 	return PoseData.Entries.Num();
 }
 
 const FPoseHistoryEntry& FPoseHistory::GetEntry(int32 EntryIndex) const
 {
-	CheckThreadSafetyRead(PoseDataThreadSafeCounter);
+	UE_MT_SCOPED_READ_ACCESS(PoseDataThreadSafeCounter);
 	return PoseData.Entries[EntryIndex];
 }
 
@@ -817,9 +815,9 @@ void FPoseHistory::SetTrajectory(const FPoseSearchQueryTrajectory& InTrajectory,
 {
 	if (!InTrajectory.Samples.IsEmpty())
 	{
-		CheckThreadSafetyWrite(PoseDataThreadSafeCounter);
+		UE_MT_SCOPED_WRITE_ACCESS(PoseDataThreadSafeCounter);
 
-		// @todo: THIS IS NOT THREAD SAFE! in the contex of multi character motion matching (CheckThreadSafetyWrite will assert in case of improper usage)
+		// @todo: THIS IS NOT THREAD SAFE! in the context of multi character motion matching (UE_MT_SCOPED_WRITE_ACCESS will assert in case of improper usage)
 		Trajectory = InTrajectory;
 
 		TrajectorySpeedMultiplier = InTrajectorySpeedMultiplier;
@@ -840,7 +838,7 @@ void FPoseHistory::EvaluateComponentSpace_AnyThread(float DeltaTime, IComponentS
 	bool bNeedsReset, bool bCacheBones, const TArray<FBoneIndexType>& RequiredBones,
 	const FBlendedCurve& Curves, const TConstArrayView<FName>& CollectedCurves)
 {
-	CheckThreadSafetyWrite(PoseDataThreadSafeCounter);
+	UE_MT_SCOPED_WRITE_ACCESS(PoseDataThreadSafeCounter);
 
 	check(MaxNumPoses >= 2);
 
@@ -952,7 +950,7 @@ void FPoseHistory::EvaluateComponentSpace_AnyThread(float DeltaTime, FCSPose<FCo
 #if ENABLE_DRAW_DEBUG && ENABLE_ANIM_DEBUG
 void FPoseHistory::DebugDraw(FAnimInstanceProxy& AnimInstanceProxy, FColor Color) const
 {
-	CheckThreadSafetyRead(PoseDataThreadSafeCounter);
+	UE_MT_SCOPED_READ_ACCESS(PoseDataThreadSafeCounter);
 
 	if (GVarAnimPoseHistoryDebugDrawTrajectory)
 	{
diff --git a/Engine/Plugins/Animation/PoseSearch/Source/Runtime/Public/PoseSearch/PoseSearchDefines.h b/Engine/Plugins/Animation/PoseSearch/Source/Runtime/Public/PoseSearch/PoseSearchDefines.h
--- a/Engine/Plugins/Animation/PoseSearch/Source/Runtime/Public/PoseSearch/PoseSearchDefines.h
+++ b/Engine/Plugins/Animation/PoseSearch/Source/Runtime/Public/PoseSearch/PoseSearchDefines.h
@@ -60,67 +60,6 @@ template<typename ElementType> constexpr int32 TAlignOf() { return TMax(alignof(
 template<typename ElementType> using TAlignedArray = TArray<ElementType, TAlignedHeapAllocator<TAlignOf<ElementType>()>>;
 template<typename ElementType> using TStackAlignedArray = TArray<ElementType, TMemStackAllocator<TAlignOf<ElementType>()>>;
 
-#if ENABLE_ANIM_DEBUG
-
-	template<bool bIsRead = 0>
-	struct FDebugThreadSafetyChecker
-	{
-		// if any read is in progress Counter > 0
-		// if any write is in progress Counter < 0
-		explicit FDebugThreadSafetyChecker(FThreadSafeCounter& InCounter)
-			: Counter(InCounter)
-		{
-			if (bIsRead)
-			{
-				const int32 PreviousCounter = Counter.Add(1);
-				if (PreviousCounter < 0)
-				{
-					UE_LOG(LogPoseSearch, Error, TEXT("Non thread safe call! Read detected a Write in progress"));
-				}
-			}
-			else
-			{
-				const int32 PreviousCounter = Counter.Set(-1);
-				if (PreviousCounter > 0)
-				{
-					UE_LOG(LogPoseSearch, Error, TEXT("Non thread safe call! Write detected %d Read(s) in progress"), PreviousCounter);
-				}
-			}
-		}
-
-		~FDebugThreadSafetyChecker()
-		{
-			if (bIsRead)
-			{
-				const int32 PreviousCounter = Counter.Add(-1);
-				if (PreviousCounter <= 0)
-				{
-					UE_LOG(LogPoseSearch, Error, TEXT("Non thread safe call! Read detected a Write in progress"));
-				}
-			}
-			else
-			{
-				const int32 PreviousCounter = Counter.Set(0);
-				if (PreviousCounter >= 0)
-				{
-					UE_LOG(LogPoseSearch, Error, TEXT("Non thread safe call! Write detected %d Read(s) in progress"), PreviousCounter);
-				}
-			}
-		}
-
-		FThreadSafeCounter& Counter;
-	};
-
-#define CheckThreadSafetyRead(ThreadSafeCounter) FDebugThreadSafetyChecker<1> DebugThreadSafetyChecker##ThreadSafeCounter(ThreadSafeCounter)
-#define CheckThreadSafetyWrite(ThreadSafeCounter) FDebugThreadSafetyChecker DebugThreadSafetyChecker##ThreadSafeCounter(ThreadSafeCounter)
-
-#else // ENABLE_ANIM_DEBUG
-
-	#define CheckThreadSafetyRead(ThreadSafeCounter)
-	#define CheckThreadSafetyWrite(ThreadSafeCounter)
-
-#endif // ENABLE_ANIM_DEBUG
-
 }
 
 
diff --git a/Engine/Plugins/Animation/PoseSearch/Source/Runtime/Public/PoseSearch/PoseSearchHistory.h b/Engine/Plugins/Animation/PoseSearch/Source/Runtime/Public/PoseSearch/PoseSearchHistory.h
--- a/Engine/Plugins/Animation/PoseSearch/Source/Runtime/Public/PoseSearch/PoseSearchHistory.h
+++ b/Engine/Plugins/Animation/PoseSearch/Source/Runtime/Public/PoseSearch/PoseSearchHistory.h
@@ -5,6 +5,7 @@
 #include "BonePose.h"
 #include "Containers/RingBuffer.h"
 #include "DrawDebugHelpers.h"
+#include "Misc/MTTransactionallySafeAccessDetector.h"
 #include "PoseSearch/PoseHistoryProvider.h"
 #include "PoseSearch/PoseSearchDefines.h"
 #include "PoseSearch/PoseSearchTrajectoryLibrary.h"
@@ -207,12 +208,7 @@ private:
 		TRingBuffer<FPoseHistoryEntry> Entries;
 	} PoseData;
 
-#if ENABLE_ANIM_DEBUG
-	
-	// used to analyze thread safety
-	mutable FThreadSafeCounter PoseDataThreadSafeCounter = 0;
-
-#endif // ENABLE_ANIM_DEBUG
+	UE_MT_DECLARE_TS_RW_ACCESS_DETECTOR(PoseDataThreadSafeCounter);
 };
 
 struct POSESEARCH_API FMemStackPoseHistory : public IPoseHistory
