posesearch - fix for DDC serialization for FPoseSearchDatabaseAnimationAssetBase::bDisableReselection

diff --git a/Engine/Plugins/Animation/PoseSearch/Source/Runtime/Public/PoseSearch/PoseSearchDatabase.h b/Engine/Plugins/Animation/PoseSearch/Source/Runtime/Public/PoseSearch/PoseSearchDatabase.h
--- a/Engine/Plugins/Animation/PoseSearch/Source/Runtime/Public/PoseSearch/PoseSearchDatabase.h
+++ b/Engine/Plugins/Animation/PoseSearch/Source/Runtime/Public/PoseSearch/PoseSearchDatabase.h
@@ -93,7 +93,7 @@ struct POSESEARCH_API FPoseSearchDatabaseAnimationAssetBase
 	bool bEnabled = true;
 
 	// if bDisableReselection is true, poses from the same asset cannot be reselected. Useful to avoid jumping on frames on the same looping animations
-	UPROPERTY(EditAnywhere, Category = "Settings", meta = (ExcludeFromHash, DisplayPriority = 2))
+	UPROPERTY(EditAnywhere, Category = "Settings", meta = (DisplayPriority = 2))
 	bool bDisableReselection = false;
 
 	// This allows users to set if this animation is original only (no mirrored data), original and mirrored, or only the mirrored version of this animation.
