posesearch - fix for InitSearchIndexAssets failing gracefully in case of assets with duplicated roles

diff --git a/Engine/Plugins/Animation/PoseSearch/Source/Runtime/Private/PoseSearchDerivedData.cpp b/Engine/Plugins/Animation/PoseSearch/Source/Runtime/Private/PoseSearchDerivedData.cpp
--- a/Engine/Plugins/Animation/PoseSearch/Source/Runtime/Private/PoseSearchDerivedData.cpp
+++ b/Engine/Plugins/Animation/PoseSearch/Source/Runtime/Private/PoseSearchDerivedData.cpp
@@ -420,6 +420,21 @@ static bool InitSearchIndexAssets(FSearchIndexBase& SearchIndex, const UPoseSear
 				continue;
 			}
 
+			// checking for duplicated roles in DatabaseAsset
+			TSet<FRole, DefaultKeyFuncs<FRole>, TInlineSetAllocator<PreallocatedRolesNum>> DatabaseAssetRoles;
+			const int32 NumRoles = DatabaseAsset->GetNumRoles();
+			for (int32 RoleIndex = 0; RoleIndex < NumRoles; ++RoleIndex)
+			{
+				bool bIsAlreadyInSet = false;
+				const FRole Role = DatabaseAsset->GetRole(RoleIndex);
+				DatabaseAssetRoles.Add(Role, &bIsAlreadyInSet);
+				if (bIsAlreadyInSet)
+				{
+					UE_LOG(LogPoseSearch, Error, TEXT("DatabaseAsset '%s' contains duplicate Role '%s'"), *DatabaseAsset->GetAnimationAsset()->GetName(), *Role.ToString());
+					bAnyErrors = true;
+				}
+			}
+
 			// checking for valid roles in DatabaseMultiAnimAsset against the Schema
 			bool bAreAllRolesSupported = true;
 			for (const FPoseSearchRoledSkeleton& RoledSkeleton : Schema->GetRoledSkeletons())
