posesearch - fix for crash in FPoseSearchTrajectoryData::UpdateData

diff --git a/Engine/Plugins/Animation/PoseSearch/Source/Runtime/Private/PoseSearchTrajectoryLibrary.cpp b/Engine/Plugins/Animation/PoseSearch/Source/Runtime/Private/PoseSearchTrajectoryLibrary.cpp
--- a/Engine/Plugins/Animation/PoseSearch/Source/Runtime/Private/PoseSearchTrajectoryLibrary.cpp
+++ b/Engine/Plugins/Animation/PoseSearch/Source/Runtime/Private/PoseSearchTrajectoryLibrary.cpp
@@ -35,6 +35,11 @@ bool FPoseSearchTrajectoryData::UpdateData(
 	if (!Character)
 	{
 		const UAnimInstance* AnimInstance = Cast<UAnimInstance>(Context);
+		if (!AnimInstance)
+		{
+			return false;
+		}
+
 		Character = Cast<ACharacter>(AnimInstance->GetOwningActor());
 		if (!Character)
 		{
