PoseSearch - logging error in case of no valid databases in UPoseSearchLibrary::UpdateMotionMatchingState

diff --git a/Engine/Plugins/Animation/PoseSearch/Source/Runtime/Private/PoseSearchLibrary.cpp b/Engine/Plugins/Animation/PoseSearch/Source/Runtime/Private/PoseSearchLibrary.cpp
--- a/Engine/Plugins/Animation/PoseSearch/Source/Runtime/Private/PoseSearchLibrary.cpp
+++ b/Engine/Plugins/Animation/PoseSearch/Source/Runtime/Private/PoseSearchLibrary.cpp
@@ -241,7 +241,7 @@ namespace UE::PoseSearch
 
 		for (const TObjectPtr<const UPoseSearchDatabase>& Database : Databases)
 		{
-			if (ensure(Database))
+			if (Database)
 			{
 				if (OneOfTheSchemas != Database->Schema)
 				{
@@ -640,7 +640,16 @@ void UPoseSearchLibrary::UpdateMotionMatchingState(
 
 	using namespace UE::PoseSearch;
 
-	if (Databases.IsEmpty())
+	bool bAnyValidDatabases = false;
+	for (const TObjectPtr<const UPoseSearchDatabase>& Database : Databases)
+	{
+		if (Database)
+		{
+			bAnyValidDatabases = true;
+			break;
+		}
+	}
+	if (!bAnyValidDatabases)
 	{
 		UE_LOG(LogPoseSearch, Error, TEXT("No database assets provided for motion matching"));
 		return;
@@ -680,7 +689,7 @@ void UPoseSearchLibrary::UpdateMotionMatchingState(
 
 		for (const TObjectPtr<const UPoseSearchDatabase>& Database : Databases)
 		{
-			if (ensure(Database))
+			if (Database)
 			{
 				const FSearchResult NewSearchResult = Database->Search(SearchContext);
 
