PoseSearch - fix for FKeyBuilder to account for POSESEARCHDB_DERIVEDDATA_VER_SMALL, POSESEARCHDB_DERIVEDDATA_VER, and UE::Anim::Compression::AnimationCompressionVersionString in its cache

diff --git a/Engine/Plugins/Animation/PoseSearch/Source/Runtime/Private/PoseSearchDerivedDataKey.cpp b/Engine/Plugins/Animation/PoseSearch/Source/Runtime/Private/PoseSearchDerivedDataKey.cpp
--- a/Engine/Plugins/Animation/PoseSearch/Source/Runtime/Private/PoseSearchDerivedDataKey.cpp
+++ b/Engine/Plugins/Animation/PoseSearch/Source/Runtime/Private/PoseSearchDerivedDataKey.cpp
@@ -54,8 +54,21 @@ FKeyBuilder::FKeyBuilder(const UObject* Object, bool bUseDataVer, bool bPerformC
 	PartialKeyHashes = InPartialKeyHashes;
 	DebugPartialKeyHashesMode = InDebugPartialKeyHashesMode;
 
-	if (bUseDataVer)
+	// FKeyBuilder is a saving only archiver, and since it doesn't modify the input Object it's safe to do a const_cast 
+	UObject* NonConstObject = const_cast<UObject*>(Object);
+	*this << NonConstObject;
+
+	while (!ObjectsToSerialize.IsEmpty() && !bAnyAssetNotReady)
 	{
+		SerializeObjectInternal(ObjectsToSerialize.Pop(EAllowShrinking::No));
+	}
+
+	if (bUseDataVer && !bAnyAssetNotReady)
+	{
+		FLocalPartialKeyHash& LocalCachedHash = LocalPartialKeyHashes.AddDefaulted_GetRef();
+
+		Hasher.Reset();
+
 		// used to invalidate the key without having to change POSESEARCHDB_DERIVEDDATA_VER all the times
 		int32 POSESEARCHDB_DERIVEDDATA_VER_SMALL = 2659;
 		FGuid VersionGuid = FDevSystemGuids::GetSystemGuid(FDevSystemGuids::Get().POSESEARCHDB_DERIVEDDATA_VER);
@@ -64,15 +77,8 @@ FKeyBuilder::FKeyBuilder(const UObject* Object, bool bUseDataVer, bool bPerformC
 		*this << VersionGuid;
 		*this << AnimationCompressionVersionString;
 		*this << POSESEARCHDB_DERIVEDDATA_VER_SMALL;
-	}
-
-	// FKeyBuilder is a saving only archiver, and since it doesn't modify the input Object it's safe to do a const_cast 
-	UObject* NonConstObject = const_cast<UObject*>(Object);
-	*this << NonConstObject;
 
-	while (!ObjectsToSerialize.IsEmpty() && !bAnyAssetNotReady)
-	{
-		SerializeObjectInternal(ObjectsToSerialize.Pop(EAllowShrinking::No));
+		LocalCachedHash.Hash = Hasher.Finalize();
 	}
 }
 
