PoseSearch - removing 5.4 deprecations and fixing UPoseSearchSchema 5.5 deprecations

diff --git a/Engine/Plugins/Animation/PoseSearch/Source/Runtime/Private/AnimNode_MotionMatching.cpp b/Engine/Plugins/Animation/PoseSearch/Source/Runtime/Private/AnimNode_MotionMatching.cpp
--- a/Engine/Plugins/Animation/PoseSearch/Source/Runtime/Private/AnimNode_MotionMatching.cpp
+++ b/Engine/Plugins/Animation/PoseSearch/Source/Runtime/Private/AnimNode_MotionMatching.cpp
@@ -72,31 +72,6 @@ void FAnimNode_MotionMatching::Evaluate_AnyThread(FPoseContext& Output)
 		RootMotionTransformDelta = FTransform::Identity;
 		RootMotionProvider = nullptr;
 	}
-
-PRAGMA_DISABLE_DEPRECATION_WARNINGS
-	// applying MotionMatchingState.ComponentDeltaYaw (considered as root bone delta yaw) to the root bone and the root motion delta transform
-	if (!FMath::IsNearlyZero(MotionMatchingState.ComponentDeltaYaw))
-	{
-		const FQuat RootBoneDelta(FRotator(0.f, MotionMatchingState.ComponentDeltaYaw, 0.f));
-		FCompactPoseBoneIndex RootBoneIndex(RootBoneIndexType);
-
-		// @todo: this code assumes the root bone from animation is always identity. Implement if not the case as
-		//		Output.Pose[RootBoneIndex].SetRotation(Output.Pose[RootBoneIndex].GetRotation() * RootBoneDelta);
-		//		Output.Pose[RootBoneIndex].NormalizeRotation();
-		// 		etc etc
-		check(Output.Pose[RootBoneIndex].GetRotation().IsIdentity());
-
-		Output.Pose[RootBoneIndex].SetRotation(RootBoneDelta);
-		RootMotionTransformDelta.SetTranslation(RootBoneDelta.RotateVector(RootMotionTransformDelta.GetTranslation()));
-
-		if (RootMotionProvider)
-		{
-			RootMotionProvider->OverrideRootMotion(RootMotionTransformDelta, Output.CustomAttributes);
-		}
-	}
-
-	MotionMatchingState.AnimationDeltaYaw = FRotator(RootMotionTransformDelta.GetRotation()).Yaw;
-PRAGMA_ENABLE_DEPRECATION_WARNINGS
 }
 
 void FAnimNode_MotionMatching::UpdateAssetPlayer(const FAnimationUpdateContext& Context)
diff --git a/Engine/Plugins/Animation/PoseSearch/Source/Runtime/Private/PoseSearchDatabase.cpp b/Engine/Plugins/Animation/PoseSearch/Source/Runtime/Private/PoseSearchDatabase.cpp
--- a/Engine/Plugins/Animation/PoseSearch/Source/Runtime/Private/PoseSearchDatabase.cpp
+++ b/Engine/Plugins/Animation/PoseSearch/Source/Runtime/Private/PoseSearchDatabase.cpp
@@ -707,55 +707,12 @@ int32 UPoseSearchDatabase::GetPoseIndexFromTime(float Time, const UE::PoseSearch
 
 void UPoseSearchDatabase::AddAnimationAsset(FInstancedStruct AnimationAsset)
 {
-	PRAGMA_DISABLE_DEPRECATION_WARNINGS
 	AnimationAssets.Add(AnimationAsset);
-	PRAGMA_ENABLE_DEPRECATION_WARNINGS
 }
 
 void UPoseSearchDatabase::RemoveAnimationAssetAt(int32 AnimationAssetIndex)
 {
-	PRAGMA_DISABLE_DEPRECATION_WARNINGS
 	AnimationAssets.RemoveAt(AnimationAssetIndex);
-	PRAGMA_ENABLE_DEPRECATION_WARNINGS
-}
-
-const TArray<FInstancedStruct>& UPoseSearchDatabase::GetAnimationAssets() const
-{
-	PRAGMA_DISABLE_DEPRECATION_WARNINGS
-	return AnimationAssets;
-	PRAGMA_ENABLE_DEPRECATION_WARNINGS
-}
-
-const FInstancedStruct& UPoseSearchDatabase::GetAnimationAssetStruct(int32 AnimationAssetIndex) const
-{
-	PRAGMA_DISABLE_DEPRECATION_WARNINGS
-	check(AnimationAssets.IsValidIndex(AnimationAssetIndex));
-	return AnimationAssets[AnimationAssetIndex];
-	PRAGMA_ENABLE_DEPRECATION_WARNINGS
-}
-
-const FInstancedStruct& UPoseSearchDatabase::GetAnimationAssetStruct(const UE::PoseSearch::FSearchIndexAsset& SearchIndexAsset) const
-{
-	PRAGMA_DISABLE_DEPRECATION_WARNINGS
-	check(AnimationAssets.IsValidIndex(SearchIndexAsset.GetSourceAssetIdx()));
-	return AnimationAssets[SearchIndexAsset.GetSourceAssetIdx()];
-	PRAGMA_ENABLE_DEPRECATION_WARNINGS
-}
-
-FInstancedStruct& UPoseSearchDatabase::GetMutableAnimationAssetStruct(int32 AnimationAssetIndex)
-{
-	PRAGMA_DISABLE_DEPRECATION_WARNINGS
-	check(AnimationAssets.IsValidIndex(AnimationAssetIndex));
-	return AnimationAssets[AnimationAssetIndex];
-	PRAGMA_ENABLE_DEPRECATION_WARNINGS
-}
-
-FInstancedStruct& UPoseSearchDatabase::GetMutableAnimationAssetStruct(const UE::PoseSearch::FSearchIndexAsset& SearchIndexAsset)
-{
-	PRAGMA_DISABLE_DEPRECATION_WARNINGS
-	check(AnimationAssets.IsValidIndex(SearchIndexAsset.GetSourceAssetIdx()));
-	return AnimationAssets[SearchIndexAsset.GetSourceAssetIdx()];
-	PRAGMA_ENABLE_DEPRECATION_WARNINGS
 }
 
 #if WITH_EDITOR
@@ -781,33 +738,7 @@ void UPoseSearchDatabase::PostLoad()
 #if WITH_EDITOR
 	using namespace UE::PoseSearch;
 
-	bool bRequiresSynchronization = false;
-
-	for (int32 AnimationAssetIndex = 0; AnimationAssetIndex < GetNumAnimationAssets(); ++AnimationAssetIndex)
-	{
-		if (FPoseSearchDatabaseAnimationAssetBase* AnimationAssetBase = GetMutableDatabaseAnimationAsset<FPoseSearchDatabaseAnimationAssetBase>(AnimationAssetIndex))
-		{
-			if (AnimationAssetBase->bSynchronizeWithExternalDependency_DEPRECATED)
-			{
-				Modify();
-				bRequiresSynchronization = true;
-				break;
-			}
-		}
-	}
-	
-	if (bRequiresSynchronization)
-	{
-		SynchronizeWithExternalDependencies();
-
-		for (int32 AnimationAssetIndex = 0; AnimationAssetIndex < GetNumAnimationAssets(); ++AnimationAssetIndex)
-		{
-			if (FPoseSearchDatabaseAnimationAssetBase* AnimationAssetBase = GetMutableDatabaseAnimationAsset<FPoseSearchDatabaseAnimationAssetBase>(AnimationAssetIndex))
-			{
-				AnimationAssetBase->bSynchronizeWithExternalDependency_DEPRECATED = false;
-			}
-		}
-	}
+	// todo: should we SynchronizeWithExternalDependencies() here?
 
 	ERequestAsyncBuildFlag Flag = ERequestAsyncBuildFlag::NewRequest;
 #if WITH_ENGINE
@@ -831,14 +762,11 @@ bool UPoseSearchDatabase::Contains(const UObject* Object) const
 
 int32 UPoseSearchDatabase::GetNumAnimationAssets() const
 {
-	PRAGMA_DISABLE_DEPRECATION_WARNINGS
 	return AnimationAssets.Num();
-	PRAGMA_ENABLE_DEPRECATION_WARNINGS
 }
 
 UObject* UPoseSearchDatabase::GetAnimationAsset(int32 Index) const
 {
-	PRAGMA_DISABLE_DEPRECATION_WARNINGS
 	if (AnimationAssets.IsValidIndex(Index))
 	{
 		if (const FPoseSearchDatabaseAnimationAssetBase* AnimationAssetBase = AnimationAssets[Index].GetPtr<FPoseSearchDatabaseAnimationAssetBase>())
@@ -846,7 +774,6 @@ UObject* UPoseSearchDatabase::GetAnimationAsset(int32 Index) const
 			return AnimationAssetBase->GetAnimationAsset();
 		}
 	}
-	PRAGMA_ENABLE_DEPRECATION_WARNINGS
 	return nullptr;
 }
 
@@ -901,7 +828,6 @@ void UPoseSearchDatabase::SynchronizeWithExternalDependencies(TConstArrayView<UA
 
 	// collecting all the database AnimationAsset(s) that don't require synchronization
 	TArray<bool> DisableReselection;
-	PRAGMA_DISABLE_DEPRECATION_WARNINGS
 	DisableReselection.Reserve(GetNumAnimationAssets());
 
 	for (FInstancedStruct& AnimationAsset : AnimationAssets)
@@ -993,8 +919,6 @@ void UPoseSearchDatabase::SynchronizeWithExternalDependencies(TConstArrayView<UA
 		bModified = true;
 	}
 
-	PRAGMA_ENABLE_DEPRECATION_WARNINGS
-
 	if (bModified)
 	{
 		Modify();
@@ -1020,7 +944,6 @@ bool UPoseSearchDatabase::IsCachedCookedPlatformDataLoaded(const ITargetPlatform
 #if WITH_EDITOR && ENABLE_ANIM_DEBUG
 void UPoseSearchDatabase::TestSynchronizeWithExternalDependencies()
 {
-	PRAGMA_DISABLE_DEPRECATION_WARNINGS
 	TArray<FInstancedStruct> AnimationAssetsCopy = AnimationAssets;
 	SynchronizeWithExternalDependencies();
 
@@ -1029,7 +952,6 @@ void UPoseSearchDatabase::TestSynchronizeWithExternalDependencies()
 		UE_LOG(LogPoseSearch, Error, TEXT("TestSynchronizeWithExternalDependencies failed"));
 		AnimationAssets = AnimationAssetsCopy;
 	}
-	PRAGMA_ENABLE_DEPRECATION_WARNINGS
 }
 #endif // WITH_EDITOR && ENABLE_ANIM_DEBUG
 
diff --git a/Engine/Plugins/Animation/PoseSearch/Source/Runtime/Private/PoseSearchLibrary.cpp b/Engine/Plugins/Animation/PoseSearch/Source/Runtime/Private/PoseSearchLibrary.cpp
--- a/Engine/Plugins/Animation/PoseSearch/Source/Runtime/Private/PoseSearchLibrary.cpp
+++ b/Engine/Plugins/Animation/PoseSearch/Source/Runtime/Private/PoseSearchLibrary.cpp
@@ -302,9 +302,6 @@ void FMotionMatchingState::Reset(const FTransform& ComponentTransform)
 {
 PRAGMA_DISABLE_DEPRECATION_WARNINGS
 	bJumpedToPose = false;
-	ComponentDeltaYaw = 0.f;
-	ComponentWorldYaw = FRotator(ComponentTransform.GetRotation()).Yaw;
-	AnimationDeltaYaw = 0.f;
 PRAGMA_ENABLE_DEPRECATION_WARNINGS
 
 	Reset();
@@ -1308,79 +1305,6 @@ UE::PoseSearch::FSearchResult UPoseSearchLibrary::MotionMatch(
 	return MotionMatch(AnimContexts, Roles, PoseHistories, AssetsToSearch, ContinuingProperties, Future);
 }
 
-void UPoseSearchLibrary::MotionMatch(
-	const TArrayView<UAnimInstance*> AnimInstances,
-	const TArrayView<const UE::PoseSearch::FRole> Roles,
-	const TArrayView<const UObject*> AssetsToSearch,
-	const FName PoseHistoryName,
-	const FPoseSearchContinuingProperties& ContinuingProperties,
-	const FPoseSearchFutureProperties& Future,
-	FPoseSearchBlueprintResult& Result,
-	const int32 DebugSessionUniqueIdentifier)
-{
-	PRAGMA_DISABLE_DEPRECATION_WARNINGS
-	MotionMatch(AnimInstances, Roles, AssetsToSearch, PoseHistoryName, ContinuingProperties, Future, Result);
-	PRAGMA_ENABLE_DEPRECATION_WARNINGS
-}
-
-UE::PoseSearch::FSearchResult UPoseSearchLibrary::MotionMatch(
-	const TArrayView<UAnimInstance*> AnimInstances,
-	const TArrayView<const UE::PoseSearch::FRole> Roles,
-	const TArrayView<const UE::PoseSearch::IPoseHistory*> PoseHistories,
-	const TArrayView<const UObject*> AssetsToSearch,
-	const FPoseSearchContinuingProperties& ContinuingProperties,
-	const FPoseSearchFutureProperties& Future,
-	const int32 DebugSessionUniqueIdentifier)
-{
-	TArray<const UObject*, TInlineAllocator<UE::PoseSearch::PreallocatedRolesNum>> AnimContexts;
-	AnimContexts.Reserve(AnimInstances.Num());
-
-	for (UAnimInstance* AnimInstance : AnimInstances)
-	{
-		AnimContexts.Add(AnimInstance);
-	}
-
-	return MotionMatch(AnimContexts, Roles, PoseHistories, AssetsToSearch, ContinuingProperties, Future);
-}
-
-UE::PoseSearch::FSearchResult UPoseSearchLibrary::MotionMatch(
-	const FAnimationBaseContext& Context,
-	TArrayView<const UObject*> AssetsToSearch,
-	const FPoseSearchContinuingProperties& ContinuingProperties)
-{
-	using namespace UE::PoseSearch;
-
-	const IPoseHistory* PoseHistory = nullptr;
-	if (FPoseHistoryProvider* PoseHistoryProvider = Context.GetMessage<FPoseHistoryProvider>())
-	{
-		PoseHistory = &PoseHistoryProvider->GetPoseHistory();
-	}
-
-	const UObject* AnimContext = Context.AnimInstanceProxy->GetAnimInstanceObject();
-	return MotionMatch(MakeArrayView(&AnimContext, 1), MakeArrayView(&DefaultRole, 1), MakeArrayView(&PoseHistory, 1),
-		AssetsToSearch, ContinuingProperties, FPoseSearchFutureProperties());
-}
-		
-UE::PoseSearch::FSearchResult UPoseSearchLibrary::MotionMatch(
-	TArrayView<UAnimInstance*> AnimInstances,
-	TArrayView<const UE::PoseSearch::FRole> Roles,
-	TArrayView<const UE::PoseSearch::IPoseHistory*> PoseHistories,
-	TArrayView<const UObject*> AssetsToSearch,
-	const FPoseSearchContinuingProperties& ContinuingProperties,
-	const int32 DebugSessionUniqueIdentifier,
-	float DesiredPermutationTimeOffset)
-{
-	TArray<const UObject*, TInlineAllocator<UE::PoseSearch::PreallocatedRolesNum>> AnimContexts;
-	AnimContexts.Reserve(AnimInstances.Num());
-
-	for (UAnimInstance* AnimInstance : AnimInstances)
-	{
-		AnimContexts.Add(AnimInstance);
-	}
-
-	return MotionMatch(AnimContexts, Roles, PoseHistories, AssetsToSearch, ContinuingProperties, FPoseSearchFutureProperties());
-}
-
 // End deprecated signatures
 ///////////////////////////////////////////////////////////////////////////////////////////
 
diff --git a/Engine/Plugins/Animation/PoseSearch/Source/Runtime/Private/PoseSearchSchema.cpp b/Engine/Plugins/Animation/PoseSearch/Source/Runtime/Private/PoseSearchSchema.cpp
--- a/Engine/Plugins/Animation/PoseSearch/Source/Runtime/Private/PoseSearchSchema.cpp
+++ b/Engine/Plugins/Animation/PoseSearch/Source/Runtime/Private/PoseSearchSchema.cpp
@@ -367,6 +367,7 @@ void UPoseSearchSchema::PostLoad()
 {
 	Super::PostLoad();
 
+	PRAGMA_DISABLE_DEPRECATION_WARNINGS
 	if (Skeleton_DEPRECATED)
 	{
 		Skeletons.AddDefaulted_GetRef().Skeleton = Skeleton_DEPRECATED;
@@ -385,6 +386,7 @@ void UPoseSearchSchema::PostLoad()
 		}
 		MirrorDataTable_DEPRECATED = nullptr;
 	}
+	PRAGMA_ENABLE_DEPRECATION_WARNINGS
 
 	for (FPoseSearchRoledSkeleton& Skeleton : Skeletons)
 	{
diff --git a/Engine/Plugins/Animation/PoseSearch/Source/Runtime/Public/PoseSearch/PoseHistoryProvider.h b/Engine/Plugins/Animation/PoseSearch/Source/Runtime/Public/PoseSearch/PoseHistoryProvider.h
--- a/Engine/Plugins/Animation/PoseSearch/Source/Runtime/Public/PoseSearch/PoseHistoryProvider.h
+++ b/Engine/Plugins/Animation/PoseSearch/Source/Runtime/Public/PoseSearch/PoseHistoryProvider.h
@@ -16,8 +16,6 @@ class POSESEARCH_API FPoseHistoryProvider : public UE::Anim::IGraphMessage
 	DECLARE_ANIMGRAPH_MESSAGE(FPoseHistoryProvider);
 
 public:
-	UE_DEPRECATED(5.4, "Use FPoseHistoryProvider(const FAnimNode_PoseSearchHistoryCollector_Base* HistoryCollector) instead")
-	FPoseHistoryProvider(const IPoseHistory& InPoseHistory) : PoseHistory(&InPoseHistory), HistoryCollector(nullptr) { }
 	FPoseHistoryProvider(const FAnimNode_PoseSearchHistoryCollector_Base* InHistoryCollector) : PoseHistory(nullptr), HistoryCollector(InHistoryCollector) { }
 	const IPoseHistory& GetPoseHistory() const;
 	const FAnimNode_PoseSearchHistoryCollector_Base* GetHistoryCollector() const { return HistoryCollector; }
diff --git a/Engine/Plugins/Animation/PoseSearch/Source/Runtime/Public/PoseSearch/PoseSearchDatabase.h b/Engine/Plugins/Animation/PoseSearch/Source/Runtime/Public/PoseSearch/PoseSearchDatabase.h
--- a/Engine/Plugins/Animation/PoseSearch/Source/Runtime/Public/PoseSearch/PoseSearchDatabase.h
+++ b/Engine/Plugins/Animation/PoseSearch/Source/Runtime/Public/PoseSearch/PoseSearchDatabase.h
@@ -86,7 +86,7 @@ struct POSESEARCH_API FPoseSearchDatabaseAnimationAssetBase
 	virtual int64 GetEditorMemSize() const;
 	virtual int64 GetApproxCookedSize() const { return GetEditorMemSize(); }
 
-	bool IsSynchronizedWithExternalDependency() const { return bSynchronizeWithExternalDependency_DEPRECATED || BranchInId != 0; }
+	bool IsSynchronizedWithExternalDependency() const { return BranchInId != 0; }
 
 	// This allows users to enable or exclude animations from this database. Useful for debugging.
 	UPROPERTY(EditAnywhere, Category = "Settings", meta = (DisplayPriority = 1))
@@ -101,12 +101,7 @@ struct POSESEARCH_API FPoseSearchDatabaseAnimationAssetBase
 	UPROPERTY(EditAnywhere, Category = "Settings", meta = (DisplayPriority = 3))
 	EPoseSearchMirrorOption MirrorOption = EPoseSearchMirrorOption::UnmirroredOnly;
 
-	// SynchronizeWithExternalDependency is true when this asset has been added via SynchronizeWithExternalDependencies.
-	// To delete it, remove the PoseSearchBranchIn notify state
-	UPROPERTY()
-	bool bSynchronizeWithExternalDependency_DEPRECATED = false;
-
-	// SynchronizeWithExternalDependency is true when this asset has been added via SynchronizeWithExternalDependencies.
+	// BranchInId != 0 when this asset has been added via SynchronizeWithExternalDependencies.
 	// To delete it, remove the PoseSearchBranchIn notify state
 	UPROPERTY(VisibleAnywhere, Category = "Settings", meta = (DisplayPriority = 20))
 	uint32 BranchInId = 0;
@@ -118,7 +113,6 @@ struct POSESEARCH_API FPoseSearchDatabaseAnimationAssetBase
 		return A.bEnabled == B.bEnabled &&
 			A.bDisableReselection == B.bDisableReselection &&
 			A.MirrorOption == B.MirrorOption &&
-			A.bSynchronizeWithExternalDependency_DEPRECATED == B.bSynchronizeWithExternalDependency_DEPRECATED &&
 			A.BranchInId == B.BranchInId;
 #else // WITH_EDITORONLY_DATA
 		return true;
@@ -432,11 +426,11 @@ public:
 	FFloatInterval AdditionalExtrapolationTime = FFloatInterval(-100.f, 100.f);
 #endif // WITH_EDITORONLY_DATA
 
-	// use GetNumAnimationAssets() and GetDatabaseAnimationAsset(...) APIs to query for AnimationAssets
-	UE_DEPRECATED(5.4, "This property will be made private")
+private:
 	UPROPERTY()
 	TArray<FInstancedStruct> AnimationAssets;
 
+public:
 	/** Array of tags that can be used as metadata. */
 	UPROPERTY(EditAnywhere, BlueprintReadOnly, Category = "Database")
 	TArray<FName> Tags;
@@ -523,27 +517,6 @@ public:
 	template<typename TDatabaseAnimationAsset> TDatabaseAnimationAsset* GetMutableDatabaseAnimationAsset(int32 AnimationAssetIndex);
 	template<typename TDatabaseAnimationAsset> TDatabaseAnimationAsset* GetMutableDatabaseAnimationAsset(const UE::PoseSearch::FSearchIndexAsset& SearchIndexAsset);
 
-	UE_DEPRECATED(5.4, "Use GetNumAnimationAssets to iterate over GetDatabaseAnimationAsset instead")
-	const TArray<FInstancedStruct>& GetAnimationAssets() const;
-
-	UE_DEPRECATED(5.4, "Use GetDatabaseAnimationAsset instead")
-	const FInstancedStruct& GetAnimationAssetStruct(int32 AnimationAssetIndex) const;
-	UE_DEPRECATED(5.4, "Use GetDatabaseAnimationAsset instead")
-	const FInstancedStruct& GetAnimationAssetStruct(const UE::PoseSearch::FSearchIndexAsset& SearchIndexAsset) const;
-	UE_DEPRECATED(5.4, "Use GetMutableDatabaseAnimationAsset instead")
-	FInstancedStruct& GetMutableAnimationAssetStruct(int32 AnimationAssetIndex);
-	UE_DEPRECATED(5.4, "Use GetMutableDatabaseAnimationAsset instead")
-	FInstancedStruct& GetMutableAnimationAssetStruct(const UE::PoseSearch::FSearchIndexAsset& SearchIndexAsset);
-	
-	UE_DEPRECATED(5.4, "Use GetDatabaseAnimationAsset instead")
-	const FPoseSearchDatabaseAnimationAssetBase* GetAnimationAssetBase(int32 AnimationAssetIndex) const { return GetDatabaseAnimationAsset<FPoseSearchDatabaseAnimationAssetBase>(AnimationAssetIndex); }
-	UE_DEPRECATED(5.4, "Use GetDatabaseAnimationAsset instead")
-	const FPoseSearchDatabaseAnimationAssetBase* GetAnimationAssetBase(const UE::PoseSearch::FSearchIndexAsset& SearchIndexAsset) const { return GetDatabaseAnimationAsset<FPoseSearchDatabaseAnimationAssetBase>(SearchIndexAsset); }
-	UE_DEPRECATED(5.4, "Use GetMutableDatabaseAnimationAsset instead")
-	FPoseSearchDatabaseAnimationAssetBase* GetMutableAnimationAssetBase(int32 AnimationAssetIndex) { return GetMutableDatabaseAnimationAsset<FPoseSearchDatabaseAnimationAssetBase>(AnimationAssetIndex); }
-	UE_DEPRECATED(5.4, "Use GetMutableDatabaseAnimationAsset instead")
-	FPoseSearchDatabaseAnimationAssetBase* GetMutableAnimationAssetBase(const UE::PoseSearch::FSearchIndexAsset& SearchIndexAsset) { return GetMutableDatabaseAnimationAsset<FPoseSearchDatabaseAnimationAssetBase>(SearchIndexAsset); }
-
 	float GetRealAssetTime(int32 PoseIdx) const;
 	float GetNormalizedAssetTime(int32 PoseIdx) const;
 
@@ -614,12 +587,10 @@ private:
 template<typename TDatabaseAnimationAsset>
 inline const TDatabaseAnimationAsset* UPoseSearchDatabase::GetDatabaseAnimationAsset(int32 AnimationAssetIndex) const
 {
-	PRAGMA_DISABLE_DEPRECATION_WARNINGS
 	if (AnimationAssets.IsValidIndex(AnimationAssetIndex))
 	{
 		return AnimationAssets[AnimationAssetIndex].GetPtr<TDatabaseAnimationAsset>();
 	}
-	PRAGMA_ENABLE_DEPRECATION_WARNINGS
 	return nullptr;
 }
 
@@ -632,12 +603,10 @@ inline const TDatabaseAnimationAsset* UPoseSearchDatabase::GetDatabaseAnimationA
 template<typename TDatabaseAnimationAsset>
 inline TDatabaseAnimationAsset* UPoseSearchDatabase::GetMutableDatabaseAnimationAsset(int32 AnimationAssetIndex)
 {
-	PRAGMA_DISABLE_DEPRECATION_WARNINGS
 	if (AnimationAssets.IsValidIndex(AnimationAssetIndex))
 	{
 		return AnimationAssets[AnimationAssetIndex].GetMutablePtr<TDatabaseAnimationAsset>();
 	}
-	PRAGMA_ENABLE_DEPRECATION_WARNINGS
 	return nullptr;
 }
 
diff --git a/Engine/Plugins/Animation/PoseSearch/Source/Runtime/Public/PoseSearch/PoseSearchLibrary.h b/Engine/Plugins/Animation/PoseSearch/Source/Runtime/Public/PoseSearch/PoseSearchLibrary.h
--- a/Engine/Plugins/Animation/PoseSearch/Source/Runtime/Public/PoseSearch/PoseSearchLibrary.h
+++ b/Engine/Plugins/Animation/PoseSearch/Source/Runtime/Public/PoseSearch/PoseSearchLibrary.h
@@ -75,23 +75,6 @@ PRAGMA_ENABLE_DEPRECATION_WARNINGS
 	bool bJumpedToPose = false;
 
 	UE::PoseSearch::FPoseIndicesHistory PoseIndicesHistory;
-
-	// Component delta yaw (also considered as root bone delta yaw)
-	UE_DEPRECATED(5.4, "Use Steering, OrientationWarping, OffsetRootBone nodes instead")
-	float ComponentDeltaYaw = 0.f;
-
-	// Internal component yaw in world space. Initialized as FRotator(AnimInstanceProxy->GetComponentTransform().GetRotation()).Yaw, but then integrated by ComponentDeltaYaw
-	UE_DEPRECATED(5.4, "Use Steering, OrientationWarping, OffsetRootBone nodes instead")
-	float ComponentWorldYaw = 0.f;
-	
-	// RootMotionTransformDelta yaw at the end of FAnimNode_MotionMatching::Evaluate_AnyThread (it represents the previous frame animation delta yaw)
-	UE_DEPRECATED(5.4, "Use Steering, OrientationWarping, OffsetRootBone nodes instead")
-	float AnimationDeltaYaw = 0.f;
-
-#if UE_POSE_SEARCH_TRACE_ENABLED
-	UE_DEPRECATED(5.4, "Debug RootMotionDelta is now calculated in-place.")
-	FTransform RootMotionTransformDelta = FTransform::Identity;
-#endif //UE_POSE_SEARCH_TRACE_ENABLED
 };
 
 USTRUCT(Experimental, BlueprintType, Category="Animation|Pose Search")
@@ -170,20 +153,6 @@ class POSESEARCH_API UPoseSearchLibrary : public UBlueprintFunctionLibrary
 		TraceMotionMatching(SearchContext, CurrentResult, ElapsedPoseSearchTime, DeltaTime, bSearch, 1.f);
 	}
 
-	UE_DEPRECATED(5.4, "Use TraceMotionMatching instead")
-	static void TraceMotionMatchingState(
-		UE::PoseSearch::FSearchContext& SearchContext,
-		const UE::PoseSearch::FSearchResult& CurrentResult,
-		float ElapsedPoseSearchTime,
-		const FTransform& RootMotionTransformDelta,
-		int32 NodeId,
-		float DeltaTime,
-		bool bSearch,
-		float RecordingTime)
-	{
-		TraceMotionMatching(SearchContext, CurrentResult, ElapsedPoseSearchTime, DeltaTime, bSearch, 1.f);
-	}
-	
 #endif // UE_POSE_SEARCH_TRACE_ENABLED
 
 public:
@@ -282,43 +251,6 @@ public:
 		const FPoseSearchContinuingProperties& ContinuingProperties,
 		const FPoseSearchFutureProperties& Future);
 
-	UE_DEPRECATED(5.4, "Use other MotionMatch signatures instead")
-	static void MotionMatch(
-		const TArrayView<UAnimInstance*> AnimInstances,
-		const TArrayView<const UE::PoseSearch::FRole> Roles,
-		const TArrayView<const UObject*> AssetsToSearch,
-		const FName PoseHistoryName,
-		const FPoseSearchContinuingProperties& ContinuingProperties,
-		const FPoseSearchFutureProperties& Future,
-		FPoseSearchBlueprintResult& Result,
-		const int32 DebugSessionUniqueIdentifier);
-
-	UE_DEPRECATED(5.4, "Use other MotionMatch signatures instead")
-	static UE::PoseSearch::FSearchResult MotionMatch(
-		const TArrayView<UAnimInstance*> AnimInstances,
-		const TArrayView<const UE::PoseSearch::FRole> Roles,
-		const TArrayView<const UE::PoseSearch::IPoseHistory*> PoseHistories, 
-		const TArrayView<const UObject*> AssetsToSearch,
-		const FPoseSearchContinuingProperties& ContinuingProperties,
-		const FPoseSearchFutureProperties& Future,
-		const int32 DebugSessionUniqueIdentifier);
-
-	UE_DEPRECATED(5.4, "Use other MotionMatch signatures instead")
-	static UE::PoseSearch::FSearchResult MotionMatch(
-		const FAnimationBaseContext& Context,
-		TArrayView<const UObject*> AssetsToSearch,
-		const FPoseSearchContinuingProperties& ContinuingProperties);
-		
-	UE_DEPRECATED(5.4, "Use other MotionMatch signatures instead")
-	static UE::PoseSearch::FSearchResult MotionMatch(
-		TArrayView<UAnimInstance*> AnimInstances,
-		TArrayView<const UE::PoseSearch::FRole> Roles,
-		TArrayView<const UE::PoseSearch::IPoseHistory*> PoseHistories, 
-		TArrayView<const UObject*> AssetsToSearch,
-		const FPoseSearchContinuingProperties& ContinuingProperties,
-		const int32 DebugSessionUniqueIdentifier,
-		float DesiredPermutationTimeOffset = 0.f);
-
 	static const FAnimNode_PoseSearchHistoryCollector_Base* FindPoseHistoryNode(
 		const FName PoseHistoryName,
 		const UAnimInstance* AnimInstance);
diff --git a/Engine/Plugins/Animation/PoseSearch/Source/Runtime/Public/PoseSearch/PoseSearchSchema.h b/Engine/Plugins/Animation/PoseSearch/Source/Runtime/Public/PoseSearch/PoseSearchSchema.h
--- a/Engine/Plugins/Animation/PoseSearch/Source/Runtime/Public/PoseSearch/PoseSearchSchema.h
+++ b/Engine/Plugins/Animation/PoseSearch/Source/Runtime/Public/PoseSearch/PoseSearchSchema.h
@@ -64,6 +64,7 @@ class POSESEARCH_API UPoseSearchSchema : public UDataAsset
 	GENERATED_BODY()
 
 public:
+	UE_DEPRECATED(5.5, "Use Skeletons[i].Skeleton instead")
 	UPROPERTY()
 	TObjectPtr<USkeleton> Skeleton_DEPRECATED;
 	
@@ -84,6 +85,7 @@ private:
 	TArray<TObjectPtr<UPoseSearchFeatureChannel>> FinalizedChannels;
 
 public:
+	UE_DEPRECATED(5.5, "Use Skeletons[i].MirrorDataTable instead")
 	UPROPERTY()
 	TObjectPtr<UMirrorDataTable> MirrorDataTable_DEPRECATED;
 
