Add CVar to disable playrate modifier in motion matching

diff --git a/Engine/Plugins/Animation/PoseSearch/Source/Runtime/Private/AnimNode_MotionMatching.cpp b/Engine/Plugins/Animation/PoseSearch/Source/Runtime/Private/AnimNode_MotionMatching.cpp
--- a/Engine/Plugins/Animation/PoseSearch/Source/Runtime/Private/AnimNode_MotionMatching.cpp
+++ b/Engine/Plugins/Animation/PoseSearch/Source/Runtime/Private/AnimNode_MotionMatching.cpp
@@ -21,6 +21,16 @@
 #define LOCTEXT_NAMESPACE "AnimNode_MotionMatching"
 
 #if ENABLE_ANIM_DEBUG
+namespace UE::Private
+{
+	enum EPlayRateState : int8
+	{
+		Disabled = 0,
+		Enabled = 1,
+		PoseSearchOnly = 2
+	};
+} // namespace UE::Private
+
 static bool GVarAnimNodeMotionMatchingDrawQuery = false;
 static FAutoConsoleVariableRef CVarAnimNodeMotionMatchingDrawQuery(TEXT("a.AnimNode.MotionMatching.DebugDrawQuery"), GVarAnimNodeMotionMatchingDrawQuery, TEXT("Draw input query"));
 
@@ -35,6 +45,16 @@ static FAutoConsoleVariableRef CVarAnimNodeMotionMatchingDrawInfoVerbose(TEXT("a
 
 static float GVarAnimNodeMotionMatchingDrawInfoHeight = 50.f;
 static FAutoConsoleVariableRef CVarAnimNodeMotionMatchingDrawInfoHeight(TEXT("a.AnimNode.MotionMatching.DebugDrawInfoHeight"), GVarAnimNodeMotionMatchingDrawInfoHeight, TEXT("Vertical offset for DebugDrawInfo"));
+
+static int32 GVarAnimNodeMotionMatchingPlayRateEnabled = UE::Private::EPlayRateState::Enabled;
+static FAutoConsoleVariableRef CVarAnimNodeMotionMatchingPlayRateEnabled(
+	TEXT("a.AnimNode.MotionMatching.DebugPlayRateEnabled"), 
+	GVarAnimNodeMotionMatchingPlayRateEnabled,
+	TEXT("Toggles if PlayRate is used in motion matching. Same as setting PlayRate to (1,1) when disabled.\n")
+	TEXT("0: Completely disable PlayRate usage.\n")
+	TEXT("1: Enable all usages of PlayRate.\n")
+	TEXT("2: Enable PlayRate in PoseSeach only (Not used in actual playback).\n")
+);
 #endif
 
 /////////////////////////////////////////////////////
@@ -170,6 +190,12 @@ void FAnimNode_MotionMatching::UpdateAssetPlayer(const FAnimationUpdateContext&
 			PoseHistory = &PoseHistoryProvider->GetPoseHistory();
 		}
 
+		FFloatInterval PoseSearchPlayRate = PlayRate;
+#if ENABLE_ANIM_DEBUG
+		bool bPoseSearchPlayRateEnabled = GVarAnimNodeMotionMatchingPlayRateEnabled == UE::Private::EPlayRateState::Enabled || GVarAnimNodeMotionMatchingPlayRateEnabled == UE::Private::EPlayRateState::PoseSearchOnly;
+		PoseSearchPlayRate = bPoseSearchPlayRateEnabled ? PoseSearchPlayRate : FFloatInterval(1.f, 1.f);
+#endif // ENABLE_ANIM_DEBUG
+
 		UPoseSearchLibrary::UpdateMotionMatchingState(
 			Context.AnimInstanceProxy->GetAnimInstanceObject(),
 			PoseHistory,
@@ -178,7 +204,7 @@ void FAnimNode_MotionMatching::UpdateAssetPlayer(const FAnimationUpdateContext&
 			PoseJumpThresholdTime,
 			PoseReselectHistory,
 			bShouldSearch ? SearchThrottleTime : UE_BIG_NUMBER,
-			PlayRate,
+			PoseSearchPlayRate,
 			MotionMatchingState,
 			NextUpdateInterruptMode,
 			bShouldUseCachedChannelData
@@ -195,6 +221,11 @@ void FAnimNode_MotionMatching::UpdateAssetPlayer(const FAnimationUpdateContext&
 	MotionMatchingState.bJumpedToPose = bJumpToPose;
 	PRAGMA_ENABLE_DEPRECATION_WARNINGS
 
+	float DesiredPlayRate = MotionMatchingState.WantedPlayRate * PlayRateMultiplier;
+#if ENABLE_ANIM_DEBUG
+	DesiredPlayRate = GVarAnimNodeMotionMatchingPlayRateEnabled == UE::Private::EPlayRateState::Enabled ? DesiredPlayRate : PlayRateMultiplier;
+#endif // ENABLE_ANIM_DEBUG
+
 	UE::Anim::FNodeFunctionCaller::CallFunction(GetOnUpdateMotionMatchingStateFunction(), Context, *this);
 
 	// If a new pose is requested, blend into the new asset via BlendStackNode
@@ -241,7 +272,7 @@ void FAnimNode_MotionMatching::UpdateAssetPlayer(const FAnimationUpdateContext&
 				{
 					FAnimNode_BlendStack_Standalone::BlendTo(Context, AnimationAsset, MotionMatchingState.CurrentSearchResult.AssetTime,
 						SearchIndexAsset->IsLooping(), SearchIndexAsset->IsMirrored(), CurrentResultDatabase->Schema->GetMirrorDataTable(DefaultRole), BlendTime,
-						BlendProfile, BlendOption, bUseInertialBlend, NAME_None, SearchIndexAsset->GetBlendParameters(), MotionMatchingState.WantedPlayRate * PlayRateMultiplier, 0,
+						BlendProfile, BlendOption, bUseInertialBlend, NAME_None, SearchIndexAsset->GetBlendParameters(), DesiredPlayRate, 0,
 						GetGroupName(), GetGroupRole(), GetGroupMethod(), GetOverridePositionWhenJoiningSyncGroupAsLeader());
 				}
 			}
@@ -255,7 +286,7 @@ void FAnimNode_MotionMatching::UpdateAssetPlayer(const FAnimationUpdateContext&
 	const bool bDidBlendToRequestAnInertialBlend = bJumpToPose && bUseInertialBlend;
 	UE::Anim::TOptionalScopedGraphMessage<UE::Anim::FAnimInertializationSyncScope> InertializationSync(bDidBlendToRequestAnInertialBlend, Context);
 	
-	FAnimNode_BlendStack_Standalone::UpdatePlayRate(MotionMatchingState.WantedPlayRate * PlayRateMultiplier);
+	FAnimNode_BlendStack_Standalone::UpdatePlayRate(DesiredPlayRate);
 	FAnimNode_BlendStack_Standalone::UpdateBlendspaceParameters(GetBlendspaceUpdateMode(), GetBlendspaceParameters());
 
 	FAnimNode_BlendStack_Standalone::UpdateAssetPlayer(Context);
diff --git a/Engine/Plugins/Experimental/AnimNextPoseSearch/Source/AnimNextPoseSearch/Private/MotionMatchingTrait.cpp b/Engine/Plugins/Experimental/AnimNextPoseSearch/Source/AnimNextPoseSearch/Private/MotionMatchingTrait.cpp
--- a/Engine/Plugins/Experimental/AnimNextPoseSearch/Source/AnimNextPoseSearch/Private/MotionMatchingTrait.cpp
+++ b/Engine/Plugins/Experimental/AnimNextPoseSearch/Source/AnimNextPoseSearch/Private/MotionMatchingTrait.cpp
@@ -15,6 +15,28 @@
 #include "TraitInterfaces/IGraphFactory.h"
 #include "TraitInterfaces/ITimeline.h"
 
+#if ENABLE_ANIM_DEBUG
+namespace UE::Private
+{
+	enum EPlayRateState : int8
+	{
+		Disabled = 0,
+		Enabled = 1,
+		PoseSearchOnly = 2
+	};
+} // namespace UE::Private
+
+static int32 GVarAnimNextMotionMatchingPlayRateEnabled = UE::Private::EPlayRateState::Enabled;
+static FAutoConsoleVariableRef CVarAnimNextMotionMatchingPlayRateEnabled(
+	TEXT("a.AnimNext.MotionMatchingTrait.DebugPlayRateEnabled"),
+	GVarAnimNextMotionMatchingPlayRateEnabled,
+	TEXT("Toggles if PlayRate is used in motion matching. Same as setting PlayRate to (1,1) when disabled.\n")
+	TEXT("0: Completely disable PlayRate usage.\n")
+	TEXT("1: Enable all usages of PlayRate.\n")
+	TEXT("2: Enable PlayRate in PoseSeach only (Not used in actual playback).\n")
+);
+#endif // ENABLE_ANIM_DEBUG
+
 namespace UE::AnimNext
 {
 	AUTO_REGISTER_ANIM_TRAIT(FMotionMatchingTrait)
@@ -121,6 +143,12 @@ namespace UE::AnimNext
 				AnimContext = ModuleInstance->GetObject();
 			}
 
+			FFloatInterval PoseSearchPlayRate = SharedData->GetPlayRate(Binding);
+#if ENABLE_ANIM_DEBUG
+			bool bPoseSearchPlayRateEnabled = GVarAnimNextMotionMatchingPlayRateEnabled == UE::Private::EPlayRateState::Enabled || GVarAnimNextMotionMatchingPlayRateEnabled == UE::Private::EPlayRateState::PoseSearchOnly;
+			PoseSearchPlayRate = bPoseSearchPlayRateEnabled ? PoseSearchPlayRate : FFloatInterval(1.f, 1.f);
+#endif // ENABLE_ANIM_DEBUG
+
 			UPoseSearchLibrary::UpdateMotionMatchingState(
 				AnimContext,
 				PoseHistory,
@@ -129,7 +157,7 @@ namespace UE::AnimNext
 				SharedData->GetPoseJumpThresholdTime(Binding),
 				SharedData->GetPoseReselectHistory(Binding),
 				SharedData->GetbShouldSearch(Binding) ? SharedData->GetSearchThrottleTime(Binding) : UE_BIG_NUMBER,
-				SharedData->GetPlayRate(Binding),
+				PoseSearchPlayRate,
 				MotionMatchingState,
 				SharedData->GetNextUpdateInterruptMode(Binding),
 				SharedData->GetbShouldUseCachedChannelData(Binding)
@@ -141,6 +169,11 @@ namespace UE::AnimNext
 			);
 		}
 
+		float DesiredPlayRate = MotionMatchingState.WantedPlayRate;
+#if ENABLE_ANIM_DEBUG
+		DesiredPlayRate = GVarAnimNextMotionMatchingPlayRateEnabled == UE::Private::EPlayRateState::Enabled ? DesiredPlayRate : 1.0f;
+#endif // ENABLE_ANIM_DEBUG
+
 		const bool bJumpToPose = MotionMatchingState.CurrentSearchResult.IsValid() && !MotionMatchingState.CurrentSearchResult.bIsContinuingPoseSearch;
 		if (bJumpToPose)
 		{
@@ -161,7 +194,7 @@ namespace UE::AnimNext
 						NewPayload.InitializeAs<FAnimNextNativeDataInterface_AnimSequencePlayer>();
 						FAnimNextNativeDataInterface_AnimSequencePlayer& NewNativeInterface = NewPayload.GetMutable<FAnimNextNativeDataInterface_AnimSequencePlayer>();
 						NewNativeInterface.AnimSequence = Cast<UAnimSequence>(AnimationAsset);
-						NewNativeInterface.PlayRate = MotionMatchingState.WantedPlayRate;
+						NewNativeInterface.PlayRate = DesiredPlayRate;
 						NewNativeInterface.StartPosition = MotionMatchingState.CurrentSearchResult.AssetTime;
 						NewNativeInterface.Loop = SearchIndexAsset->IsLooping();
 
@@ -177,7 +210,7 @@ namespace UE::AnimNext
 				}
 			}
 		}
-		else if(!FMath::IsNearlyEqual(InstanceData->GraphPlaybackRate, MotionMatchingState.WantedPlayRate))
+		else if(!FMath::IsNearlyEqual(InstanceData->GraphPlaybackRate, DesiredPlayRate))
 		{
 			// Sync latest motion matching wanted playrate to current graph
 			IBlendStack::FGraphRequestPtr ActiveGraphRequest;
@@ -190,8 +223,8 @@ namespace UE::AnimNext
 				{
 					if (FAnimNextNativeDataInterface_AnimSequencePlayer* PlayAnimPayload = Payload.GetPtr<FAnimNextNativeDataInterface_AnimSequencePlayer>())
 					{
-						PlayAnimPayload->PlayRate = MotionMatchingState.WantedPlayRate;
-						InstanceData->GraphPlaybackRate = MotionMatchingState.WantedPlayRate;
+						PlayAnimPayload->PlayRate = DesiredPlayRate;
+						InstanceData->GraphPlaybackRate = DesiredPlayRate;
 						break;
 					}
 				}
