PoseSearch - WIP (experimental!) adding Chooser to database as workflow proof of concept

diff --git a/Engine/Plugins/Animation/PoseSearch/PoseSearch.uplugin b/Engine/Plugins/Animation/PoseSearch/PoseSearch.uplugin
--- a/Engine/Plugins/Animation/PoseSearch/PoseSearch.uplugin
+++ b/Engine/Plugins/Animation/PoseSearch/PoseSearch.uplugin
@@ -36,6 +36,10 @@
 		{
 			"Name": "BlendStack",
 			"Enabled": true
+		},
+		{
+			"Name": "Chooser",
+			"Enabled": true
 		}
 	]
 }
\ No newline at end of file
diff --git a/Engine/Plugins/Animation/PoseSearch/Source/Runtime/PoseSearch.Build.cs b/Engine/Plugins/Animation/PoseSearch/Source/Runtime/PoseSearch.Build.cs
--- a/Engine/Plugins/Animation/PoseSearch/Source/Runtime/PoseSearch.Build.cs
+++ b/Engine/Plugins/Animation/PoseSearch/Source/Runtime/PoseSearch.Build.cs
@@ -16,6 +16,7 @@ public class PoseSearch : ModuleRules
 		PublicDependencyModuleNames.AddRange(
 			new string[]
 			{
+				"Chooser",
 				"Core",
 				"CoreUObject",
 				"Engine",
diff --git a/Engine/Plugins/Animation/PoseSearch/Source/Runtime/Private/PoseSearchDatabase.cpp b/Engine/Plugins/Animation/PoseSearch/Source/Runtime/Private/PoseSearchDatabase.cpp
--- a/Engine/Plugins/Animation/PoseSearch/Source/Runtime/Private/PoseSearchDatabase.cpp
+++ b/Engine/Plugins/Animation/PoseSearch/Source/Runtime/Private/PoseSearchDatabase.cpp
@@ -6,6 +6,8 @@
 #include "Animation/AnimSequence.h"
 #include "Animation/BlendSpace.h"
 #include "Animation/BlendSpace1D.h"
+#include "Chooser/Internal/Chooser.h"
+#include "IObjectChooser.h"
 #include "PoseSearch/MultiAnimAsset.h"
 #include "PoseSearch/PoseSearchAnimNotifies.h"
 #include "PoseSearch/PoseSearchContext.h"
@@ -739,6 +741,7 @@ void UPoseSearchDatabase::PostLoad()
 	using namespace UE::PoseSearch;
 
 	// todo: should we SynchronizeWithExternalDependencies() here?
+	SynchronizeChooser();
 
 	ERequestAsyncBuildFlag Flag = ERequestAsyncBuildFlag::NewRequest;
 #if WITH_ENGINE
@@ -926,6 +929,104 @@ void UPoseSearchDatabase::SynchronizeWithExternalDependencies(TConstArrayView<UA
 	}
 }
 
+const UChooserTable* UPoseSearchDatabase::GetChooser() const
+{
+	return Chooser;
+}
+
+void UPoseSearchDatabase::SynchronizeChooser()
+{
+	if (Chooser)
+	{
+		TArray<UObject*, TInlineAllocator<128>> ChooserAssets;
+		UChooserTable::IterateChooser(Chooser, FObjectChooserBase::FObjectChooserIteratorCallback::CreateLambda([&ChooserAssets](UObject* Object)
+			{
+				if (Object)
+				{
+					if (Object->GetClass()->IsChildOf(UAnimSequence::StaticClass()) ||
+						Object->GetClass()->IsChildOf(UAnimComposite::StaticClass()) ||
+						Object->GetClass()->IsChildOf(UAnimMontage::StaticClass()) ||
+						Object->GetClass()->IsChildOf(UBlendSpace::StaticClass()) ||
+						Object->GetClass()->IsChildOf(UMultiAnimAsset::StaticClass()))
+					{
+						ChooserAssets.Add(Object);
+					}
+				}
+				return FObjectChooserBase::EIteratorStatus::Continue;
+			}));
+
+		TArray<UObject*, TInlineAllocator<128>> DatabaseAssets;
+		for (FInstancedStruct& AnimationAsset : AnimationAssets)
+		{
+			if (const FPoseSearchDatabaseAnimationAssetBase* AnimationAssetBase = AnimationAsset.GetPtr<FPoseSearchDatabaseAnimationAssetBase>())
+			{
+				DatabaseAssets.Add(AnimationAssetBase->GetAnimationAsset());
+			}
+		}
+
+		if (ChooserAssets != DatabaseAssets)
+		{
+			TArray<FInstancedStruct> AnimationAssetsBackup = MoveTemp(AnimationAssets);
+		
+			for (UObject* ChooserAsset : ChooserAssets)
+			{
+				// searching ChooserAsset in AnimationAssetsBackup
+				bool bFoundInAnimationAssetsBackup = false;
+				for (int32 AnimationAssetBackupIndex = 0; AnimationAssetBackupIndex < AnimationAssetsBackup.Num(); ++AnimationAssetBackupIndex)
+				{
+					if (const FPoseSearchDatabaseAnimationAssetBase* AnimationAssetBase = AnimationAssetsBackup[AnimationAssetBackupIndex].GetPtr<FPoseSearchDatabaseAnimationAssetBase>())
+					{
+						if (AnimationAssetBase->GetAnimationAsset() == ChooserAsset)
+						{
+							bFoundInAnimationAssetsBackup = true;
+							AnimationAssets.Add(AnimationAssetsBackup[AnimationAssetBackupIndex]);
+							AnimationAssetsBackup.RemoveAt(AnimationAssetBackupIndex);
+							break;
+						}
+					}
+				}
+
+				if (!bFoundInAnimationAssetsBackup)
+				{
+					if (UAnimSequence* Sequence = Cast<UAnimSequence>(ChooserAsset))
+					{
+						FPoseSearchDatabaseSequence DatabaseSequence;
+						DatabaseSequence.Sequence = Sequence;
+						AnimationAssets.Add(FInstancedStruct::Make(DatabaseSequence));
+					}
+					else if (UAnimComposite* AnimComposite = Cast<UAnimComposite>(ChooserAsset))
+					{
+						FPoseSearchDatabaseAnimComposite DatabaseAnimComposite;
+						DatabaseAnimComposite.AnimComposite = AnimComposite;
+						AnimationAssets.Add(FInstancedStruct::Make(DatabaseAnimComposite));
+					}
+					else if (UAnimMontage* AnimMontage = Cast<UAnimMontage>(ChooserAsset))
+					{
+						FPoseSearchDatabaseAnimMontage DatabaseAnimMontage;
+						DatabaseAnimMontage.AnimMontage = AnimMontage;
+						AnimationAssets.Add(FInstancedStruct::Make(DatabaseAnimMontage));
+					}
+					else if (UBlendSpace* BlendSpace = Cast<UBlendSpace>(ChooserAsset))
+					{
+						FPoseSearchDatabaseBlendSpace DatabaseBlendSpace;
+						DatabaseBlendSpace.BlendSpace = BlendSpace;
+						AnimationAssets.Add(FInstancedStruct::Make(DatabaseBlendSpace));
+					}
+					else if (UMultiAnimAsset* MultiAnimAsset = Cast<UMultiAnimAsset>(ChooserAsset))
+					{
+						FPoseSearchDatabaseMultiAnimAsset DatabaseMultiAnimAsset;
+						DatabaseMultiAnimAsset.MultiAnimAsset = MultiAnimAsset;
+						AnimationAssets.Add(FInstancedStruct::Make(DatabaseMultiAnimAsset));
+					}
+				}
+			}
+
+			Modify();
+			NotifySynchronizeWithExternalDependencies();
+		}
+	}
+}
+
 void UPoseSearchDatabase::BeginCacheForCookedPlatformData(const ITargetPlatform* TargetPlatform)
 {
 	using namespace UE::PoseSearch;
@@ -1040,6 +1141,28 @@ UE::PoseSearch::FSearchResult UPoseSearchDatabase::Search(UE::PoseSearch::FSearc
 	}
 #endif // WITH_EDITOR
 
+	// updating the SearchContext::AssetsToConsider from the UPoseSearchDatabase::Chooser evaluation 
+	// (and restore them with the original one CurrentAssetToConsider at the end of the method)
+	TConstArrayView<const UObject*> CurrentAssetToConsider = SearchContext.GetAssetsToConsider();
+	TArray<const UObject*, TInlineAllocator<128>> ChooserAssetToConsider;
+	if (Chooser && !SearchContext.GetAnimContexts().IsEmpty())
+	{
+		if (!CurrentAssetToConsider.IsEmpty())
+		{
+			UE_LOG(LogPoseSearch, Warning, TEXT("UPoseSearchDatabase::Search - database (%s) has an associated Chooser filtering the search down, so previous SearchContext.GetAssetsToConsider() context will be skipped"), *GetName());
+		}
+
+		// @todo: add MM interaction support. Maybe iterate over all the SearchContext.GetAnimContexts and use the union of the ChosenAssets?
+		FChooserEvaluationContext ChooserContext(const_cast<UObject*>(SearchContext.GetAnimContexts()[0]));
+		UChooserTable::EvaluateChooser(ChooserContext, Chooser, FObjectChooserBase::FObjectChooserIteratorCallback::CreateLambda([this, &ChooserAssetToConsider](UObject* InResult)
+			{
+				ChooserAssetToConsider.Add(InResult);
+				return FObjectChooserBase::EIteratorStatus::Continue;
+			}));
+
+		SearchContext.SetAssetsToConsider(ChooserAssetToConsider);
+	}
+
 	switch (PoseSearchMode)
 	{
 	case EPoseSearchMode::BruteForce:
@@ -1085,6 +1208,8 @@ UE::PoseSearch::FSearchResult UPoseSearchDatabase::Search(UE::PoseSearch::FSearc
 	SearchContext.Track(this);
 #endif // UE_POSE_SEARCH_TRACE_ENABLED
 
+	SearchContext.SetAssetsToConsider(CurrentAssetToConsider);
+
 	return Result;
 }
 
diff --git a/Engine/Plugins/Animation/PoseSearch/Source/Runtime/Private/PoseSearchDerivedData.cpp b/Engine/Plugins/Animation/PoseSearch/Source/Runtime/Private/PoseSearchDerivedData.cpp
--- a/Engine/Plugins/Animation/PoseSearch/Source/Runtime/Private/PoseSearchDerivedData.cpp
+++ b/Engine/Plugins/Animation/PoseSearch/Source/Runtime/Private/PoseSearchDerivedData.cpp
@@ -9,6 +9,7 @@
 #include "Animation/BlendSpace.h"
 #include "AssetRegistry/ARFilter.h"
 #include "AssetRegistry/AssetRegistryModule.h"
+#include "Chooser/Internal/Chooser.h"
 #include "DerivedDataCache.h"
 #include "DerivedDataRequestOwner.h"
 #include "StructUtils/InstancedStruct.h"
@@ -1159,6 +1160,39 @@ static bool IndexDatabase(FSearchIndexBase& SearchIndexBase, const UPoseSearchDa
 	return true;
 }
 
+static void SynchronizeDatabaseChooser(UObject* Object)
+{
+	if (UPoseSearchDatabase* Database = Cast<UPoseSearchDatabase>(Object))
+	{
+		Database->SynchronizeChooser();
+	}
+	else if (const UChooserTable* ChooserTable = Cast<UChooserTable>(Object))
+	{
+		const UChooserTable* RootChooser = ChooserTable->GetRootChooser();
+
+		IAssetRegistry& AssetRegistry = FModuleManager::LoadModuleChecked<FAssetRegistryModule>("AssetRegistry").Get();
+		TArray<FAssetIdentifier> Referencers;
+		AssetRegistry.GetReferencers(RootChooser->GetPackage()->GetFName(), Referencers);
+
+		for (const FAssetIdentifier& Referencer : Referencers)
+		{
+			TArray<FAssetData> Assets;
+			AssetRegistry.GetAssetsByPackageName(Referencer.PackageName, Assets);
+
+			for (const FAssetData& Asset : Assets)
+			{
+				if (Asset.IsInstanceOf(UPoseSearchDatabase::StaticClass()))
+				{
+					if (UPoseSearchDatabase* ReferencedDatabase = CastChecked<UPoseSearchDatabase>(Asset.FastGetAsset(true)))
+					{
+						ReferencedDatabase->SynchronizeChooser();
+					}
+				}
+			}
+		}
+	}
+}
+
 #if ENABLE_ANIM_DEBUG
 
 static void CompareChannelValues(int32 RecursionIndex, int32 PoseIndex, const TConstArrayView<float>& PoseA, const TConstArrayView<float>& PoseB, const TConstArrayView<TObjectPtr<UPoseSearchFeatureChannel>>& Channels, FStringBuilderBase& StringBuilder)
@@ -2131,6 +2165,8 @@ void FAsyncPoseSearchDatabasesManagement::PostModified(UObject* Object)
 
 	SynchronizeDatabases();
 
+	SynchronizeDatabaseChooser(Object);
+
 	ClearPreCancelled();
 }
 
diff --git a/Engine/Plugins/Animation/PoseSearch/Source/Runtime/Public/PoseSearch/PoseSearchDatabase.h b/Engine/Plugins/Animation/PoseSearch/Source/Runtime/Public/PoseSearch/PoseSearchDatabase.h
--- a/Engine/Plugins/Animation/PoseSearch/Source/Runtime/Public/PoseSearch/PoseSearchDatabase.h
+++ b/Engine/Plugins/Animation/PoseSearch/Source/Runtime/Public/PoseSearch/PoseSearchDatabase.h
@@ -15,6 +15,7 @@ class UAnimationAsset;
 class UAnimComposite;
 class UAnimMontage;
 class UBlendSpace;
+class UChooserTable;
 class UMultiAnimAsset;
 
 #if WITH_EDITORONLY_DATA
@@ -430,6 +431,11 @@ private:
 	UPROPERTY()
 	TArray<FInstancedStruct> AnimationAssets;
 
+	// Experimental, this feature might be removed without warning, not for production use
+	// if Chooser is valid, all the AnimationAssets will come from it and AnimationAssets is not gonna be used AT ALL
+	UPROPERTY(EditAnywhere, Category = "Database|Experimental")
+	TObjectPtr<UChooserTable> Chooser;
+
 public:
 	/** Array of tags that can be used as metadata. */
 	UPROPERTY(EditAnywhere, BlueprintReadOnly, Category = "Database")
@@ -556,6 +562,11 @@ public:
 
 	void SynchronizeWithExternalDependencies();
 	void SynchronizeWithExternalDependencies(TConstArrayView<UAnimSequenceBase*> SequencesBase);
+	
+	// Experimental, this feature might be removed without warning, not for production use
+	void SynchronizeChooser();
+	// Experimental, this feature might be removed without warning, not for production use
+	const UChooserTable* GetChooser() const;
 #endif // WITH_EDITOR
 
 #if WITH_EDITOR && ENABLE_ANIM_DEBUG
