PoseSearch - initial MM interaction support for AnimNext

diff --git a/Engine/Plugins/Animation/PoseSearch/Source/Runtime/Private/AnimNode_MotionMatchingInteraction.cpp b/Engine/Plugins/Animation/PoseSearch/Source/Runtime/Private/AnimNode_MotionMatchingInteraction.cpp
--- a/Engine/Plugins/Animation/PoseSearch/Source/Runtime/Private/AnimNode_MotionMatchingInteraction.cpp
+++ b/Engine/Plugins/Animation/PoseSearch/Source/Runtime/Private/AnimNode_MotionMatchingInteraction.cpp
@@ -108,7 +108,7 @@ void FAnimNode_MotionMatchingInteraction::UpdateAssetPlayer(const FAnimationUpda
 		ContinuingProperties.PlayingAsset = CurrentResult.SelectedAnimation; // GetAnimAsset();
 		ContinuingProperties.PlayingAssetAccumulatedTime = GetAccumulatedTime();
 
-		CurrentResult = UPoseSearchInteractionLibrary::MotionMatchInteraction(Availabilities, Context.AnimInstanceProxy->GetAnimInstanceObject(), ContinuingProperties, PoseHistoryProvider->GetHistoryCollector(), bValidateResultAgainstAvailabilities);
+		CurrentResult = UPoseSearchInteractionLibrary::MotionMatchInteraction(Availabilities, Context.AnimInstanceProxy->GetAnimInstanceObject(), ContinuingProperties, &PoseHistoryProvider->GetPoseHistory(), bValidateResultAgainstAvailabilities);
 		check(CurrentResult.ActorRootTransforms.Num() == CurrentResult.ActorRootBoneTransforms.Num());
 
 		UAnimationAsset* RoledAnimAsset = nullptr;
diff --git a/Engine/Plugins/Animation/PoseSearch/Source/Runtime/Private/AnimNode_PoseSearchHistoryCollector.cpp b/Engine/Plugins/Animation/PoseSearch/Source/Runtime/Private/AnimNode_PoseSearchHistoryCollector.cpp
--- a/Engine/Plugins/Animation/PoseSearch/Source/Runtime/Private/AnimNode_PoseSearchHistoryCollector.cpp
+++ b/Engine/Plugins/Animation/PoseSearch/Source/Runtime/Private/AnimNode_PoseSearchHistoryCollector.cpp
@@ -14,6 +14,75 @@
 /////////////////////////////////////////////////////
 // FAnimNode_PoseSearchHistoryCollector_Base
 
+// UE::PoseSearch::IPoseHistory interface
+bool FAnimNode_PoseSearchHistoryCollector_Base::GetTransformAtTime(float Time, FTransform& OutBoneTransform, const USkeleton* BoneIndexSkeleton, FBoneIndexType BoneIndexType, FBoneIndexType ReferenceBoneIndexType, bool bExtrapolate) const
+{
+	return PoseHistory.GetTransformAtTime(Time, OutBoneTransform, BoneIndexSkeleton, BoneIndexType, ReferenceBoneIndexType, bExtrapolate);
+}
+
+bool FAnimNode_PoseSearchHistoryCollector_Base::GetCurveValueAtTime(float Time, const FName& CurveName, float& OutCurveValue, bool bExtrapolate) const
+{
+	return PoseHistory.GetCurveValueAtTime(Time, CurveName, OutCurveValue, bExtrapolate);
+}
+
+const FPoseSearchQueryTrajectory& FAnimNode_PoseSearchHistoryCollector_Base::GetTrajectory() const
+{
+	return PoseHistory.GetTrajectory();
+}
+
+void FAnimNode_PoseSearchHistoryCollector_Base::SetTrajectory(const FPoseSearchQueryTrajectory& InTrajectory, float InTrajectorySpeedMultiplier)
+{
+	PoseHistory.SetTrajectory(InTrajectory, InTrajectorySpeedMultiplier);
+}
+
+float FAnimNode_PoseSearchHistoryCollector_Base::GetTrajectorySpeedMultiplier() const
+{
+	return PoseHistory.GetTrajectorySpeedMultiplier();
+}
+
+bool FAnimNode_PoseSearchHistoryCollector_Base::IsEmpty() const
+{
+	return PoseHistory.IsEmpty();
+}
+
+const UE::PoseSearch::FBoneToTransformMap& FAnimNode_PoseSearchHistoryCollector_Base::GetBoneToTransformMap() const
+{
+	return PoseHistory.GetBoneToTransformMap();
+}
+
+const TConstArrayView<FName> FAnimNode_PoseSearchHistoryCollector_Base::GetCollectedCurves() const
+{
+	return PoseHistory.GetCollectedCurves();
+}
+
+int32 FAnimNode_PoseSearchHistoryCollector_Base::GetNumEntries() const
+{
+	return PoseHistory.GetNumEntries();
+}
+
+const UE::PoseSearch::FPoseHistoryEntry& FAnimNode_PoseSearchHistoryCollector_Base::GetEntry(int32 EntryIndex) const
+{
+	return PoseHistory.GetEntry(EntryIndex);
+}
+
+#if ENABLE_DRAW_DEBUG && ENABLE_ANIM_DEBUG
+void FAnimNode_PoseSearchHistoryCollector_Base::DebugDraw(const UWorld* World, FColor Color) const
+{
+	PoseHistory.DebugDraw(World, Color);
+}
+
+void FAnimNode_PoseSearchHistoryCollector_Base::DebugDraw(FAnimInstanceProxy& AnimInstanceProxy, FColor Color) const
+{
+	PoseHistory.DebugDraw(AnimInstanceProxy, Color);
+}
+#endif
+// End of UE::PoseSearch::IPoseHistory interface
+
+void FAnimNode_PoseSearchHistoryCollector_Base::GenerateTrajectory(const UAnimInstance* InAnimInstance)
+{
+	GenerateTrajectory(InAnimInstance, InAnimInstance->GetDeltaSeconds());
+}
+
 void FAnimNode_PoseSearchHistoryCollector_Base::Initialize_AnyThread(const FAnimationInitializeContext& Context)
 {
 	Super::Initialize_AnyThread(Context);
@@ -78,7 +147,7 @@ void FAnimNode_PoseSearchHistoryCollector_Base::Update_AnyThread(const FAnimatio
 	
 	if (bGenerateTrajectory)
 	{
-		GenerateTrajectory(Cast<UAnimInstance>(Context.AnimInstanceProxy->GetAnimInstanceObject()));
+		GenerateTrajectory(Context.AnimInstanceProxy->GetAnimInstanceObject(), Context.GetDeltaTime());
 	}
 	else
 	{
@@ -89,7 +158,7 @@ void FAnimNode_PoseSearchHistoryCollector_Base::Update_AnyThread(const FAnimatio
 	UpdateCounter.SynchronizeWith(Context.AnimInstanceProxy->GetUpdateCounter());
 }
 
-void FAnimNode_PoseSearchHistoryCollector_Base::GenerateTrajectory(const UAnimInstance* InAnimInstance)
+void FAnimNode_PoseSearchHistoryCollector_Base::GenerateTrajectory(const UObject* AnimContext, float DeltaTime)
 {
 	if (bGenerateTrajectory && !bIsTrajectoryGeneratedBeforePreUpdate)
 	{
@@ -99,7 +168,7 @@ void FAnimNode_PoseSearchHistoryCollector_Base::GenerateTrajectory(const UAnimIn
 		TrajectoryDataSampling.NumPredictionSamples = TrajectoryPredictionCount;
 		TrajectoryDataSampling.SecondsPerPredictionSample = PredictionSamplingInterval;
 
-		PoseHistory.GenerateTrajectory(InAnimInstance, InAnimInstance->GetDeltaSeconds(), TrajectoryData, TrajectoryDataSampling);
+		PoseHistory.GenerateTrajectory(AnimContext, DeltaTime, TrajectoryData, TrajectoryDataSampling);
 
 		bIsTrajectoryGeneratedBeforePreUpdate = true;
 	}
diff --git a/Engine/Plugins/Animation/PoseSearch/Source/Runtime/Private/PoseHistoryProvider.cpp b/Engine/Plugins/Animation/PoseSearch/Source/Runtime/Private/PoseHistoryProvider.cpp
--- a/Engine/Plugins/Animation/PoseSearch/Source/Runtime/Private/PoseHistoryProvider.cpp
+++ b/Engine/Plugins/Animation/PoseSearch/Source/Runtime/Private/PoseHistoryProvider.cpp
@@ -8,13 +8,8 @@ namespace UE::PoseSearch
 	
 const IPoseHistory& FPoseHistoryProvider::GetPoseHistory() const
 {
-	if (HistoryCollector)
-	{
-		return HistoryCollector->GetPoseHistory();
-	}
-	
-	check(PoseHistory);
-	return *PoseHistory;
+	check(HistoryCollector);
+	return *HistoryCollector;
 }
 
 } // namespace UE::PoseSearch
\ No newline at end of file
diff --git a/Engine/Plugins/Animation/PoseSearch/Source/Runtime/Private/PoseSearchHistory.cpp b/Engine/Plugins/Animation/PoseSearch/Source/Runtime/Private/PoseSearchHistory.cpp
--- a/Engine/Plugins/Animation/PoseSearch/Source/Runtime/Private/PoseSearchHistory.cpp
+++ b/Engine/Plugins/Animation/PoseSearch/Source/Runtime/Private/PoseSearchHistory.cpp
@@ -796,11 +796,11 @@ const FPoseHistoryEntry& FPoseHistory::GetEntry(int32 EntryIndex) const
 	return PoseData.Entries[EntryIndex];
 }
 
-void FPoseHistory::GenerateTrajectory(const UAnimInstance* AnimInstance, float DeltaTime, const FPoseSearchTrajectoryData& TrajectoryData, const FPoseSearchTrajectoryData::FSampling& TrajectoryDataSampling)
+void FPoseHistory::GenerateTrajectory(const UObject* AnimContext, float DeltaTime, const FPoseSearchTrajectoryData& TrajectoryData, const FPoseSearchTrajectoryData::FSampling& TrajectoryDataSampling)
 {
 	// @todo: Synchronize the FPoseSearchQueryTrajectorySample::AccumulatedSeconds of the generated trajectory with the FPoseHistoryEntry::AccumulatedSeconds of the captured poses
 	FPoseSearchTrajectoryData::FDerived TrajectoryDataDerived;
-	if (TrajectoryData.UpdateData(DeltaTime, AnimInstance, TrajectoryDataDerived, TrajectoryDataState))
+	if (TrajectoryData.UpdateData(DeltaTime, AnimContext, TrajectoryDataDerived, TrajectoryDataState))
 	{
 		UPoseSearchTrajectoryLibrary::InitTrajectorySamples(Trajectory, TrajectoryDataDerived.Position, TrajectoryDataDerived.Facing, TrajectoryDataSampling, DeltaTime);
 		UPoseSearchTrajectoryLibrary::UpdateHistory_TransformHistory(Trajectory, TrajectoryDataDerived.Position, TrajectoryDataDerived.Velocity, TrajectoryDataSampling, DeltaTime);
@@ -811,13 +811,18 @@ void FPoseHistory::GenerateTrajectory(const UAnimInstance* AnimInstance, float D
 	}
 }
 
+void FPoseHistory::GenerateTrajectory(const UObject* AnimContext, float DeltaTime)
+{
+	GenerateTrajectory(AnimContext, DeltaTime, FPoseSearchTrajectoryData(), FPoseSearchTrajectoryData::FSampling());
+}
+
 void FPoseHistory::SetTrajectory(const FPoseSearchQueryTrajectory& InTrajectory, float InTrajectorySpeedMultiplier)
 {
 	if (!InTrajectory.Samples.IsEmpty())
 	{
 		UE_MT_SCOPED_WRITE_ACCESS(PoseDataThreadSafeCounter);
 
-		// @todo: THIS IS NOT THREAD SAFE! in the context of multi character motion matching (UE_MT_SCOPED_WRITE_ACCESS will assert in case of improper usage)
+		// UE_MT_SCOPED_WRITE_ACCESS will assert in case of improper usage
 		Trajectory = InTrajectory;
 
 		TrajectorySpeedMultiplier = InTrajectorySpeedMultiplier;
@@ -999,6 +1004,20 @@ void FMemStackPoseHistory::Init(const IPoseHistory* InPoseHistory)
 	PoseHistory = InPoseHistory;
 }
 
+void FMemStackPoseHistory::SetTrajectory(const FPoseSearchQueryTrajectory& InTrajectory, float InTrajectorySpeedMultiplier)
+{
+	check(PoseHistory);
+	// FMemStackPoseHistory should never change the trajectory!
+	checkNoEntry();
+}
+
+void FMemStackPoseHistory::GenerateTrajectory(const UObject* AnimContext, float DeltaTime)
+{
+	check(PoseHistory);
+	// FMemStackPoseHistory should never change the trajectory!
+	checkNoEntry();
+}
+
 bool FMemStackPoseHistory::GetTransformAtTime(float Time, FTransform& OutBoneTransform, const USkeleton* BoneIndexSkeleton, FBoneIndexType BoneIndexType, FBoneIndexType ReferenceBoneIndexType, bool bExtrapolate) const
 {
 	check(PoseHistory);
diff --git a/Engine/Plugins/Animation/PoseSearch/Source/Runtime/Private/PoseSearchHistoryCollectorAnimNodeLibrary.cpp b/Engine/Plugins/Animation/PoseSearch/Source/Runtime/Private/PoseSearchHistoryCollectorAnimNodeLibrary.cpp
--- a/Engine/Plugins/Animation/PoseSearch/Source/Runtime/Private/PoseSearchHistoryCollectorAnimNodeLibrary.cpp
+++ b/Engine/Plugins/Animation/PoseSearch/Source/Runtime/Private/PoseSearchHistoryCollectorAnimNodeLibrary.cpp
@@ -12,7 +12,7 @@ void UPoseSearchHistoryCollectorAnimNodeLibrary::GetPoseHistoryNodeTrajectory(co
 {
 	if (FAnimNode_PoseSearchHistoryCollector* PoseSearchHistoryCollectorNodePtr = PoseSearchHistoryCollectorNode.GetAnimNodePtr<FAnimNode_PoseSearchHistoryCollector>())
 	{
-		Trajectory = PoseSearchHistoryCollectorNodePtr->GetPoseHistory().GetTrajectory();
+		Trajectory = PoseSearchHistoryCollectorNodePtr->GetTrajectory();
 	}
 }
 
@@ -20,6 +20,6 @@ void UPoseSearchHistoryCollectorAnimNodeLibrary::SetPoseHistoryNodeTrajectory(co
 {
 	if (FAnimNode_PoseSearchHistoryCollector* PoseSearchHistoryCollectorNodePtr = PoseSearchHistoryCollectorNode.GetAnimNodePtr<FAnimNode_PoseSearchHistoryCollector>())
 	{
-		PoseSearchHistoryCollectorNodePtr->GetPoseHistory().SetTrajectory(Trajectory);
+		PoseSearchHistoryCollectorNodePtr->SetTrajectory(Trajectory);
 	}
 }
diff --git a/Engine/Plugins/Animation/PoseSearch/Source/Runtime/Private/PoseSearchInteractionIsland.cpp b/Engine/Plugins/Animation/PoseSearch/Source/Runtime/Private/PoseSearchInteractionIsland.cpp
--- a/Engine/Plugins/Animation/PoseSearch/Source/Runtime/Private/PoseSearchInteractionIsland.cpp
+++ b/Engine/Plugins/Animation/PoseSearch/Source/Runtime/Private/PoseSearchInteractionIsland.cpp
@@ -6,6 +6,7 @@
 #include "Components/SkeletalMeshComponent.h"
 #include "PoseSearch/PoseSearchInteractionLibrary.h"
 #include "DrawDebugHelpers.h"
+#include "Engine/SkeletalMesh.h"
 #include "GameFramework/CharacterMovementComponent.h"
 #include "PoseSearch/AnimNode_PoseSearchHistoryCollector.h"
 #include "PoseSearch/MultiAnimAsset.h"
@@ -27,11 +28,66 @@ const FTransform& GetContextTransform(const UObject* AnimContext)
 	{
 		return AnimInstance->GetSkelMeshComponent()->GetComponentTransform();
 	}
+	
+	if (const UActorComponent* ActorComponent = Cast<UActorComponent>(AnimContext))
+	{
+		const AActor* Actor = ActorComponent->GetOwner();
+		check(Actor);
+		// @todo: this code depends on how AnimNext gather its context object, and will likely change in the future
+		if (const USkeletalMeshComponent* SkeletalMeshComponent = Actor->GetComponentByClass<USkeletalMeshComponent>())
+		{
+			return SkeletalMeshComponent->GetComponentTransform();
+		}
+
+		return ActorComponent->GetOwner()->GetTransform();
+	}
 
 	unimplemented();
 	return FTransform::Identity;
 }
 
+const USkeleton* GetContextSkeleton(const UObject* AnimContext)
+{
+	if (const UAnimInstance* AnimInstance = Cast<UAnimInstance>(AnimContext))
+	{
+		return AnimInstance->GetRequiredBonesOnAnyThread().GetSkeletonAsset();
+	}
+
+	if (const UActorComponent* ActorComponent = Cast<UActorComponent>(AnimContext))
+	{
+		const AActor* Actor = ActorComponent->GetOwner();
+		check(Actor);
+		// @todo: this code depends on how AnimNext gather its context object, and will likely change in the future
+		if (const USkeletalMeshComponent* SkeletalMeshComponent = Actor->GetComponentByClass<USkeletalMeshComponent>())
+		{
+			if (const USkeletalMesh* SkeletalMesh = SkeletalMeshComponent->GetSkeletalMeshAsset())
+			{
+				return SkeletalMesh->GetSkeleton();
+			}
+		}
+		return nullptr;
+	}
+
+	unimplemented();
+	return nullptr;
+}
+
+const AActor* GetContextOwningActor(const UObject* AnimContext)
+{
+	if (const UAnimInstance* AnimInstance = Cast<UAnimInstance>(AnimContext))
+	{
+		return AnimInstance->GetOwningActor();
+	}
+	
+	if (const UActorComponent* AnimNextComponent = Cast<UActorComponent>(AnimContext))
+	{
+		return AnimNextComponent->GetOwner();
+	}
+	
+	unimplemented();
+	return nullptr;
+}
+
 FVector GetContextLocation(const UObject* AnimContext)
 {
 	return GetContextTransform(AnimContext).GetLocation();
@@ -51,28 +107,13 @@ static FInteractionSearchResult InitSearchResult(const FSearchResult& SearchResu
 	{
 		if (const UObject* AnimContext = SearchContext.AnimContexts[AnimContextIndex].Get())
 		{
-			if (const UAnimInstance* AnimInstance = Cast<UAnimInstance>(AnimContext))
+			InteractionSearchResult.ActorRootTransforms[AnimContextIndex] = GetContextTransform(AnimContext);
+			if (const IPoseHistory* PoseHistory = SearchContext.PoseHistories[AnimContextIndex])
 			{
-				if (const FAnimNode_PoseSearchHistoryCollector_Base* HistoryCollector = SearchContext.HistoryCollectors[AnimContextIndex])
-				{
-					const USkeleton* Skeleton = AnimInstance->GetRequiredBonesOnAnyThread().GetSkeletonAsset();
-					check(Skeleton);
-					HistoryCollector->GetPoseHistory().GetTransformAtTime(0.f, InteractionSearchResult.ActorRootBoneTransforms[AnimContextIndex], Skeleton, RootBoneIndexType, ComponentSpaceIndexType);
-
-					// @todo: should we use the pose history ActorRootTransforms?
-					//HistoryCollector->GetPoseHistory().GetTransformAtTime(0.f, InteractionSearchResult.ActorRootTransforms[AnimContextIndex], Skeleton, ComponentSpaceIndexType, WorldSpaceIndexType);
-					InteractionSearchResult.ActorRootTransforms[AnimContextIndex] = AnimInstance->GetSkelMeshComponent()->GetComponentTransform();
-				}
-				else
-				{
-					InteractionSearchResult.ActorRootBoneTransforms[AnimContextIndex] = FTransform::Identity;
-					InteractionSearchResult.ActorRootTransforms[AnimContextIndex] = AnimInstance->GetSkelMeshComponent()->GetComponentTransform();
-				}
+				PoseHistory->GetTransformAtTime(0.f, InteractionSearchResult.ActorRootBoneTransforms[AnimContextIndex], GetContextSkeleton(AnimContext), RootBoneIndexType, ComponentSpaceIndexType);
 			}
 			else
 			{
-				unimplemented();
-				InteractionSearchResult.ActorRootTransforms[AnimContextIndex] = FTransform::Identity;
 				InteractionSearchResult.ActorRootBoneTransforms[AnimContextIndex] = FTransform::Identity;
 			}
 		}
@@ -201,7 +242,7 @@ bool FInteractionSearchContextBase::IsEquivalent(const FInteractionSearchContext
 		for (int32 IndexOther = 0; IndexOther < Num; ++IndexOther)
 		{
 			if (AnimContexts[IndexThis] == Other.AnimContexts[IndexOther] &&
-				HistoryCollectors[IndexThis] == Other.HistoryCollectors[IndexOther] &&
+				PoseHistories[IndexThis] == Other.PoseHistories[IndexOther] &&
 				Roles[IndexThis] == Other.Roles[IndexOther])
 			{
 				++CommonRoledAnimContexts;
@@ -269,7 +310,7 @@ bool FInteractionSearchContextBase::CheckForConsistency() const
 		return false;
 	}
 
-	if (Num != HistoryCollectors.Num())
+	if (Num != PoseHistories.Num())
 	{
 		return false;
 	}
@@ -308,14 +349,14 @@ bool FInteractionSearchContextBase::CheckForConsistency() const
 
 	for (int32 IndexA = 0; IndexA < Num; ++IndexA)
 	{
-		if (HistoryCollectors[IndexA] == nullptr)
+		if (PoseHistories[IndexA] == nullptr)
 		{
 			return false;
 		}
 
 		for (int32 IndexB = IndexA + 1; IndexB < Num; ++IndexB)
 		{
-			if (HistoryCollectors[IndexA] == HistoryCollectors[IndexB])
+			if (PoseHistories[IndexA] == PoseHistories[IndexB])
 			{
 				return false;
 			}
@@ -330,15 +371,15 @@ void FInteractionSearchContextBase::TestHistoryCollectorsThreadingAccess() const
 	for (int32 Index = 0; Index < AnimContexts.Num(); ++Index)
 	{
 		const UAnimInstance* AnimInstance = Cast<UAnimInstance>(AnimContexts[Index].Get());
-		const FAnimNode_PoseSearchHistoryCollector_Base* HistoryCollector = HistoryCollectors[Index];
+		const IPoseHistory* PoseHistory = PoseHistories[Index];
 		
-		if (AnimInstance && HistoryCollector)
+		if (AnimInstance && PoseHistory)
 		{
 			const USkeleton* Skeleton = AnimInstance->GetRequiredBonesOnAnyThread().GetSkeletonAsset();
 			check(Skeleton);
 
 			FTransform ActorRootBoneTransform;
-			HistoryCollector->GetPoseHistory().GetTransformAtTime(0.f, ActorRootBoneTransform, Skeleton, RootBoneIndexType, ComponentSpaceIndexType);
+			PoseHistory->GetTransformAtTime(0.f, ActorRootBoneTransform, Skeleton, RootBoneIndexType, ComponentSpaceIndexType);
 		}
 	}
 }
@@ -370,17 +411,11 @@ void FInteractionIsland::FPreTickFunction::ExecuteTick(float DeltaTime, enum ELe
 			{
 				if (const UObject* AnimContext = SearchContext.AnimContexts[Index].Get())
 				{
-					if (const UAnimInstance* AnimInstance = Cast<UAnimInstance>(AnimContext))
-					{
-						// since FInteractionIsland has a tick dependency with the USkeletalMeshComponent it's safe modify the FAnimNode_PoseSearchHistoryCollector_Base
-						FAnimNode_PoseSearchHistoryCollector_Base* HistoryCollector = const_cast<FAnimNode_PoseSearchHistoryCollector_Base*>(SearchContext.HistoryCollectors[Index]);
-						check(HistoryCollector);
-						HistoryCollector->GenerateTrajectory(AnimInstance);
-					}
-					else
-					{
-						unimplemented();
-					}
+					const IPoseHistory* PoseHistory = SearchContext.PoseHistories[Index];
+					check(PoseHistory);
+
+					// since FInteractionIsland has a tick dependency with the USkeletalMeshComponent it's safe modify the IPoseHistory
+					const_cast<IPoseHistory*>(PoseHistory)->GenerateTrajectory(AnimContext, DeltaTime);
 				}
 			}
 		}
@@ -466,33 +501,32 @@ FInteractionIsland::~FInteractionIsland()
 
 UActorComponent* FInteractionIsland::FindPreTickComponent(const AActor* Actor)
 {
-	if (UActorComponent* CharacterMovementComponent = Actor->GetComponentByClass<UCharacterMovementComponent>())
-	{
-		return CharacterMovementComponent;
-	}
-	
 	FName Tag = "MotionMatchingInteractionPre";
 	if (UActorComponent* PreTickComponent = Actor->FindComponentByTag<UActorComponent>(Tag))
 	{
 		return PreTickComponent;
 	}
 
+	if (UActorComponent* CharacterMovementComponent = Actor->GetComponentByClass<UCharacterMovementComponent>())
+	{
+		return CharacterMovementComponent;
+	}
+
 	UE_LOG(LogPoseSearch, Error, TEXT("FInteractionIsland::FindPreTickComponent requires AActor %s to have a UCharacterMovementComponent, or a component tagged as 'MotionMatchingInteractionPre' to work!"), *Actor->GetName());
 	return nullptr;
 }
 
 UActorComponent* FInteractionIsland::FindPostTickComponent(const AActor* Actor)
 {
-	if (UActorComponent* SkeletalMeshComponent = Actor->GetComponentByClass<USkeletalMeshComponent>())
-	{
-		return SkeletalMeshComponent;
-	}
-
 	FName Tag = "MotionMatchingInteractionPost";
 	if (UActorComponent* PostTickComponent = Actor->FindComponentByTag<UActorComponent>(Tag))
 	{
 		return PostTickComponent;
+	}
 
+	if (UActorComponent* SkeletalMeshComponent = Actor->GetComponentByClass<USkeletalMeshComponent>())
+	{
+		return SkeletalMeshComponent;
 	}
 
 	UE_LOG(LogPoseSearch, Error, TEXT("FInteractionIsland::FindPostTickComponent requires AActor %s to have a USkeletalMeshComponent, or a component tagged as 'MotionMatchingInteractionPost' to work!"), *Actor->GetName());
@@ -522,60 +556,52 @@ void FInteractionIsland::InjectToActor(const UObject* AnimContext, bool bAddTick
 		}
 #endif // ENABLE_ANIM_DEBUG
 
-		if (const UAnimInstance* AnimInstance = Cast<UAnimInstance>(AnimContext))
+		const AActor* Actor = GetContextOwningActor(AnimContext);
+		if (UActorComponent* PreTickComponent = FindPreTickComponent(Actor))
 		{
-			AActor* Actor = AnimInstance->GetOwningActor();
-
-			if (UActorComponent* PreTickComponent = FindPreTickComponent(Actor))
+			if (UActorComponent* PostTickComponent = FindPostTickComponent(Actor))
 			{
-				if (UActorComponent* PostTickComponent = FindPostTickComponent(Actor))
+				const bool bIsMainActor = PostTickComponents.IsEmpty();
+
+				//	tick order: 
+				//		PreTickComponents (UCharacterMovementComponent, or UCharacterMoverComponent) ->
+				//			Island.PreTickFunction ->
+				//				first injected actor PostTickComponent (USkeletalMeshComponent, or WIP UAnimNextComponent) ->
+				//					Island.PostTickFunction ->
+				//						other PostTickComponent(s)
+				PreTickComponents.AddUnique(PreTickComponent);
+				PostTickComponents.AddUnique(PostTickComponent);
+
+				// making sure that if we add a unique PreTickComponent, we add as well a unique PostTickComponent
+				// (so we can remove them later on in a consistent fashion)
+				check(PreTickComponents.Num() == PostTickComponents.Num());
+
+				if (bAddTickDependencies)
 				{
-					const bool bIsMainActor = PostTickComponents.IsEmpty();
-
-					//	tick order: 
-					//		PreTickComponents (UCharacterMovementComponent, or UCharacterMoverComponent) ->
-					//			Island.PreTickFunction ->
-					//				first injected actor PostTickComponent (USkeletalMeshComponent, or WIP UAnimNextComponent) ->
-					//					Island.PostTickFunction ->
-					//						other PostTickComponent(s)
-					PreTickComponents.AddUnique(PreTickComponent);
-					PostTickComponents.AddUnique(PostTickComponent);
-
-					// making sure that if we add a unique PreTickComponent, we add as well a unique PostTickComponent
-					// (so we can remove them later on in a consistent fashion)
-					check(PreTickComponents.Num() == PostTickComponents.Num());
-
-					if (bAddTickDependencies)
-					{
-						PreTickFunction.AddPrerequisite(PreTickComponent, PreTickComponent->PrimaryComponentTick);
+					PreTickFunction.AddPrerequisite(PreTickComponent, PreTickComponent->PrimaryComponentTick);
 
-						if (bIsMainActor)
-						{
-							// it should be this island, but it's not a UObject, so we use the InteractionSubsystem
-							PostTickComponent->PrimaryComponentTick.AddPrerequisite(InteractionSubsystem, PreTickFunction);
-							PostTickFunction.AddPrerequisite(PostTickComponent, PostTickComponent->PrimaryComponentTick);
+					if (bIsMainActor)
+					{
+						// it should be this island, but it's not a UObject, so we use the InteractionSubsystem
+						PostTickComponent->PrimaryComponentTick.AddPrerequisite(InteractionSubsystem, PreTickFunction);
+						PostTickFunction.AddPrerequisite(PostTickComponent, PostTickComponent->PrimaryComponentTick);
 
-							check(!bHasTickDependencies);
-							bHasTickDependencies = bAddTickDependencies;
-						}
-						else
-						{
-							// it should be this island, but it's not a UObject, so we use the InteractionSubsystem
-							PostTickComponent->PrimaryComponentTick.AddPrerequisite(InteractionSubsystem, PostTickFunction);
-							check(bHasTickDependencies);
-						}
+						check(!bHasTickDependencies);
+						bHasTickDependencies = bAddTickDependencies;
 					}
 					else
 					{
-						check(!bHasTickDependencies);
+						// it should be this island, but it's not a UObject, so we use the InteractionSubsystem
+						PostTickComponent->PrimaryComponentTick.AddPrerequisite(InteractionSubsystem, PostTickFunction);
+						check(bHasTickDependencies);
 					}
 				}
+				else
+				{
+					check(!bHasTickDependencies);
+				}
 			}
 		}
-		else
-		{
-			unimplemented();
-		}
 	}
 }
 
@@ -717,14 +743,6 @@ bool FInteractionIsland::CanPerformSearch(const UObject* AnimContext) const
 {
 	check(AnimContext);
 
-	const UAnimInstance* AnimInstance = Cast<UAnimInstance>(AnimContext);
-	if (!AnimInstance)
-	{
-		// @todo: implement AnimNext support (AnimContext is not an UAnimInstance)
-		unimplemented();
-		return false;
-	}
-
 	const AActor* MainActor = GetMainActor();
 	if (!MainActor)
 	{
@@ -737,7 +755,7 @@ bool FInteractionIsland::CanPerformSearch(const UObject* AnimContext) const
 		return false;
 	}
 
-	const AActor* Actor = AnimInstance->GetOwningActor();
+	const AActor* Actor = GetContextOwningActor(AnimContext);
 	check(Actor);
 
 	if (Actor != MainActor)
@@ -805,7 +823,6 @@ bool FInteractionIsland::DoSearch_AnyThread(const UObject* AnimContext, const FP
 		FMemMark Mark(FMemStack::Get());
 
 		TArray<const UObject*, TInlineAllocator<PreallocatedRolesNum, TMemStackAllocator<>>> AnimContexts;
-		TArray<const IPoseHistory*, TInlineAllocator<PreallocatedRolesNum, TMemStackAllocator<>>> PoseHistories;
 		TArray<FSearchResult, TMemStackAllocator<>> PoseSearchResults;
 
 		// SearchContexts are modified only by UPoseSearchInteractionSubsystem::Tick and constant otherwise, so it's safe to access them in a threaded enviroment without locks
@@ -813,7 +830,7 @@ bool FInteractionIsland::DoSearch_AnyThread(const UObject* AnimContext, const FP
 
 		for (int32 SearchIndex = 0; SearchIndex < SearchContexts.Num(); ++SearchIndex)
 		{
-			const FInteractionSearchContext& SearchContext = SearchContexts[SearchIndex];
+			FInteractionSearchContext& SearchContext = SearchContexts[SearchIndex];
 			const UPoseSearchDatabase* Database = SearchContext.Database.Get();
 			if (!Database)
 			{
@@ -840,12 +857,6 @@ bool FInteractionIsland::DoSearch_AnyThread(const UObject* AnimContext, const FP
 				AnimContexts.Add(SearchContextAnimContext);
 			}
 
-			PoseHistories.Reset();
-			for (const FAnimNode_PoseSearchHistoryCollector_Base* HistoryCollector : SearchContext.HistoryCollectors)
-			{
-				PoseHistories.Add(&HistoryCollector->GetPoseHistory());
-			}
-
 			const UObject* AssetsToSearch[] = { Database };
 			FPoseSearchFutureProperties PoseSearchFutureProperties;
 
@@ -883,7 +894,7 @@ bool FInteractionIsland::DoSearch_AnyThread(const UObject* AnimContext, const FP
 
 			// @todo: we could perform multiple UPoseSearchLibrary::MotionMatch in parallel!
 			PoseSearchResults[SearchIndex] = UPoseSearchLibrary::MotionMatch(AnimContexts, SearchContext.Roles,
-				PoseHistories, AssetsToSearch, ContinuingPropertiesToUse, PoseSearchFutureProperties);
+				SearchContext.PoseHistories, AssetsToSearch, ContinuingPropertiesToUse, PoseSearchFutureProperties);
 		}
 		
 #if UE_POSE_SEARCH_ENABLE_ISLAND_CRITICAL_SECTION
diff --git a/Engine/Plugins/Animation/PoseSearch/Source/Runtime/Private/PoseSearchInteractionLibrary.cpp b/Engine/Plugins/Animation/PoseSearch/Source/Runtime/Private/PoseSearchInteractionLibrary.cpp
--- a/Engine/Plugins/Animation/PoseSearch/Source/Runtime/Private/PoseSearchInteractionLibrary.cpp
+++ b/Engine/Plugins/Animation/PoseSearch/Source/Runtime/Private/PoseSearchInteractionLibrary.cpp
@@ -21,12 +21,12 @@ FPoseSearchInteractionBlueprintResult UPoseSearchInteractionLibrary::MotionMatch
 	return MotionMatchInteraction_Pure(Availabilities, AnimContext, ContinuingProperties, PoseHistoryName, bValidateResultAgainstAvailabilities);
 }
 
-FPoseSearchInteractionBlueprintResult UPoseSearchInteractionLibrary::MotionMatchInteraction(const TArrayView<const FPoseSearchInteractionAvailability> Availabilities, const UObject* AnimContext, const FPoseSearchContinuingProperties& ContinuingProperties, const FAnimNode_PoseSearchHistoryCollector_Base* HistoryCollector, bool bValidateResultAgainstAvailabilities)
+FPoseSearchInteractionBlueprintResult UPoseSearchInteractionLibrary::MotionMatchInteraction(const TArrayView<const FPoseSearchInteractionAvailability> Availabilities, const UObject* AnimContext, const FPoseSearchContinuingProperties& ContinuingProperties, const UE::PoseSearch::IPoseHistory* PoseHistory, bool bValidateResultAgainstAvailabilities)
 {
 	FPoseSearchInteractionBlueprintResult Result;
 	if (UPoseSearchInteractionSubsystem* InteractionSubsystem = UPoseSearchInteractionSubsystem::GetSubsystem_AnyThread(AnimContext))
 	{
-		InteractionSubsystem->Query_AnyThread(Availabilities, AnimContext, ContinuingProperties, Result, FName(), HistoryCollector, bValidateResultAgainstAvailabilities);
+		InteractionSubsystem->Query_AnyThread(Availabilities, AnimContext, ContinuingProperties, Result, FName(), PoseHistory, bValidateResultAgainstAvailabilities);
 	}
 	return Result;
 }
diff --git a/Engine/Plugins/Animation/PoseSearch/Source/Runtime/Private/PoseSearchInteractionSubsystem.cpp b/Engine/Plugins/Animation/PoseSearch/Source/Runtime/Private/PoseSearchInteractionSubsystem.cpp
--- a/Engine/Plugins/Animation/PoseSearch/Source/Runtime/Private/PoseSearchInteractionSubsystem.cpp
+++ b/Engine/Plugins/Animation/PoseSearch/Source/Runtime/Private/PoseSearchInteractionSubsystem.cpp
@@ -1,6 +1,7 @@
 // Copyright Epic Games, Inc. All Rights Reserved.
 
 #include "PoseSearch/PoseSearchInteractionSubsystem.h"
+#include "PoseSearch/AnimNode_PoseSearchHistoryCollector.h"
 #include "PoseSearch/PoseSearchInteractionIsland.h"
 #include "PoseSearch/PoseSearchDatabase.h"
 #include "PoseSearch/PoseSearchRole.h"
@@ -24,23 +25,26 @@ static FAutoConsoleVariableRef CVarPoseSearchInteractionCacheIslands(TEXT("a.Pos
 
 FString FInteractionAvailability::GetPoseHistoryName() const
 {
-	if (HistoryCollector)
+	if (PoseHistory)
 	{
 		return "HistoryProvider";
 	}
 	return PoseHistoryName.ToString();
 }
-	
-const FAnimNode_PoseSearchHistoryCollector_Base* FInteractionAvailability::GetHistoryCollector(const UObject* AnimContext) const
+
+const IPoseHistory* FInteractionAvailability::GetPoseHistory(const UObject* AnimContext) const
 {
-	if (HistoryCollector)
+	if (PoseHistory)
 	{
-		return HistoryCollector;
+		return PoseHistory;
 	}
 		
 	if (const UAnimInstance* AnimInstance = Cast<UAnimInstance>(AnimContext))
 	{
-		return UPoseSearchLibrary::FindPoseHistoryNode(PoseHistoryName, AnimInstance);
+		if (const FAnimNode_PoseSearchHistoryCollector_Base* PoseSearchHistoryCollector = UPoseSearchLibrary::FindPoseHistoryNode(PoseHistoryName, AnimInstance))
+		{
+			return PoseSearchHistoryCollector;
+		}
 	}
 
 	unimplemented();
@@ -399,29 +403,22 @@ UE::PoseSearch::FInteractionIsland* UPoseSearchInteractionSubsystem::FindIsland(
 
 	if (AnimContext)
 	{
-		if (const UAnimInstance* AnimInstance = Cast<UAnimInstance>(AnimContext))
+		if (const AActor* Actor = GetContextOwningActor(AnimContext))
 		{
-			if (AActor* Actor = AnimInstance->GetOwningActor())
+			if (UActorComponent* PostTickComponent = FInteractionIsland::FindPostTickComponent(Actor))
 			{
-				if (UActorComponent* PostTickComponent = FInteractionIsland::FindPostTickComponent(Actor))
+				for (FInteractionIsland* Island : Islands)
 				{
-					for (FInteractionIsland* Island : Islands)
+					for (const TWeakObjectPtr<UActorComponent>& IslandPostTickComponent : Island->GetPostTickComponents())
 					{
-						for (const TWeakObjectPtr<UActorComponent>& IslandPostTickComponent : Island->GetPostTickComponents())
+						if (IslandPostTickComponent.Get() == PostTickComponent)
 						{
-							if (IslandPostTickComponent.Get() == PostTickComponent)
-							{
-								return Island;
-							}
+							return Island;
 						}
 					}
 				}
 			}
 		}
-		else
-		{
-			unimplemented();
-		}
 	}
 	return nullptr;
 }
@@ -468,10 +465,10 @@ bool UPoseSearchInteractionSubsystem::IsThereAnyAvailability() const
 	return !AvailabilitiesMap.IsEmpty();
 }
 
-void UPoseSearchInteractionSubsystem::AddAvailabilities(const TArrayView<const FPoseSearchInteractionAvailability> Availabilities, const UObject* AnimContext, FName PoseHistoryName, const FAnimNode_PoseSearchHistoryCollector_Base* HistoryCollector)
+void UPoseSearchInteractionSubsystem::AddAvailabilities(const TArrayView<const FPoseSearchInteractionAvailability> Availabilities, const UObject* AnimContext, FName PoseHistoryName, const UE::PoseSearch::IPoseHistory* PoseHistory)
 {
 	FScopeLock Lock(&AvailabilitiesMapMutex);
-	AvailabilitiesMap.AddAvailabilities(Availabilities, AnimContext, PoseHistoryName, HistoryCollector);
+	AvailabilitiesMap.AddAvailabilities(Availabilities, AnimContext, PoseHistoryName, PoseHistory);
 }
 
 #if ENABLE_DRAW_DEBUG && ENABLE_ANIM_DEBUG
@@ -622,10 +619,10 @@ void UPoseSearchInteractionSubsystem::GenerateSearchContexts(UE::PoseSearch::FIn
 
 	struct FRoledAnimContextInfo
 	{
-		FRoledAnimContextInfo(const FAnimContextInfo& InAnimContextInfo, const FRole InRole, const FAnimNode_PoseSearchHistoryCollector_Base& InHistoryCollector, const UPoseSearchDatabase& InDatabase, float InMaxDistanceSquared, bool bInDisableCollisions)
+		FRoledAnimContextInfo(const FAnimContextInfo& InAnimContextInfo, const FRole InRole, const IPoseHistory* InPoseHistory, const UPoseSearchDatabase& InDatabase, float InMaxDistanceSquared, bool bInDisableCollisions)
 			: AnimContextInfo(&InAnimContextInfo)
 			, Role(InRole)
-			, HistoryCollector(&InHistoryCollector)
+			, PoseHistory(InPoseHistory)
 			, Database(&InDatabase)
 			, MaxDistanceSquared(InMaxDistanceSquared)
 			, bDisableCollisions(bInDisableCollisions)
@@ -637,14 +634,14 @@ void UPoseSearchInteractionSubsystem::GenerateSearchContexts(UE::PoseSearch::FIn
 
 		const FAnimContextInfo* AnimContextInfo = nullptr;
 		FRole Role = DefaultRole;
-		const FAnimNode_PoseSearchHistoryCollector_Base* HistoryCollector = nullptr;
+		const IPoseHistory* PoseHistory = nullptr;
 		const UPoseSearchDatabase* Database = nullptr;
 		float MaxDistanceSquared = 0.f;
 		bool bDisableCollisions = false;
 	};
 	struct FRoledAnimContextInfos : TArray<FRoledAnimContextInfo, TInlineAllocator<PreallocatedRolesNum, TMemStackAllocator<>>>
 	{
-		void AddRoledAnimContextInfos(const FAnimContextInfo& AnimContextInfo, const TArray<FName>& RolesFilter, const FAnimNode_PoseSearchHistoryCollector_Base& HistoryCollector, const UPoseSearchDatabase& Database, float MaxDistanceSquared, bool bDisableCollisions)
+		void AddRoledAnimContextInfos(const FAnimContextInfo& AnimContextInfo, const TArray<FName>& RolesFilter, const IPoseHistory* PoseHistory, const UPoseSearchDatabase& Database, float MaxDistanceSquared, bool bDisableCollisions)
 		{
 			const UPoseSearchSchema* Schema = Database.Schema;
 			check(Schema);
@@ -653,7 +650,7 @@ void UPoseSearchInteractionSubsystem::GenerateSearchContexts(UE::PoseSearch::FIn
 				// adding ALL the possible roles from the database:
 				for (const FPoseSearchRoledSkeleton& RoledSkeleton : Schema->GetRoledSkeletons())
 				{
-					AddUnique(FRoledAnimContextInfo(AnimContextInfo, RoledSkeleton.Role, HistoryCollector, Database, MaxDistanceSquared, bDisableCollisions));
+					AddUnique(FRoledAnimContextInfo(AnimContextInfo, RoledSkeleton.Role, PoseHistory, Database, MaxDistanceSquared, bDisableCollisions));
 				}
 			}
 			else
@@ -662,7 +659,7 @@ void UPoseSearchInteractionSubsystem::GenerateSearchContexts(UE::PoseSearch::FIn
 				{
 					if (Schema->GetRoledSkeleton(Role))
 					{
-						AddUnique(FRoledAnimContextInfo(AnimContextInfo, Role, HistoryCollector, Database, MaxDistanceSquared, bDisableCollisions));
+						AddUnique(FRoledAnimContextInfo(AnimContextInfo, Role, PoseHistory, Database, MaxDistanceSquared, bDisableCollisions));
 					}
 					else
 					{
@@ -749,14 +746,14 @@ void UPoseSearchInteractionSubsystem::GenerateSearchContexts(UE::PoseSearch::FIn
 				{
 					UE_LOG(LogPoseSearch, Log, TEXT("UPoseSearchInteractionSubsystem::GenerateSearchContexts too small Availability.BroadPhaseRadius. Character will not interact with anything within such small radius!"));
 				}
-				else if (const FAnimNode_PoseSearchHistoryCollector_Base* HistoryCollector = Availability.GetHistoryCollector(AnimContextInfo.AnimContext))
+				else if (const IPoseHistory* PoseHistory = Availability.GetPoseHistory(AnimContextInfo.AnimContext))
 				{
 					if (const UPoseSearchDatabase* Database = Availability.Database.Get())
 					{
 						check(Database->Schema);
 						FRoledAnimContextInfos& RoledAnimContextInfos = DatabaseToRoledAnimContextInfos.FindOrAdd(Database);
 						const float MaxDistanceSquared = Availability.BroadPhaseRadius * Availability.BroadPhaseRadius;
-						RoledAnimContextInfos.AddRoledAnimContextInfos(AnimContextInfo, Availability.RolesFilter, *HistoryCollector, *Database, MaxDistanceSquared, Availability.bDisableCollisions);
+						RoledAnimContextInfos.AddRoledAnimContextInfos(AnimContextInfo, Availability.RolesFilter, PoseHistory, *Database, MaxDistanceSquared, Availability.bDisableCollisions);
 					}
 					else if (Availability.IsTagValid())
 					{
@@ -769,7 +766,7 @@ void UPoseSearchInteractionSubsystem::GenerateSearchContexts(UE::PoseSearch::FIn
 							{
 								check(DatabaseFromTag && DatabaseFromTag->Schema);
 								FRoledAnimContextInfos& RoledAnimContextInfos = DatabaseToRoledAnimContextInfos.FindOrAdd(DatabaseFromTag);
-								RoledAnimContextInfos.AddRoledAnimContextInfos(AnimContextInfo, Availability.RolesFilter, *HistoryCollector, *DatabaseFromTag, MaxDistanceSquared, Availability.bDisableCollisions);
+								RoledAnimContextInfos.AddRoledAnimContextInfos(AnimContextInfo, Availability.RolesFilter, PoseHistory, *DatabaseFromTag, MaxDistanceSquared, Availability.bDisableCollisions);
 							}
 						}
 						else
@@ -846,7 +843,7 @@ void UPoseSearchInteractionSubsystem::GenerateSearchContexts(UE::PoseSearch::FIn
 						FInteractionSearchContext SearchContext;
 						SearchContext.Database = Database;
 						SearchContext.AnimContexts.Reserve(CombinationCardinality);
-						SearchContext.HistoryCollectors.Reserve(CombinationCardinality);
+						SearchContext.PoseHistories.Reserve(CombinationCardinality);
 						SearchContext.Roles.Reserve(CombinationCardinality);
 
 						// checking if this combination is valid for the Database:
@@ -871,7 +868,7 @@ void UPoseSearchInteractionSubsystem::GenerateSearchContexts(UE::PoseSearch::FIn
 
 							SearchContext.AnimContexts.Add(RoledAnimContextInfo.AnimContextInfo->AnimContext);
 							SearchContext.Roles.Add(RoledAnimContextInfo.Role);
-							SearchContext.HistoryCollectors.Add(RoledAnimContextInfo.HistoryCollector);
+							SearchContext.PoseHistories.Add(RoledAnimContextInfo.PoseHistory);
 							SearchContext.bDisableCollisions |= RoledAnimContextInfo.bDisableCollisions;
 						}
 
@@ -1134,7 +1131,7 @@ TStatId UPoseSearchInteractionSubsystem::GetStatId() const
 
 void UPoseSearchInteractionSubsystem::Query_AnyThread(const TArrayView<const FPoseSearchInteractionAvailability> Availabilities, const UObject* AnimContext, 
 	const FPoseSearchContinuingProperties& ContinuingProperties, FPoseSearchInteractionBlueprintResult& Result,
-	FName PoseHistoryName, const FAnimNode_PoseSearchHistoryCollector_Base* HistoryCollector, bool bValidateResultAgainstAvailabilities)
+	FName PoseHistoryName, const UE::PoseSearch::IPoseHistory* PoseHistory, bool bValidateResultAgainstAvailabilities)
 {
 	QUICK_SCOPE_CYCLE_COUNTER(STAT_UPoseSearchInteractionSubsystem_Query_AnyThread);
 
@@ -1178,7 +1175,7 @@ void UPoseSearchInteractionSubsystem::Query_AnyThread(const TArrayView<const FPo
 	// queuing the availabilities for the next frame Query_AnyThread
 	if (!Availabilities.IsEmpty())
 	{
-		AddAvailabilities(Availabilities, AnimContext, PoseHistoryName, HistoryCollector);
+		AddAvailabilities(Availabilities, AnimContext, PoseHistoryName, PoseHistory);
 	}
 }
 
diff --git a/Engine/Plugins/Animation/PoseSearch/Source/Runtime/Private/PoseSearchInteractionValidator.cpp b/Engine/Plugins/Animation/PoseSearch/Source/Runtime/Private/PoseSearchInteractionValidator.cpp
--- a/Engine/Plugins/Animation/PoseSearch/Source/Runtime/Private/PoseSearchInteractionValidator.cpp
+++ b/Engine/Plugins/Animation/PoseSearch/Source/Runtime/Private/PoseSearchInteractionValidator.cpp
@@ -7,7 +7,7 @@ namespace UE::PoseSearch
 {
 #if ENABLE_ANIM_DEBUG
 
-FInteractionValidator::FInteractionValidator(UObject* AnimContext)
+FInteractionValidator::FInteractionValidator(const UObject* AnimContext)
 	: ValidatingAnimContext(AnimContext)
 	, ValidatingIsland(nullptr)
 {
@@ -24,15 +24,7 @@ FInteractionValidator::FInteractionValidator(UObject* AnimContext)
 		return;
 	}
 
-	const UAnimInstance* AnimInstance = Cast<UAnimInstance>(AnimContext);
-	if (!AnimInstance)
-	{
-		// @todo: implement AnimNext support (AnimContext is not an UAnimInstance)
-		unimplemented();
-		return;
-	}
-
-	const AActor* Actor = AnimInstance->GetOwningActor();
+	const AActor* Actor = GetContextOwningActor(AnimContext);
 	check(Actor);
 
 	const AActor* MainActor = ValidatingIsland->GetMainActor();
@@ -107,15 +99,7 @@ FInteractionValidator::~FInteractionValidator()
 		return;
 	}
 
-	const UAnimInstance* AnimInstance = Cast<UAnimInstance>(ValidatingAnimContext);
-	if (!AnimInstance)
-	{
-		// @todo: implement AnimNext support (AnimContext is not an UAnimInstance)
-		unimplemented();
-		return;
-	}
-
-	const AActor* Actor = AnimInstance->GetOwningActor();
+	const AActor* Actor = GetContextOwningActor(ValidatingAnimContext);
 	check(Actor);
 
 	const AActor* MainActor = Island->GetMainActor();
diff --git a/Engine/Plugins/Animation/PoseSearch/Source/Runtime/Private/PoseSearchLibrary.cpp b/Engine/Plugins/Animation/PoseSearch/Source/Runtime/Private/PoseSearchLibrary.cpp
--- a/Engine/Plugins/Animation/PoseSearch/Source/Runtime/Private/PoseSearchLibrary.cpp
+++ b/Engine/Plugins/Animation/PoseSearch/Source/Runtime/Private/PoseSearchLibrary.cpp
@@ -905,7 +905,7 @@ void UPoseSearchLibrary::MotionMatch(
 	{
 		if (const FAnimNode_PoseSearchHistoryCollector_Base* PoseHistoryNode = FindPoseHistoryNode(PoseHistoryName, AnimInstance))
 		{
-			PoseHistories.Add(&PoseHistoryNode->GetPoseHistory());
+			PoseHistories.Add(PoseHistoryNode);
 		}
 
 		AnimContexts.Add(AnimInstance);
diff --git a/Engine/Plugins/Animation/PoseSearch/Source/Runtime/Private/PoseSearchTrajectoryLibrary.cpp b/Engine/Plugins/Animation/PoseSearch/Source/Runtime/Private/PoseSearchTrajectoryLibrary.cpp
--- a/Engine/Plugins/Animation/PoseSearch/Source/Runtime/Private/PoseSearchTrajectoryLibrary.cpp
+++ b/Engine/Plugins/Animation/PoseSearch/Source/Runtime/Private/PoseSearchTrajectoryLibrary.cpp
@@ -34,13 +34,15 @@ bool FPoseSearchTrajectoryData::UpdateData(
 	const ACharacter* Character = Cast<ACharacter>(Context);
 	if (!Character)
 	{
-		const UAnimInstance* AnimInstance = Cast<UAnimInstance>(Context);
-		if (!AnimInstance)
+		if (const UAnimInstance* AnimInstance = Cast<UAnimInstance>(Context))
 		{
-			return false;
+			Character = Cast<ACharacter>(AnimInstance->GetOwningActor());
 		}
-
-		Character = Cast<ACharacter>(AnimInstance->GetOwningActor());
+		else if (const UActorComponent* AnimNextComponent = Cast<UActorComponent>(Context))
+		{
+			Character = Cast<ACharacter>(AnimNextComponent->GetOwner());
+		}
+		
 		if (!Character)
 		{
 			return false;
diff --git a/Engine/Plugins/Animation/PoseSearch/Source/Runtime/Public/PoseSearch/AnimNode_PoseSearchHistoryCollector.h b/Engine/Plugins/Animation/PoseSearch/Source/Runtime/Public/PoseSearch/AnimNode_PoseSearchHistoryCollector.h
--- a/Engine/Plugins/Animation/PoseSearch/Source/Runtime/Public/PoseSearch/AnimNode_PoseSearchHistoryCollector.h
+++ b/Engine/Plugins/Animation/PoseSearch/Source/Runtime/Public/PoseSearch/AnimNode_PoseSearchHistoryCollector.h
@@ -8,7 +8,7 @@
 #include "AnimNode_PoseSearchHistoryCollector.generated.h"
 
 USTRUCT(BlueprintInternalUseOnly)
-struct POSESEARCH_API FAnimNode_PoseSearchHistoryCollector_Base : public FAnimNode_Base
+struct POSESEARCH_API FAnimNode_PoseSearchHistoryCollector_Base : public FAnimNode_Base, public UE::PoseSearch::IPoseHistory
 {
 	GENERATED_BODY()
 
@@ -99,8 +99,33 @@ public:
 	virtual void Update_AnyThread(const FAnimationUpdateContext& Context) override;
 	// End of FAnimNode_Base interface
 
+	// UE::PoseSearch::IPoseHistory interface
+	virtual bool GetTransformAtTime(float Time, FTransform& OutBoneTransform, const USkeleton* BoneIndexSkeleton = nullptr, FBoneIndexType BoneIndexType = UE::PoseSearch::RootBoneIndexType, FBoneIndexType ReferenceBoneIndexType = UE::PoseSearch::ComponentSpaceIndexType, bool bExtrapolate = false) const override;
+	virtual bool GetCurveValueAtTime(float Time, const FName& CurveName, float& OutCurveValue, bool bExtrapolate = false) const override;
+	virtual const FPoseSearchQueryTrajectory& GetTrajectory() const override;
+	// Experimental, this feature might be removed without warning, not for production use
+	virtual void SetTrajectory(const FPoseSearchQueryTrajectory& InTrajectory, float InTrajectorySpeedMultiplier = 1.f) override;
+	virtual void GenerateTrajectory(const UObject* AnimContext, float DeltaTime) override;
+	virtual float GetTrajectorySpeedMultiplier() const override;
+	virtual bool IsEmpty() const override;
+	virtual const UE::PoseSearch::FBoneToTransformMap& GetBoneToTransformMap() const override;
+	virtual const TConstArrayView<FName> GetCollectedCurves() const override;
+	virtual int32 GetNumEntries() const override;
+	virtual const UE::PoseSearch::FPoseHistoryEntry& GetEntry(int32 EntryIndex) const override;
+
+#if ENABLE_DRAW_DEBUG && ENABLE_ANIM_DEBUG
+	virtual void DebugDraw(const UWorld* World, FColor Color) const override;
+	virtual void DebugDraw(FAnimInstanceProxy& AnimInstanceProxy, FColor Color) const override;
+#endif
+	// End of UE::PoseSearch::IPoseHistory interface
+
+	UE_DEPRECATED(5.6, "Use this::IPoseHistory interface instead")
 	const UE::PoseSearch::FPoseHistory& GetPoseHistory() const { return PoseHistory; }
+
+	UE_DEPRECATED(5.6, "Use this::IPoseHistory interface instead")
 	UE::PoseSearch::FPoseHistory& GetPoseHistory() { return PoseHistory; }
+	
+	UE_DEPRECATED(5.6, "Use this::IPoseHistory::GenerateTrajectory interface instead")
 	void GenerateTrajectory(const UAnimInstance* InAnimInstance);
 
 protected:
diff --git a/Engine/Plugins/Animation/PoseSearch/Source/Runtime/Public/PoseSearch/PoseHistoryProvider.h b/Engine/Plugins/Animation/PoseSearch/Source/Runtime/Public/PoseSearch/PoseHistoryProvider.h
--- a/Engine/Plugins/Animation/PoseSearch/Source/Runtime/Public/PoseSearch/PoseHistoryProvider.h
+++ b/Engine/Plugins/Animation/PoseSearch/Source/Runtime/Public/PoseSearch/PoseHistoryProvider.h
@@ -16,13 +16,13 @@ class POSESEARCH_API FPoseHistoryProvider : public UE::Anim::IGraphMessage
 	DECLARE_ANIMGRAPH_MESSAGE(FPoseHistoryProvider);
 
 public:
-	FPoseHistoryProvider(const FAnimNode_PoseSearchHistoryCollector_Base* InHistoryCollector) : PoseHistory(nullptr), HistoryCollector(InHistoryCollector) { }
+	FPoseHistoryProvider(const FAnimNode_PoseSearchHistoryCollector_Base* InHistoryCollector) : HistoryCollector(InHistoryCollector) { check(HistoryCollector); }
 	const IPoseHistory& GetPoseHistory() const;
+	
+	UE_DEPRECATED(5.6, "Use GetPoseHistory instead")
 	const FAnimNode_PoseSearchHistoryCollector_Base* GetHistoryCollector() const { return HistoryCollector; }
 
 private:
-	// delete after deleting FPoseHistoryProvider(const IPoseHistory& InPoseHistory)
-	const IPoseHistory* PoseHistory = nullptr;
 	const FAnimNode_PoseSearchHistoryCollector_Base* HistoryCollector = nullptr;
 };
 
diff --git a/Engine/Plugins/Animation/PoseSearch/Source/Runtime/Public/PoseSearch/PoseSearchHistory.h b/Engine/Plugins/Animation/PoseSearch/Source/Runtime/Public/PoseSearch/PoseSearchHistory.h
--- a/Engine/Plugins/Animation/PoseSearch/Source/Runtime/Public/PoseSearch/PoseSearchHistory.h
+++ b/Engine/Plugins/Animation/PoseSearch/Source/Runtime/Public/PoseSearch/PoseSearchHistory.h
@@ -84,7 +84,10 @@ struct POSESEARCH_API IPoseHistory
 	// @todo: consider consolidating into a (templated?) get X value at time once we add custom attributes to pose history.
 	virtual bool GetCurveValueAtTime(float Time, const FName& CurveName, float& OutCurveValue, bool bExtrapolate = false) const = 0;
 	virtual const FPoseSearchQueryTrajectory& GetTrajectory() const = 0;
-	
+	// Experimental, this feature might be removed without warning, not for production use
+	virtual void SetTrajectory(const FPoseSearchQueryTrajectory& InTrajectory, float InTrajectorySpeedMultiplier = 1.f) = 0;
+	// Experimental, this feature might be removed without warning, not for production use
+	virtual void GenerateTrajectory(const UObject* AnimContext, float DeltaTime) = 0;
 	// @todo: deprecate this API. TrajectorySpeedMultiplier should be a global query scaling value passed as input parameter of FSearchContext during config BuildQuery
 	virtual float GetTrajectorySpeedMultiplier() const = 0;
 	virtual bool IsEmpty() const = 0;
@@ -109,6 +112,10 @@ struct POSESEARCH_API FArchivedPoseHistory : public IPoseHistory
 	virtual bool GetTransformAtTime(float Time, FTransform& OutBoneTransform, const USkeleton* BoneIndexSkeleton = nullptr, FBoneIndexType BoneIndexType = RootBoneIndexType, FBoneIndexType ReferenceBoneIndexType = ComponentSpaceIndexType, bool bExtrapolate = false) const override;
 	virtual bool GetCurveValueAtTime(float Time, const FName& CurveName, float& outcurvevalue, bool bExtrapolate = false) const override;
 	virtual const FPoseSearchQueryTrajectory& GetTrajectory() const override { return Trajectory; }
+	// Experimental, this feature might be removed without warning, not for production use
+	virtual void SetTrajectory(const FPoseSearchQueryTrajectory& InTrajectory, float InTrajectorySpeedMultiplier = 1.f) override { Trajectory = InTrajectory; }
+	// Experimental, this feature might be removed without warning, not for production use
+	virtual void GenerateTrajectory(const UObject* AnimContext, float DeltaTime) override { checkNoEntry(); }
 	virtual float GetTrajectorySpeedMultiplier() const override { return 1.f; }
 	virtual bool IsEmpty() const override { return Entries.IsEmpty(); }
 	virtual const FBoneToTransformMap& GetBoneToTransformMap() const override { return BoneToTransformMap; }
@@ -138,7 +145,7 @@ struct POSESEARCH_API FPoseHistory : public IPoseHistory
 	FPoseHistory& operator=(const FPoseHistory& Other);
 	FPoseHistory& operator=(FPoseHistory&& Other);
 
-	void GenerateTrajectory(const UAnimInstance* AnimInstance, float DeltaTime, const FPoseSearchTrajectoryData& TrajectoryData, const FPoseSearchTrajectoryData::FSampling& TrajectoryDataSampling);
+	void GenerateTrajectory(const UObject* AnimContext, float DeltaTime, const FPoseSearchTrajectoryData& TrajectoryData, const FPoseSearchTrajectoryData::FSampling& TrajectoryDataSampling);
 	
 	UE_DEPRECATED(5.6, "No longer necessary method")
 	void PreUpdate() {}
@@ -159,12 +166,15 @@ struct POSESEARCH_API FPoseHistory : public IPoseHistory
 	virtual bool GetTransformAtTime(float Time, FTransform& OutBoneTransform, const USkeleton* BoneIndexSkeleton = nullptr, FBoneIndexType BoneIndexType = RootBoneIndexType, FBoneIndexType ReferenceBoneIndexType = ComponentSpaceIndexType, bool bExtrapolate = false) const override;
 	virtual bool GetCurveValueAtTime(float Time, const FName& CurveName, float& outcurvevalue, bool bExtrapolate = false) const override;
 	virtual const FPoseSearchQueryTrajectory& GetTrajectory() const override;
+	// Experimental, this feature might be removed without warning, not for production use
+	virtual void SetTrajectory(const FPoseSearchQueryTrajectory& InTrajectory, float InTrajectorySpeedMultiplier = 1.f) override;
+	// Experimental, this feature might be removed without warning, not for production use
+	virtual void GenerateTrajectory(const UObject* AnimContext, float DeltaTime) override;
 	virtual float GetTrajectorySpeedMultiplier() const override;
 	virtual bool IsEmpty() const override;
 	virtual const FBoneToTransformMap& GetBoneToTransformMap() const override;
 	virtual const TConstArrayView<FName> GetCollectedCurves() const override;
 	
-	void SetTrajectory(const FPoseSearchQueryTrajectory& InTrajectory, float InTrajectorySpeedMultiplier = 1.f);
 	virtual int32 GetNumEntries() const override;
 	virtual const FPoseHistoryEntry& GetEntry(int32 EntryIndex) const override;
 
@@ -219,6 +229,10 @@ struct POSESEARCH_API FMemStackPoseHistory : public IPoseHistory
 	virtual bool GetTransformAtTime(float Time, FTransform& OutBoneTransform, const USkeleton* BoneIndexSkeleton = nullptr, FBoneIndexType BoneIndexType = RootBoneIndexType, FBoneIndexType ReferenceBoneIndexType = ComponentSpaceIndexType, bool bExtrapolate = false) const override;
 	virtual bool GetCurveValueAtTime(float Time, const FName& CurveName, float& outcurvevalue, bool bExtrapolate = false) const override;
 	virtual const FPoseSearchQueryTrajectory& GetTrajectory() const override { check(PoseHistory); return PoseHistory->GetTrajectory(); }
+	// Experimental, this feature might be removed without warning, not for production use
+	virtual void SetTrajectory(const FPoseSearchQueryTrajectory& InTrajectory, float InTrajectorySpeedMultiplier = 1.f) override;
+	// Experimental, this feature might be removed without warning, not for production use
+	virtual void GenerateTrajectory(const UObject* AnimContext, float DeltaTime) override;
 	virtual float GetTrajectorySpeedMultiplier() const override { check(PoseHistory); return PoseHistory->GetTrajectorySpeedMultiplier(); }
 	virtual bool IsEmpty() const override { check(PoseHistory); return PoseHistory->IsEmpty() && FutureEntries.IsEmpty(); }
 	virtual const FBoneToTransformMap& GetBoneToTransformMap() const override { check(PoseHistory); return PoseHistory->GetBoneToTransformMap(); }
diff --git a/Engine/Plugins/Animation/PoseSearch/Source/Runtime/Public/PoseSearch/PoseSearchInteractionIsland.h b/Engine/Plugins/Animation/PoseSearch/Source/Runtime/Public/PoseSearch/PoseSearchInteractionIsland.h
--- a/Engine/Plugins/Animation/PoseSearch/Source/Runtime/Public/PoseSearch/PoseSearchInteractionIsland.h
+++ b/Engine/Plugins/Animation/PoseSearch/Source/Runtime/Public/PoseSearch/PoseSearchInteractionIsland.h
@@ -17,6 +17,13 @@ struct FInteractionIsland;
 
 // Experimental, this feature might be removed without warning, not for production use
 const FTransform& GetContextTransform(const UObject* AnimContext);
+
+// Experimental, this feature might be removed without warning, not for production use
+const USkeleton* GetContextSkeleton(const UObject* AnimContext);
+
+// Experimental, this feature might be removed without warning, not for production use
+const AActor* GetContextOwningActor(const UObject* AnimContext);
+
 // Experimental, this feature might be removed without warning, not for production use
 FVector GetContextLocation(const UObject* AnimContext);
 
@@ -27,7 +34,7 @@ struct FInteractionSearchContextBase
 	// add an array of structs instead of many arrays
 	TArray<TWeakObjectPtr<const UObject>, TInlineAllocator<PreallocatedRolesNum>> AnimContexts;
 	// @todo: perhaps use the animnode of the pose history collector instead
-	TArray<const FAnimNode_PoseSearchHistoryCollector_Base*, TInlineAllocator<PreallocatedRolesNum>> HistoryCollectors;
+	TArray<const UE::PoseSearch::IPoseHistory*, TInlineAllocator<PreallocatedRolesNum>> PoseHistories;
 	TArray<UE::PoseSearch::FRole, TInlineAllocator<PreallocatedRolesNum>> Roles;
 
 	TWeakObjectPtr<const UPoseSearchDatabase> Database;
diff --git a/Engine/Plugins/Animation/PoseSearch/Source/Runtime/Public/PoseSearch/PoseSearchInteractionLibrary.h b/Engine/Plugins/Animation/PoseSearch/Source/Runtime/Public/PoseSearch/PoseSearchInteractionLibrary.h
--- a/Engine/Plugins/Animation/PoseSearch/Source/Runtime/Public/PoseSearch/PoseSearchInteractionLibrary.h
+++ b/Engine/Plugins/Animation/PoseSearch/Source/Runtime/Public/PoseSearch/PoseSearchInteractionLibrary.h
@@ -104,7 +104,7 @@ public:
 
 // Experimental, this feature might be removed without warning, not for production use
 UCLASS(Experimental)
-class UPoseSearchInteractionLibrary : public UBlueprintFunctionLibrary
+class POSESEARCH_API UPoseSearchInteractionLibrary : public UBlueprintFunctionLibrary
 {
 	GENERATED_BODY()
 
@@ -128,7 +128,7 @@ public:
 	static FPoseSearchInteractionBlueprintResult MotionMatchInteraction(TArray<FPoseSearchInteractionAvailability> Availabilities, const UObject* AnimContext, FPoseSearchContinuingProperties ContinuingProperties, FName PoseHistoryName, bool bValidateResultAgainstAvailabilities = true);
 
 	// version of MotionMatchInteraction_Pure referencing directly the HistoryCollector rather than looking for it by name
-	static FPoseSearchInteractionBlueprintResult MotionMatchInteraction(const TArrayView<const FPoseSearchInteractionAvailability> Availabilities, const UObject* AnimContext, const FPoseSearchContinuingProperties& ContinuingProperties, const FAnimNode_PoseSearchHistoryCollector_Base* HistoryCollector, bool bValidateResultAgainstAvailabilities = true);
+	static FPoseSearchInteractionBlueprintResult MotionMatchInteraction(const TArrayView<const FPoseSearchInteractionAvailability> Availabilities, const UObject* AnimContext, const FPoseSearchContinuingProperties& ContinuingProperties, const UE::PoseSearch::IPoseHistory* PoseHistory, bool bValidateResultAgainstAvailabilities = true);
 
 	UFUNCTION(Experimental, BlueprintPure, Category = "Animation|Pose Search", meta = (BlueprintThreadSafe))
 	static FPoseSearchContinuingProperties GetMontageContinuingProperties(UAnimInstance* AnimInstance);
diff --git a/Engine/Plugins/Animation/PoseSearch/Source/Runtime/Public/PoseSearch/PoseSearchInteractionSubsystem.h b/Engine/Plugins/Animation/PoseSearch/Source/Runtime/Public/PoseSearch/PoseSearchInteractionSubsystem.h
--- a/Engine/Plugins/Animation/PoseSearch/Source/Runtime/Public/PoseSearch/PoseSearchInteractionSubsystem.h
+++ b/Engine/Plugins/Animation/PoseSearch/Source/Runtime/Public/PoseSearch/PoseSearchInteractionSubsystem.h
@@ -17,19 +17,19 @@ namespace UE::PoseSearch
 // Experimental, this feature might be removed without warning, not for production use
 struct FInteractionAvailability : public FPoseSearchInteractionAvailability
 {
-	FInteractionAvailability(const FPoseSearchInteractionAvailability& InAvailability, FName InPoseHistoryName, const FAnimNode_PoseSearchHistoryCollector_Base* InHistoryCollector)
+	FInteractionAvailability(const FPoseSearchInteractionAvailability& InAvailability, FName InPoseHistoryName, const IPoseHistory* InPoseHistory)
 		: FPoseSearchInteractionAvailability(InAvailability)
 		, PoseHistoryName(InPoseHistoryName)
-		, HistoryCollector(InHistoryCollector)
+		, PoseHistory(InPoseHistory)
 	{
 	}
 
 	FString GetPoseHistoryName() const;
-	const FAnimNode_PoseSearchHistoryCollector_Base* GetHistoryCollector(const UObject* AnimContext) const;
+	const IPoseHistory* GetPoseHistory(const UObject* AnimContext) const;
 
 private:
 	FName PoseHistoryName;
-	const FAnimNode_PoseSearchHistoryCollector_Base* HistoryCollector = nullptr;
+	const IPoseHistory* PoseHistory = nullptr;
 };
 
 // Experimental, this feature might be removed without warning, not for production use
@@ -45,14 +45,14 @@ struct FAvailabilitiesMap : protected TMap<TWeakObjectPtr<const UObject>, FInter
 	using Super::Find;
 	using Super::Reset;
 	
-	void AddAvailabilities(const TArrayView<const FPoseSearchInteractionAvailability> Availabilities, const UObject* AnimContext, FName PoseHistoryName, const FAnimNode_PoseSearchHistoryCollector_Base* HistoryCollector)
+	void AddAvailabilities(const TArrayView<const FPoseSearchInteractionAvailability> Availabilities, const UObject* AnimContext, FName PoseHistoryName, const IPoseHistory* PoseHistory)
 	{
 		FInteractionAvailabilities& AnimContextAvailabilities = FindOrAdd(AnimContext);
 		for (const FPoseSearchInteractionAvailability& Availability : Availabilities)
 		{
 			if (Availability.Database || Availability.IsTagValid())
 			{
-				AnimContextAvailabilities.Emplace(Availability, PoseHistoryName, HistoryCollector);
+				AnimContextAvailabilities.Emplace(Availability, PoseHistoryName, PoseHistory);
 			}
 			// @todo: make this Warning less spammy and reenable it
 			//else
@@ -138,12 +138,12 @@ public:
 	// publishing FPoseSearchInteractionAvailability(s) for the requesting character (via AnimContext as UAnimInstance or UAnimNextComponent) 
 	// and getting the FPoseSearchInteractionBlueprintResult from the PREVIOUS Tick update
 	// containing the animation to play at what time, and the assigned role to partecipate in the selected interaction.
-	// Either a PoseHistoryName or a HistoryCollector are required to perform the associated motion matching searches
+	// Either a PoseHistoryName or a PoseHistory are required to perform the associated motion matching searches
 	// if bValidateResultAgainstAvailabilities is true, the result will be validated against the currently published availabilities,
 	// that could differ from the previous frame one for this AnimContext, and discarded if not valid within the context of the new availabilities
 	void Query_AnyThread(const TArrayView<const FPoseSearchInteractionAvailability> Availabilities, const UObject* AnimContext, 
 		const FPoseSearchContinuingProperties& ContinuingProperties, FPoseSearchInteractionBlueprintResult& Result,
-		FName PoseHistoryName, const FAnimNode_PoseSearchHistoryCollector_Base* HistoryCollector, bool bValidateResultAgainstAvailabilities);
+		FName PoseHistoryName, const UE::PoseSearch::IPoseHistory* PoseHistory, bool bValidateResultAgainstAvailabilities);
 
 	// returns and array containing all the other AnimContext(s) that the input AnimContext is currently interacting with
 	TArray<const UObject*> GetInteractingAnimContextsWith(const UObject* AnimContext) const;
@@ -173,7 +173,7 @@ private:
 #endif // ENABLE_DRAW_DEBUG && ENABLE_ANIM_DEBUG
 
 	bool IsThereAnyAvailability() const;
-	void AddAvailabilities(const TArrayView<const FPoseSearchInteractionAvailability> Availabilities, const UObject* AnimContext, FName PoseHistoryName, const FAnimNode_PoseSearchHistoryCollector_Base* HistoryCollector);
+	void AddAvailabilities(const TArrayView<const FPoseSearchInteractionAvailability> Availabilities, const UObject* AnimContext, FName PoseHistoryName, const UE::PoseSearch::IPoseHistory* PoseHistory);
 	void ResetAvailabilities();
 	void GenerateAnimContextInfosAndTagToDatabases(UE::PoseSearch::FAnimContextInfos& AnimContextInfos, UE::PoseSearch::FTagToDatabases& TagToDatabases) const;
 	void GenerateSearchContexts(UE::PoseSearch::FInteractionSearchContexts& SearchContexts) const;
diff --git a/Engine/Plugins/Animation/PoseSearch/Source/Runtime/Public/PoseSearch/PoseSearchInteractionValidator.h b/Engine/Plugins/Animation/PoseSearch/Source/Runtime/Public/PoseSearch/PoseSearchInteractionValidator.h
--- a/Engine/Plugins/Animation/PoseSearch/Source/Runtime/Public/PoseSearch/PoseSearchInteractionValidator.h
+++ b/Engine/Plugins/Animation/PoseSearch/Source/Runtime/Public/PoseSearch/PoseSearchInteractionValidator.h
@@ -10,9 +10,9 @@ namespace UE::PoseSearch
 struct FInteractionIsland;
 
 // Experimental, this feature might be removed without warning, not for production use
-struct FInteractionValidator
+struct POSESEARCH_API FInteractionValidator
 {
-	explicit FInteractionValidator(UObject* AnimContext);
+	explicit FInteractionValidator(const UObject* AnimContext);
 	explicit FInteractionValidator(FInteractionIsland* InValidatingIsland);
 	~FInteractionValidator();
 
diff --git a/Engine/Plugins/Experimental/AnimNextPoseSearch/Source/AnimNextPoseSearch/Private/Module.cpp b/Engine/Plugins/Experimental/AnimNextPoseSearch/Source/AnimNextPoseSearch/Private/Module.cpp
--- a/Engine/Plugins/Experimental/AnimNextPoseSearch/Source/AnimNextPoseSearch/Private/Module.cpp
+++ b/Engine/Plugins/Experimental/AnimNextPoseSearch/Source/AnimNextPoseSearch/Private/Module.cpp
@@ -8,6 +8,7 @@
 
 #include "AlphaBlend.h"
 #include "PoseSearch/PoseSearchDatabase.h"
+#include "PoseSearch/PoseSearchInteractionLibrary.h"
 #include "PoseSearch/PoseSearchSchema.h"
 #include "PoseSearch/PoseSearchTrajectoryTypes.h"
 
@@ -27,9 +28,10 @@ public:
 
 		static UScriptStruct* const AllowedStructTypes[] =
 		{
-			FPoseSearchQueryTrajectorySample::StaticStruct(),
-			FPoseSearchQueryTrajectory::StaticStruct(),
 			FAlphaBlendArgs::StaticStruct(),
+			FPoseSearchInteractionAvailability::StaticStruct(),
+			FPoseSearchQueryTrajectory::StaticStruct(),
+			FPoseSearchQueryTrajectorySample::StaticStruct(),
 		};
 
 		FRigVMRegistry& RigVMRegistry = FRigVMRegistry::Get();
diff --git a/Engine/Plugins/Experimental/AnimNextPoseSearch/Source/AnimNextPoseSearch/Private/MotionMatchingTrait.cpp b/Engine/Plugins/Experimental/AnimNextPoseSearch/Source/AnimNextPoseSearch/Private/MotionMatchingTrait.cpp
--- a/Engine/Plugins/Experimental/AnimNextPoseSearch/Source/AnimNextPoseSearch/Private/MotionMatchingTrait.cpp
+++ b/Engine/Plugins/Experimental/AnimNextPoseSearch/Source/AnimNextPoseSearch/Private/MotionMatchingTrait.cpp
@@ -11,6 +11,7 @@
 #include "GraphInterfaces/AnimNextNativeDataInterface_AnimSequencePlayer.h"
 #include "GraphInterfaces/AnimNextNativeDataInterface_BlendSpacePlayer.h"
 #include "Module/AnimNextModuleInstance.h"
+#include "PoseSearch/MultiAnimAsset.h"
 #include "PoseSearch/PoseSearchDatabase.h"
 #include "PoseSearch/PoseSearchDerivedData.h"
 #include "TraitInterfaces/IBlendStack.h"
@@ -98,6 +99,7 @@ namespace UE::AnimNext
 		FInstanceData* InstanceData = Binding.GetInstanceData<FInstanceData>();
 		check(InstanceData);
 		FMotionMatchingState& MotionMatchingState = InstanceData->MotionMatchingState;
+		FPoseSearchInteractionBlueprintResult& MotionMatchingInteractionState = InstanceData->MotionMatchingInteractionState;
 
 		bool bNeedsReset = false;
 #if WITH_EDITOR
@@ -116,7 +118,8 @@ namespace UE::AnimNext
 		if (bNeedsReset)
 		{
 			MotionMatchingState.Reset();
-			
+			MotionMatchingInteractionState = FPoseSearchInteractionBlueprintResult();
+
 			// @todo: implement me!
 			//BlendStackBinding.ResetBlendStack();
 		}
@@ -148,104 +151,160 @@ namespace UE::AnimNext
 			}
 
 			MotionMatchingState.AdjustAssetTime(TimelinePosition);
+			MotionMatchingInteractionState.SelectedTime = TimelinePosition;
+		}
+
+		const UObject* AnimationAsset = nullptr;
+		float AssetTime = 0.f;
+		float WantedPlayRate = 1.f;
+		bool bLoop = false;
+		FVector BlendSpaceParameters = FVector::ZeroVector;
+
+		const UObject* AnimContext = nullptr;
+		if (const FAnimNextModuleInstance* ModuleInstance = Context.GetRootGraphInstance().GetModuleInstance())
+		{
+			AnimContext = ModuleInstance->GetObject();
 		}
 
-		const TConstArrayView<TObjectPtr<const UPoseSearchDatabase>> DatabasesToSearch = SharedData->GetDatabases(Binding);
-		if (DatabasesToSearch.IsEmpty())
+		const TArray<FPoseSearchInteractionAvailability>& Availabilities = SharedData->GetAvailabilities(Binding);
+		if (Availabilities.IsEmpty())
 		{
-			UE_LOG(LogPoseSearch, Error, TEXT("FMotionMatchingTrait::PreUpdate, No database assets provided for motion matching."));
+			MotionMatchingInteractionState = FPoseSearchInteractionBlueprintResult();
 		}
 		else
 		{
-			const float DeltaTime = TraitState.GetDeltaTime();
-			const UObject* AnimContext = nullptr;
-			if (const FAnimNextModuleInstance* ModuleInstance = Context.GetRootGraphInstance().GetModuleInstance())
+			// WIP motion matched interactions
+			CheckInteractionThreadSafety(AnimContext);
+
+			// we're NOT providing the UMultiAnimAsset here, relying on the interaction subsystem to resolve the ContinuingProperties properly
+			FPoseSearchContinuingProperties ContinuingProperties;
+			ContinuingProperties.PlayingAsset = MotionMatchingInteractionState.SelectedAnimation; // GetAnimAsset();
+			ContinuingProperties.PlayingAssetAccumulatedTime = MotionMatchingInteractionState.SelectedTime;
+
+			MotionMatchingInteractionState = UPoseSearchInteractionLibrary::MotionMatchInteraction(Availabilities, AnimContext, ContinuingProperties, PoseHistory, SharedData->GetbValidateResultAgainstAvailabilities(Binding));
+			check(MotionMatchingInteractionState.ActorRootTransforms.Num() == MotionMatchingInteractionState.ActorRootBoneTransforms.Num());
+		}
+
+		if (MotionMatchingInteractionState.SelectedAnimation)
+		{
+			if (!MotionMatchingInteractionState.bIsContinuingPoseSearch)
 			{
-				AnimContext = ModuleInstance->GetObject();
+				if (const UMultiAnimAsset* MultiAnimAsset = Cast<UMultiAnimAsset>(MotionMatchingInteractionState.SelectedAnimation))
+				{
+					AnimationAsset = MultiAnimAsset->GetAnimationAsset(MotionMatchingInteractionState.Role);
+				}
+				else if (UAnimationAsset* SingleAnimAsset = Cast<UAnimationAsset>(MotionMatchingInteractionState.SelectedAnimation))
+				{
+					AnimationAsset = SingleAnimAsset;
+				}
 			}
+			AssetTime = MotionMatchingInteractionState.SelectedTime;
+			WantedPlayRate = MotionMatchingInteractionState.WantedPlayRate;
+			bLoop = MotionMatchingInteractionState.bLoop;
+			BlendSpaceParameters = MotionMatchingInteractionState.BlendParameters;
 
-			FFloatInterval PoseSearchPlayRate = SharedData->GetPlayRate(Binding);
+			MotionMatchingState.Reset();
+		}
+		else
+		{
+			const TConstArrayView<TObjectPtr<const UPoseSearchDatabase>> DatabasesToSearch = SharedData->GetDatabases(Binding);
+			if (DatabasesToSearch.IsEmpty())
+			{
+				UE_LOG(LogPoseSearch, Error, TEXT("FMotionMatchingTrait::PreUpdate, No database assets provided for motion matching."));
+			}
+			else
+			{
+				const float DeltaTime = TraitState.GetDeltaTime();
+				FFloatInterval PoseSearchPlayRate = SharedData->GetPlayRate(Binding);
 #if ENABLE_ANIM_DEBUG
-			bool bPoseSearchPlayRateEnabled = GVarAnimNextMotionMatchingPlayRateEnabled == UE::Private::EPlayRateState::Enabled || GVarAnimNextMotionMatchingPlayRateEnabled == UE::Private::EPlayRateState::PoseSearchOnly;
-			PoseSearchPlayRate = bPoseSearchPlayRateEnabled ? PoseSearchPlayRate : FFloatInterval(1.f, 1.f);
+				bool bPoseSearchPlayRateEnabled = GVarAnimNextMotionMatchingPlayRateEnabled == UE::Private::EPlayRateState::Enabled || GVarAnimNextMotionMatchingPlayRateEnabled == UE::Private::EPlayRateState::PoseSearchOnly;
+				PoseSearchPlayRate = bPoseSearchPlayRateEnabled ? PoseSearchPlayRate : FFloatInterval(1.f, 1.f);
 #endif // ENABLE_ANIM_DEBUG
 
-			UPoseSearchLibrary::UpdateMotionMatchingState(
-				AnimContext,
-				PoseHistory,
-				DatabasesToSearch,
-				DeltaTime,
-				SharedData->GetPoseJumpThresholdTime(Binding),
-				SharedData->GetPoseReselectHistory(Binding),
-				SharedData->GetbShouldSearch(Binding) ? SharedData->GetSearchThrottleTime(Binding) : UE_BIG_NUMBER,
-				PoseSearchPlayRate,
-				MotionMatchingState,
-				SharedData->GetNextUpdateInterruptMode(Binding),
-				SharedData->GetbShouldUseCachedChannelData(Binding)
-				// @todo: add debug drawing suport
-//#if ENABLE_ANIM_DEBUG
-//				, CVarAnimNodeMotionMatchingDrawQuery.GetValueOnAnyThread()
-//				, CVarAnimNodeMotionMatchingDrawCurResult.GetValueOnAnyThread()
-//#endif // ENABLE_ANIM_DEBUG
-			);
+				UPoseSearchLibrary::UpdateMotionMatchingState(
+					AnimContext,
+					PoseHistory,
+					DatabasesToSearch,
+					DeltaTime,
+					SharedData->GetPoseJumpThresholdTime(Binding),
+					SharedData->GetPoseReselectHistory(Binding),
+					SharedData->GetbShouldSearch(Binding) ? SharedData->GetSearchThrottleTime(Binding) : UE_BIG_NUMBER,
+					PoseSearchPlayRate,
+					MotionMatchingState,
+					SharedData->GetNextUpdateInterruptMode(Binding),
+					SharedData->GetbShouldUseCachedChannelData(Binding)
+					// @todo: add debug drawing suport
+	//#if ENABLE_ANIM_DEBUG
+	//				, CVarAnimNodeMotionMatchingDrawQuery.GetValueOnAnyThread()
+	//				, CVarAnimNodeMotionMatchingDrawCurResult.GetValueOnAnyThread()
+	//#endif // ENABLE_ANIM_DEBUG
+				);
+			}
+
+			const bool bJumpToPose = MotionMatchingState.CurrentSearchResult.IsValid() && !MotionMatchingState.CurrentSearchResult.bIsContinuingPoseSearch;
+			if (bJumpToPose)
+			{
+				if (const UE::PoseSearch::FSearchIndexAsset* SearchIndexAsset = MotionMatchingState.CurrentSearchResult.GetSearchIndexAsset())
+				{
+					if (const FPoseSearchDatabaseAnimationAssetBase* DatabaseAnimationAssetBase = MotionMatchingState.CurrentSearchResult.Database->GetDatabaseAnimationAsset<FPoseSearchDatabaseAnimationAssetBase>(*SearchIndexAsset))
+					{
+						AnimationAsset = DatabaseAnimationAssetBase->GetAnimationAsset();
+						AssetTime = MotionMatchingState.CurrentSearchResult.AssetTime;
+						WantedPlayRate = MotionMatchingState.WantedPlayRate;
+						bLoop = SearchIndexAsset->IsLooping();
+						BlendSpaceParameters = SearchIndexAsset->GetBlendParameters();
+					}
+				}
+			}
 		}
 
-		float DesiredPlayRate = MotionMatchingState.WantedPlayRate;
+		float DesiredPlayRate = WantedPlayRate;
 #if ENABLE_ANIM_DEBUG
 		DesiredPlayRate = GVarAnimNextMotionMatchingPlayRateEnabled == UE::Private::EPlayRateState::Enabled ? DesiredPlayRate : 1.0f;
 #endif // ENABLE_ANIM_DEBUG
 
-		const bool bJumpToPose = MotionMatchingState.CurrentSearchResult.IsValid() && !MotionMatchingState.CurrentSearchResult.bIsContinuingPoseSearch;
-		if (bJumpToPose)
+		if (AnimationAsset)
 		{
-			if (const UE::PoseSearch::FSearchIndexAsset* SearchIndexAsset = MotionMatchingState.CurrentSearchResult.GetSearchIndexAsset())
+			UE::AnimNext::IBlendStack::FGraphRequest NewGraphRequest;
+			NewGraphRequest.BlendArgs = SharedData->BlendArguments;
+
+			const UAnimNextAnimationGraph* AnimationGraph = IGraphFactory::GetGraphFromObjectWithFallback(Context, Binding, AnimationAsset);
+			if(AnimationGraph != nullptr)
 			{
-				if (const FPoseSearchDatabaseAnimationAssetBase* DatabaseAnimationAssetBase = MotionMatchingState.CurrentSearchResult.Database->GetDatabaseAnimationAsset<FPoseSearchDatabaseAnimationAssetBase>(*SearchIndexAsset))
+				// Create a default 'native data interface' payload via which we expect to communicate with the graph 
+				FInstancedStruct NewPayload;
+				// @TODO: Support asset->playload mapping in a more data-driven way
+				if (AnimationAsset->IsA<UBlendSpace>())
 				{
-					const UObject* AnimationAsset = DatabaseAnimationAssetBase->GetAnimationAsset();
-					
-					UE::AnimNext::IBlendStack::FGraphRequest NewGraphRequest;
-					NewGraphRequest.BlendArgs = SharedData->BlendArguments;
+					NewPayload.InitializeAs<FAnimNextNativeDataInterface_BlendSpacePlayer>();
+					FAnimNextNativeDataInterface_BlendSpacePlayer& NewNativeInterface = NewPayload.GetMutable<FAnimNextNativeDataInterface_BlendSpacePlayer>();
+					NewNativeInterface.BlendSpace = Cast<UBlendSpace>(AnimationAsset);
 
-					const UAnimNextAnimationGraph* AnimationGraph = IGraphFactory::GetGraphFromObjectWithFallback(Context, Binding, AnimationAsset);
-					if(AnimationGraph != nullptr)
-					{
-						// Create a default 'native data interface' payload via which we expect to communicate with the graph 
-						FInstancedStruct NewPayload;
-						// @TODO: Support asset->playload mapping in a more data-driven way
-						if (AnimationAsset->IsA<UBlendSpace>())
-						{
-							NewPayload.InitializeAs<FAnimNextNativeDataInterface_BlendSpacePlayer>();
-							FAnimNextNativeDataInterface_BlendSpacePlayer& NewNativeInterface = NewPayload.GetMutable<FAnimNextNativeDataInterface_BlendSpacePlayer>();
-							NewNativeInterface.BlendSpace = Cast<UBlendSpace>(AnimationAsset);
-
-							const FVector BlendSpaceParameters = SearchIndexAsset->GetBlendParameters();
-							NewNativeInterface.XAxisSamplePoint = BlendSpaceParameters.X;
-							NewNativeInterface.YAxisSamplePoint = BlendSpaceParameters.Y;
+					
+					NewNativeInterface.XAxisSamplePoint = BlendSpaceParameters.X;
+					NewNativeInterface.YAxisSamplePoint = BlendSpaceParameters.Y;
 
-							NewNativeInterface.PlayRate = DesiredPlayRate;
-							NewNativeInterface.StartPosition = MotionMatchingState.CurrentSearchResult.AssetTime;
-							NewNativeInterface.bLoop = SearchIndexAsset->IsLooping();
-						}
-						else
-						{
-							NewPayload.InitializeAs<FAnimNextNativeDataInterface_AnimSequencePlayer>();
-							FAnimNextNativeDataInterface_AnimSequencePlayer& NewNativeInterface = NewPayload.GetMutable<FAnimNextNativeDataInterface_AnimSequencePlayer>();
-							NewNativeInterface.AnimSequence = Cast<UAnimSequence>(AnimationAsset);
-							NewNativeInterface.PlayRate = DesiredPlayRate;
-							NewNativeInterface.StartPosition = MotionMatchingState.CurrentSearchResult.AssetTime;
-							NewNativeInterface.Loop = SearchIndexAsset->IsLooping();
-						}
+					NewNativeInterface.PlayRate = DesiredPlayRate;
+					NewNativeInterface.StartPosition = MotionMatchingState.CurrentSearchResult.AssetTime;
+					NewNativeInterface.bLoop = bLoop;
+				}
+				else
+				{
+					NewPayload.InitializeAs<FAnimNextNativeDataInterface_AnimSequencePlayer>();
+					FAnimNextNativeDataInterface_AnimSequencePlayer& NewNativeInterface = NewPayload.GetMutable<FAnimNextNativeDataInterface_AnimSequencePlayer>();
+					NewNativeInterface.AnimSequence = Cast<UAnimSequence>(AnimationAsset);
+					NewNativeInterface.PlayRate = DesiredPlayRate;
+					NewNativeInterface.StartPosition = MotionMatchingState.CurrentSearchResult.AssetTime;
+					NewNativeInterface.Loop = bLoop;
+				}
 
-						NewGraphRequest.GraphPayload.AddNative(MoveTemp(NewPayload));
+				NewGraphRequest.GraphPayload.AddNative(MoveTemp(NewPayload));
 
-						IPayloadFactory::GetPayloadForGraphWithFallback(Context, Binding, AnimationAsset, AnimationGraph, NewGraphRequest.GraphPayload);
+				IPayloadFactory::GetPayloadForGraphWithFallback(Context, Binding, AnimationAsset, AnimationGraph, NewGraphRequest.GraphPayload);
 
-						NewGraphRequest.AnimationGraph = AnimationGraph;
+				NewGraphRequest.AnimationGraph = AnimationGraph;
 
-						BlendStackBinding.PushGraph(Context, MoveTemp(NewGraphRequest));
-					}
-				}
+				BlendStackBinding.PushGraph(Context, MoveTemp(NewGraphRequest));
 			}
 		}
 
@@ -272,7 +331,6 @@ namespace UE::AnimNext
 						BlendSpacePayload->XAxisSamplePoint = SharedData->GetXAxisSamplePoint(Binding);
 						BlendSpacePayload->YAxisSamplePoint = SharedData->GetYAxisSamplePoint(Binding);
 					}
-
 					break;
 				}
 			}
diff --git a/Engine/Plugins/Experimental/AnimNextPoseSearch/Source/AnimNextPoseSearch/Private/MotionMatchingTrait.h b/Engine/Plugins/Experimental/AnimNextPoseSearch/Source/AnimNextPoseSearch/Private/MotionMatchingTrait.h
--- a/Engine/Plugins/Experimental/AnimNextPoseSearch/Source/AnimNextPoseSearch/Private/MotionMatchingTrait.h
+++ b/Engine/Plugins/Experimental/AnimNextPoseSearch/Source/AnimNextPoseSearch/Private/MotionMatchingTrait.h
@@ -3,6 +3,7 @@
 #pragma once
 
 #include "CoreMinimal.h"
+#include "PoseSearch/PoseSearchInteractionLibrary.h"
 #include "PoseSearch/PoseSearchLibrary.h"
 #include "TraitCore/Trait.h"
 #include "TraitCore/TraitSharedData.h"
@@ -14,10 +15,6 @@ struct FMotionMatchingTraitSharedData : public FAnimNextTraitSharedData
 {
 	GENERATED_BODY()
 
-	// Entry point in the Subgraph that we will use
-	UPROPERTY(EditAnywhere, Category = "Default", meta=(CustomWidget = "ParamName", AllowedParamType = "FAnimNextEntryPoint"))
-	FName EntryPoint;
-
 	UPROPERTY(EditAnywhere, Category = "Default")
 	TArray<TObjectPtr<const UPoseSearchDatabase>> Databases;
 
@@ -67,9 +64,14 @@ struct FMotionMatchingTraitSharedData : public FAnimNextTraitSharedData
 	// Applied EPoseSearchInterruptMode on the next update that controls the continuing pose search eveluation. This is set back to EPoseSearchInterruptMode::DoNotInterrupt after each update.
 	EPoseSearchInterruptMode NextUpdateInterruptMode = EPoseSearchInterruptMode::DoNotInterrupt;
 
+	UPROPERTY(EditAnywhere, Category = Interaction, meta = (PinHiddenByDefault))
+	TArray<FPoseSearchInteractionAvailability> Availabilities;
+
+	UPROPERTY(EditAnywhere, Category = Interaction, meta = (PinHiddenByDefault))
+	bool bValidateResultAgainstAvailabilities = true;
+
 	// Latent pin support boilerplate
 	#define TRAIT_LATENT_PROPERTIES_ENUMERATOR(GeneratorMacro) \
-		GeneratorMacro(EntryPoint) \
 		GeneratorMacro(Databases) \
 		GeneratorMacro(BlendArguments) \
 		GeneratorMacro(PoseJumpThresholdTime) \
@@ -82,6 +84,8 @@ struct FMotionMatchingTraitSharedData : public FAnimNextTraitSharedData
 		GeneratorMacro(bShouldUseCachedChannelData) \
 		GeneratorMacro(bUpdateBlendSpaceInputs) \
 		GeneratorMacro(NextUpdateInterruptMode) \
+		GeneratorMacro(Availabilities) \
+		GeneratorMacro(bValidateResultAgainstAvailabilities) \
 
 	GENERATE_TRAIT_LATENT_PROPERTIES(FMotionMatchingTraitSharedData, TRAIT_LATENT_PROPERTIES_ENUMERATOR)
 	#undef TRAIT_LATENT_PROPERTIES_ENUMERATOR
@@ -99,6 +103,9 @@ struct FMotionMatchingTrait : FAdditiveTrait, IUpdate
 	struct FInstanceData : FTrait::FInstanceData
 	{
 		FMotionMatchingState MotionMatchingState;
+
+		// @todo: merge this with MotionMatchingState
+		FPoseSearchInteractionBlueprintResult MotionMatchingInteractionState;
 	};
 
 	// IUpdate impl
