PoseSearch - interaction asset mirroring support WIP

diff --git a/Engine/Plugins/Animation/PoseSearch/Source/Editor/Private/PoseSearchDatabaseViewModel.cpp b/Engine/Plugins/Animation/PoseSearch/Source/Editor/Private/PoseSearchDatabaseViewModel.cpp
--- a/Engine/Plugins/Animation/PoseSearch/Source/Editor/Private/PoseSearchDatabaseViewModel.cpp
+++ b/Engine/Plugins/Animation/PoseSearch/Source/Editor/Private/PoseSearchDatabaseViewModel.cpp
@@ -50,8 +50,7 @@ bool FDatabasePreviewActor::SpawnPreviewActor(UWorld* World, const UPoseSearchDa
 	const FSearchIndexAsset& IndexAsset = SearchIndex.Assets[IndexAssetIdx];
 	const FPoseSearchDatabaseAnimationAssetBase* DatabaseAnimationAsset = PoseSearchDatabase->GetDatabaseAnimationAsset<FPoseSearchDatabaseAnimationAssetBase>(IndexAsset.GetSourceAssetIdx());
 	check(DatabaseAnimationAsset);
-
-	
+		
 	USkeleton* Skeleton = PoseSearchDatabase->Schema->GetSkeleton(Role);
 	if (!Skeleton)
 	{
@@ -900,7 +899,7 @@ int32 FDatabaseViewModel::SetSelectedNode(int32 PoseIdx, bool bClearSelection, b
 
 							MaxPreviewPlayLength = FMath::Max(MaxPreviewPlayLength, IndexAsset.GetLastSampleTime(Database->Schema->SampleRate) - PreviewActor.GetPlayTimeOffset());
 							MinPreviewPlayLength = FMath::Min(MinPreviewPlayLength, IndexAsset.GetFirstSampleTime(Database->Schema->SampleRate) - PreviewActor.GetPlayTimeOffset());
-							PreviewActors[PreviewActorGroupIndex].Add(PreviewActor);
+							PreviewActors[PreviewActorGroupIndex].Add(MoveTemp(PreviewActor));
 							SelectedSourceAssetIdx = IndexAsset.GetSourceAssetIdx();
 						}
 					}
@@ -962,7 +961,7 @@ void FDatabaseViewModel::SetSelectedNodes(const TArrayView<TSharedPtr<FDatabaseA
 							}
 
 							MaxPreviewPlayLength = FMath::Max(MaxPreviewPlayLength, IndexAsset.GetLastSampleTime(Database->Schema->SampleRate) - IndexAsset.GetFirstSampleTime(Database->Schema->SampleRate));
-							PreviewActors[PreviewActorGroupIndex].Add(PreviewActor);
+							PreviewActors[PreviewActorGroupIndex].Add(MoveTemp(PreviewActor));
 						}
 					}
 				}
diff --git a/Engine/Plugins/Animation/PoseSearch/Source/Editor/Private/PoseSearchInteractionAssetEditor.cpp b/Engine/Plugins/Animation/PoseSearch/Source/Editor/Private/PoseSearchInteractionAssetEditor.cpp
--- a/Engine/Plugins/Animation/PoseSearch/Source/Editor/Private/PoseSearchInteractionAssetEditor.cpp
+++ b/Engine/Plugins/Animation/PoseSearch/Source/Editor/Private/PoseSearchInteractionAssetEditor.cpp
@@ -312,7 +312,7 @@ void FInteractionAssetViewModel::Tick(float DeltaSeconds)
 				FInteractionAssetPreviewActor PreviewActor;
 				if (PreviewActor.SpawnPreviewActor(GetWorld(), InteractionAsset, InteractionAsset->GetRole(RoleIndex)))
 				{
-					PreviewActors.Add(PreviewActor);
+					PreviewActors.Add(MoveTemp(PreviewActor));
 				}
 			}
 		}
diff --git a/Engine/Plugins/Animation/PoseSearch/Source/Runtime/Private/AnimNode_MotionMatchingInteraction.cpp b/Engine/Plugins/Animation/PoseSearch/Source/Runtime/Private/AnimNode_MotionMatchingInteraction.cpp
--- a/Engine/Plugins/Animation/PoseSearch/Source/Runtime/Private/AnimNode_MotionMatchingInteraction.cpp
+++ b/Engine/Plugins/Animation/PoseSearch/Source/Runtime/Private/AnimNode_MotionMatchingInteraction.cpp
@@ -5,14 +5,15 @@
 #include "Animation/AnimInstanceProxy.h"
 #include "Animation/AnimRootMotionProvider.h"
 #include "HAL/IConsoleManager.h"
+#include "PoseSearch/PoseHistoryProvider.h"
 #include "PoseSearch/PoseSearchAnimNotifies.h"
+#include "PoseSearch/PoseSearchAssetSamplerLibrary.h"
+#include "PoseSearch/PoseSearchDatabase.h"
 #include "PoseSearch/PoseSearchInteractionAsset.h"
 #include "PoseSearch/PoseSearchInteractionIsland.h"
 #include "PoseSearch/PoseSearchInteractionLibrary.h"
-#include "PoseSearch/PoseHistoryProvider.h"
-#include "PoseSearch/PoseSearchDatabase.h"
 #include "PoseSearch/PoseSearchLibrary.h"
-#include "PoseSearch/PoseSearchAssetSamplerLibrary.h"
+#include "PoseSearch/PoseSearchMirrorDataCache.h"
 #include "PoseSearch/PoseSearchSchema.h"
 #include "VisualLogger/VisualLogger.h"
 
@@ -290,11 +291,24 @@ void FAnimNode_MotionMatchingInteraction::Evaluate_AnyThread(FPoseContext& Outpu
 	check(CurrentResult.SelectedDatabase != nullptr && CurrentResult.SelectedDatabase->Schema != nullptr);
 	const FRelevantRoleIndexes RelevantRoleIndexes = GetRelevantRoleIndexes(MultiAnimAsset, CurrentResult.SelectedDatabase->Schema);
 
+	TArray<const UMirrorDataTable*, TInlineAllocator<PreallocatedRolesNum, TMemStackAllocator<>>> MirrorDataTables;
+	if (CurrentResult.bIsMirrored)
+	{
+		MirrorDataTables.SetNum(NumRoles);
+		for (int32 MultiAnimAssetRoleIndex = 0; MultiAnimAssetRoleIndex < NumRoles; ++MultiAnimAssetRoleIndex)
+		{
+			const FRole Role = MultiAnimAsset->GetRole(MultiAnimAssetRoleIndex);
+			const FPoseSearchRoledSkeleton* RoledSkeleton = CurrentResult.SelectedDatabase->Schema->GetRoledSkeleton(Role);
+			check(RoledSkeleton);
+			MirrorDataTables[MultiAnimAssetRoleIndex] = RoledSkeleton->MirrorDataTable;
+		}
+	}
+
 	// calculating the world space warp tranform to have the character fully aligned to what MultiAnimAsset wants.
 	// using ActorTransforms instead of ActorRootTransforms, since we want the previous frame root bone offset rotation to be embedded into the calculation
 	TArray<FTransform, TInlineAllocator<PreallocatedRolesNum, TMemStackAllocator<>>> FullAlignedActorRootBoneTransformsAtExtractionTime;
 	FullAlignedActorRootBoneTransformsAtExtractionTime.SetNum(NumRoles);
-	MultiAnimAsset->CalculateWarpTransforms(WarpingAlignmentExtractionTime, ActorTransforms, FullAlignedActorRootBoneTransformsAtExtractionTime, RelevantRoleIndexes);
+	MultiAnimAsset->CalculateWarpTransforms(WarpingAlignmentExtractionTime, ActorTransforms, FullAlignedActorRootBoneTransformsAtExtractionTime, MirrorDataTables, RelevantRoleIndexes);
 	const FTransform AlignmentTransformFullWarpAdjustment = FullAlignedActorRootBoneTransformsAtExtractionTime[CurrentResultRoleIndex];
 					
 	// @todo: compensating AlignmentTransformFullWarpAdjustment by the delta time between AccumulatedTime and WarpingAlignmentExtractionTime?				
@@ -348,7 +362,7 @@ void FAnimNode_MotionMatchingInteraction::Evaluate_AnyThread(FPoseContext& Outpu
 		// using ActorTransforms instead of ActorRootTransforms, since we want the previous frame root bone offset rotation to be embedded into the calculation
 		TArray<FTransform, TInlineAllocator<PreallocatedRolesNum, TMemStackAllocator<>>> FullAlignedActorRootBoneTransformsAtAccumulatedTime;
 		FullAlignedActorRootBoneTransformsAtAccumulatedTime.SetNum(NumRoles);
-		MultiAnimAsset->CalculateWarpTransforms(AccumulatedTime, ActorTransforms, FullAlignedActorRootBoneTransformsAtAccumulatedTime, RelevantRoleIndexes);
+		MultiAnimAsset->CalculateWarpTransforms(AccumulatedTime, ActorTransforms, FullAlignedActorRootBoneTransformsAtAccumulatedTime, MirrorDataTables, RelevantRoleIndexes);
 		FullAlignedPoseRootTransform = FullAlignedActorRootBoneTransformsAtAccumulatedTime[CurrentResultRoleIndex];
 	}
 								
diff --git a/Engine/Plugins/Animation/PoseSearch/Source/Runtime/Private/PoseSearchDerivedData.cpp b/Engine/Plugins/Animation/PoseSearch/Source/Runtime/Private/PoseSearchDerivedData.cpp
--- a/Engine/Plugins/Animation/PoseSearch/Source/Runtime/Private/PoseSearchDerivedData.cpp
+++ b/Engine/Plugins/Animation/PoseSearch/Source/Runtime/Private/PoseSearchDerivedData.cpp
@@ -1783,7 +1783,25 @@ void FPoseSearchDatabaseAsyncCacheTask::OnGetComplete(UE::DerivedData::FCacheGet
 					}
 
 					// validating that the missing MirrorDataTable(s) are not necessary 
-					if (!DependentDatabaseSchema->AllRoledSkeletonHaveMirrorDataTable())
+					bool bAllRoledSkeletonHaveMirrorDataTable = true;
+					for (const FPoseSearchRoledSkeleton& RoledSkeleton : DependentDatabaseSchema->GetRoledSkeletons())
+					{
+						if (RoledSkeleton.MirrorDataTable)
+						{
+							if (!RoledSkeleton.MirrorDataTable->Skeleton)
+							{
+								UE_LOG(LogPoseSearch, Error, TEXT("%s - %s BuildIndex Failed because '%s' schema MirrorDataTable Skeleton is not set for Role '%s' "), *LexToString(FullIndexKey.Hash), *MainDatabaseName, *DependentDatabaseName, *RoledSkeleton.Role.ToString());
+								ResetSearchIndex();
+								return;
+							}
+						}
+						else
+						{
+							bAllRoledSkeletonHaveMirrorDataTable = false;
+						}
+					}
+
+					if (!bAllRoledSkeletonHaveMirrorDataTable)
 					{
 						for (int32 AnimationAssetIndex = 0; AnimationAssetIndex < DependentDatabase->GetNumAnimationAssets(); ++AnimationAssetIndex)
 						{
@@ -1792,7 +1810,7 @@ void FPoseSearchDatabaseAsyncCacheTask::OnGetComplete(UE::DerivedData::FCacheGet
 								if (DatabaseAsset->GetMirrorOption() == EPoseSearchMirrorOption::MirroredOnly || DatabaseAsset->GetMirrorOption() == EPoseSearchMirrorOption::UnmirroredAndMirrored)
 								{
 									// want to sample a mirrored asset
-									UE_LOG(LogPoseSearch, Error, TEXT("%s - %s BuildIndex Failed because '%s' requires MirrorDataTable(s) to sample mirrored animation assets"), *LexToString(FullIndexKey.Hash), *MainDatabaseName, *DependentDatabaseName);
+									UE_LOG(LogPoseSearch, Error, TEXT("%s - %s BuildIndex Failed because '%s' schema requires MirrorDataTable(s) to sample mirrored animation assets"), *LexToString(FullIndexKey.Hash), *MainDatabaseName, *DependentDatabaseName);
 									ResetSearchIndex();
 									return;
 								}
diff --git a/Engine/Plugins/Animation/PoseSearch/Source/Runtime/Private/PoseSearchInteractionAsset.cpp b/Engine/Plugins/Animation/PoseSearch/Source/Runtime/Private/PoseSearchInteractionAsset.cpp
--- a/Engine/Plugins/Animation/PoseSearch/Source/Runtime/Private/PoseSearchInteractionAsset.cpp
+++ b/Engine/Plugins/Animation/PoseSearch/Source/Runtime/Private/PoseSearchInteractionAsset.cpp
@@ -4,6 +4,7 @@
 #include "Animation/AnimSequence.h"
 #include "Animation/BlendSpace.h"
 #include "PoseSearch/PoseSearchAssetSampler.h"
+#include "PoseSearch/PoseSearchMirrorDataCache.h"
 
 bool UPoseSearchInteractionAsset::IsLooping() const
 {
@@ -272,7 +273,7 @@ USkeletalMesh* UPoseSearchInteractionAsset::GetPreviewMesh(const UE::PoseSearch:
 
 #endif // WITH_EDITOR
 
-void UPoseSearchInteractionAsset::CalculateWarpTransforms(float Time, const TConstArrayView<const FTransform> ActorRootBoneTransforms, TArrayView<FTransform> FullAlignedActorRootBoneTransforms, const TConstArrayView<bool> RelevantRoleIndexes) const
+void UPoseSearchInteractionAsset::CalculateWarpTransforms(float Time, const TConstArrayView<const FTransform> ActorRootBoneTransforms, TArrayView<FTransform> FullAlignedActorRootBoneTransforms, const TConstArrayView<const UMirrorDataTable*> MirrorDataTables, const TConstArrayView<bool> RelevantRoleIndexes) const
 {
 	using namespace UE::PoseSearch;
 
@@ -322,6 +323,12 @@ void UPoseSearchInteractionAsset::CalculateWarpTransforms(float Time, const TCon
 			const FAnimationAssetSampler Sampler(Item.Animation, Item.Origin);
 			AssetRootBoneTransforms[ItemIndex] = Sampler.ExtractRootTransform(Time);
 
+			if (MirrorDataTables.IsValidIndex(ItemIndex) && MirrorDataTables[ItemIndex])
+			{
+				const FMirrorDataCache MirrorDataCache(MirrorDataTables[ItemIndex]);
+				AssetRootBoneTransforms[ItemIndex] = MirrorDataCache.MirrorTransform(AssetRootBoneTransforms[ItemIndex]);
+			}
+
 #if ENABLE_ANIM_DEBUG && WITH_EDITOR
 			if (Items[ItemIndex].Animation)
 			{
diff --git a/Engine/Plugins/Animation/PoseSearch/Source/Runtime/Private/PoseSearchInteractionIsland.cpp b/Engine/Plugins/Animation/PoseSearch/Source/Runtime/Private/PoseSearchInteractionIsland.cpp
--- a/Engine/Plugins/Animation/PoseSearch/Source/Runtime/Private/PoseSearchInteractionIsland.cpp
+++ b/Engine/Plugins/Animation/PoseSearch/Source/Runtime/Private/PoseSearchInteractionIsland.cpp
@@ -860,7 +860,7 @@ bool FInteractionIsland::DoSearch_AnyThread(const UObject* AnimContext, const FP
 			const UObject* AssetsToSearch[] = { Database };
 			FPoseSearchFutureProperties PoseSearchFutureProperties;
 
-			FPoseSearchContinuingProperties ContinuingPropertiesToUse = SearchContext.ContinuingProperties;
+			FPoseSearchContinuingProperties ContinuingPropertiesToUse = SearchContext.GetContinuingProperties();
 
 			// @todo: implement this carefully, since it'll cause thread safety concerns depending on which actors is calling the query first 
 			// and if their animations get integrated by different play rates
diff --git a/Engine/Plugins/Animation/PoseSearch/Source/Runtime/Private/PoseSearchInteractionSubsystem.cpp b/Engine/Plugins/Animation/PoseSearch/Source/Runtime/Private/PoseSearchInteractionSubsystem.cpp
--- a/Engine/Plugins/Animation/PoseSearch/Source/Runtime/Private/PoseSearchInteractionSubsystem.cpp
+++ b/Engine/Plugins/Animation/PoseSearch/Source/Runtime/Private/PoseSearchInteractionSubsystem.cpp
@@ -386,8 +386,9 @@ void UPoseSearchInteractionSubsystem::PopulateContinuingProperties(float DeltaSe
 					{
 						if (const FPoseSearchDatabaseAnimationAssetBase* DatabaseAsset = SearchResult->Database->GetDatabaseAnimationAsset<FPoseSearchDatabaseAnimationAssetBase>(*SearchIndexAsset))
 						{
-							SearchContext.ContinuingProperties.PlayingAsset = DatabaseAsset->GetAnimationAsset();
-							SearchContext.ContinuingProperties.PlayingAssetAccumulatedTime = SearchResult->AssetTime + DeltaSeconds;
+							SearchContext.PlayingAsset = DatabaseAsset->GetAnimationAsset();
+							SearchContext.PlayingAssetAccumulatedTime = SearchResult->AssetTime + DeltaSeconds;
+							// @todo: populate SearchContext.InterruptMode
 						}
 					}
 				}
diff --git a/Engine/Plugins/Animation/PoseSearch/Source/Runtime/Private/PoseSearchMirrorDataCache.cpp b/Engine/Plugins/Animation/PoseSearch/Source/Runtime/Private/PoseSearchMirrorDataCache.cpp
--- a/Engine/Plugins/Animation/PoseSearch/Source/Runtime/Private/PoseSearchMirrorDataCache.cpp
+++ b/Engine/Plugins/Animation/PoseSearch/Source/Runtime/Private/PoseSearchMirrorDataCache.cpp
@@ -11,18 +11,61 @@ FMirrorDataCache::FMirrorDataCache()
 {
 }
 
+FMirrorDataCache::FMirrorDataCache(const UMirrorDataTable* InMirrorDataTable)
+{
+	Init(InMirrorDataTable);
+}
+
 FMirrorDataCache::FMirrorDataCache(const UMirrorDataTable* InMirrorDataTable, const FBoneContainer& BoneContainer)
 {
 	Init(InMirrorDataTable, BoneContainer);
 }
 
+void FMirrorDataCache::Init(const UMirrorDataTable* InMirrorDataTable)
+{
+	if (InMirrorDataTable != nullptr)
+	{
+		if (InMirrorDataTable->Skeleton)
+		{
+			MirrorDataTable = InMirrorDataTable;
+
+			// array containing the bone index of the root bone (0)
+			TArray<uint16, TInlineAllocator<1>> BoneIndices;
+			BoneIndices.SetNumZeroed(1);
+
+			// extracting the pose, containing only the root bone from the Sampler 
+			FBoneContainer BoneContainer;
+			BoneContainer.InitializeTo(BoneIndices, UE::Anim::FCurveFilterSettings(UE::Anim::ECurveFilterMode::DisallowAll), *InMirrorDataTable->Skeleton);
+
+			MirrorDataTable->FillCompactPoseAndComponentRefRotations(BoneContainer, CompactPoseMirrorBones, ComponentSpaceRefRotations);
+		}
+		else
+		{
+			UE_LOG(LogPoseSearch, Error, TEXT("FMirrorDataCache::Init - UMirrorDataTable '%s' Skeleton is not set!"), *InMirrorDataTable->GetName());
+			Reset();
+		}
+	}
+	else
+	{
+		Reset();
+	}
+}
+
 void FMirrorDataCache::Init(const UMirrorDataTable* InMirrorDataTable, const FBoneContainer& BoneContainer)
 {
 	if (InMirrorDataTable != nullptr)
 	{
-		check(BoneContainer.IsValid());
-		MirrorDataTable = InMirrorDataTable;
-		MirrorDataTable->FillCompactPoseAndComponentRefRotations(BoneContainer, CompactPoseMirrorBones, ComponentSpaceRefRotations);
+		if (InMirrorDataTable->Skeleton)
+		{
+			check(BoneContainer.IsValid());
+			MirrorDataTable = InMirrorDataTable;
+			MirrorDataTable->FillCompactPoseAndComponentRefRotations(BoneContainer, CompactPoseMirrorBones, ComponentSpaceRefRotations);
+		}
+		else
+		{
+			UE_LOG(LogPoseSearch, Error, TEXT("FMirrorDataCache::Init - UMirrorDataTable '%s' Skeleton is not set!"), *InMirrorDataTable->GetName());
+			Reset();
+		}
 	}
 	else
 	{
diff --git a/Engine/Plugins/Animation/PoseSearch/Source/Runtime/Public/PoseSearch/MultiAnimAsset.h b/Engine/Plugins/Animation/PoseSearch/Source/Runtime/Public/PoseSearch/MultiAnimAsset.h
--- a/Engine/Plugins/Animation/PoseSearch/Source/Runtime/Public/PoseSearch/MultiAnimAsset.h
+++ b/Engine/Plugins/Animation/PoseSearch/Source/Runtime/Public/PoseSearch/MultiAnimAsset.h
@@ -6,6 +6,7 @@
 
 #include "MultiAnimAsset.generated.h"
 
+class UMirrorDataTable;
 class UAnimationAsset;
 class USkeletalMesh;
 
@@ -30,7 +31,7 @@ public:
 	[[nodiscard]] virtual UAnimationAsset* GetAnimationAsset(const FName& Role) const PURE_VIRTUAL(UMultiAnimAsset::GetAnimationAsset, return nullptr;);
 	[[nodiscard]] virtual FTransform GetOrigin(const FName& Role) const PURE_VIRTUAL(UMultiAnimAsset::GetOrigin, return FTransform::Identity;);
 
-	virtual void CalculateWarpTransforms(float Time, const TConstArrayView<const FTransform> ActorRootBoneTransforms, TArrayView<FTransform> FullAlignedActorRootBoneTransforms, const TConstArrayView<bool> RelevantRoleIndexes = TArrayView<bool>()) const PURE_VIRTUAL(UMultiAnimAsset::CalculateWarpTransforms, );
+	virtual void CalculateWarpTransforms(float Time, const TConstArrayView<const FTransform> ActorRootBoneTransforms, TArrayView<FTransform> FullAlignedActorRootBoneTransforms, const TConstArrayView<const UMirrorDataTable*> MirrorDataTables = TConstArrayView<const UMirrorDataTable*>(), const TConstArrayView<bool> RelevantRoleIndexes = TArrayView<bool>()) const PURE_VIRTUAL(UMultiAnimAsset::CalculateWarpTransforms, );
 
 	UFUNCTION(BlueprintPure, Category = "Animation", meta=(BlueprintThreadSafe, DisplayName = "Get Animation Asset"))
 	UAnimationAsset* BP_GetAnimationAsset(const FName& Role) const { return GetAnimationAsset(Role); }
diff --git a/Engine/Plugins/Animation/PoseSearch/Source/Runtime/Public/PoseSearch/PoseSearchInteractionAsset.h b/Engine/Plugins/Animation/PoseSearch/Source/Runtime/Public/PoseSearch/PoseSearchInteractionAsset.h
--- a/Engine/Plugins/Animation/PoseSearch/Source/Runtime/Public/PoseSearch/PoseSearchInteractionAsset.h
+++ b/Engine/Plugins/Animation/PoseSearch/Source/Runtime/Public/PoseSearch/PoseSearchInteractionAsset.h
@@ -92,7 +92,7 @@ public:
 	virtual USkeletalMesh* GetPreviewMesh(const UE::PoseSearch::FRole& Role) const override;
 #endif // WITH_EDITOR
 
-	virtual void CalculateWarpTransforms(float Time, const TConstArrayView<const FTransform> ActorRootBoneTransforms, TArrayView<FTransform> FullAlignedActorRootBoneTransforms, const TConstArrayView<bool> RelevantRoleIndexes = TConstArrayView<bool>()) const override;
+	virtual void CalculateWarpTransforms(float Time, const TConstArrayView<const FTransform> ActorRootBoneTransforms, TArrayView<FTransform> FullAlignedActorRootBoneTransforms, const TConstArrayView<const UMirrorDataTable*> MirrorDataTables = TConstArrayView<const UMirrorDataTable*>(), const TConstArrayView<bool> RelevantRoleIndexes = TConstArrayView<bool>()) const override;
 	
 	FVector FindReferencePosition(const TArrayView<const FTransform> Transforms, const TArrayView<float> NormalizedWarpingWeightTranslation) const;
 	FQuat FindReferenceOrientation(const TArrayView<const FTransform> Transforms, const TArrayView<int32> SortedByWarpingWeightRotationItemIndex, const TArrayView<float> NormalizedWarpingWeightRotation) const;
diff --git a/Engine/Plugins/Animation/PoseSearch/Source/Runtime/Public/PoseSearch/PoseSearchInteractionIsland.h b/Engine/Plugins/Animation/PoseSearch/Source/Runtime/Public/PoseSearch/PoseSearchInteractionIsland.h
--- a/Engine/Plugins/Animation/PoseSearch/Source/Runtime/Public/PoseSearch/PoseSearchInteractionIsland.h
+++ b/Engine/Plugins/Animation/PoseSearch/Source/Runtime/Public/PoseSearch/PoseSearchInteractionIsland.h
@@ -30,10 +30,7 @@ FVector GetContextLocation(const UObject* AnimContext);
 // Experimental, this feature might be removed without warning, not for production use
 struct FInteractionSearchContextBase
 {
-	// @todo: since AnimContexts, HistoryCollectors, Roles, and BroadPhaseRadiuses have the same cardinality,
-	// add an array of structs instead of many arrays
 	TArray<TWeakObjectPtr<const UObject>, TInlineAllocator<PreallocatedRolesNum>> AnimContexts;
-	// @todo: perhaps use the animnode of the pose history collector instead
 	TArray<const UE::PoseSearch::IPoseHistory*, TInlineAllocator<PreallocatedRolesNum>> PoseHistories;
 	TArray<UE::PoseSearch::FRole, TInlineAllocator<PreallocatedRolesNum>> Roles;
 
@@ -57,7 +54,19 @@ struct FInteractionSearchContextBase
 // Experimental, this feature might be removed without warning, not for production use
 struct FInteractionSearchContext : FInteractionSearchContextBase
 {
-	FPoseSearchContinuingProperties ContinuingProperties;
+	// TWeakObjectPtr version of FPoseSearchContinuingProperties to be GC friendly
+	TWeakObjectPtr<const UObject> PlayingAsset = nullptr;
+	float PlayingAssetAccumulatedTime = 0.f;
+	EPoseSearchInterruptMode InterruptMode = EPoseSearchInterruptMode::DoNotInterrupt;
+
+	FPoseSearchContinuingProperties GetContinuingProperties() const
+	{
+		FPoseSearchContinuingProperties ContinuingProperties;
+		ContinuingProperties.PlayingAsset = PlayingAsset.Get();
+		ContinuingProperties.PlayingAssetAccumulatedTime = PlayingAssetAccumulatedTime;
+		ContinuingProperties.InterruptMode = InterruptMode;
+		return ContinuingProperties;
+	}
 };
 // Experimental, this feature might be removed without warning, not for production use
 typedef TArray<FInteractionSearchContext, TInlineAllocator<PreallocatedSearchesNum, TMemStackAllocator<>>> FInteractionSearchContexts;
diff --git a/Engine/Plugins/Animation/PoseSearch/Source/Runtime/Public/PoseSearch/PoseSearchMirrorDataCache.h b/Engine/Plugins/Animation/PoseSearch/Source/Runtime/Public/PoseSearch/PoseSearchMirrorDataCache.h
--- a/Engine/Plugins/Animation/PoseSearch/Source/Runtime/Public/PoseSearch/PoseSearchMirrorDataCache.h
+++ b/Engine/Plugins/Animation/PoseSearch/Source/Runtime/Public/PoseSearch/PoseSearchMirrorDataCache.h
@@ -15,8 +15,12 @@ namespace UE::PoseSearch
 struct POSESEARCH_API FMirrorDataCache
 {
 	FMirrorDataCache();
+	// fast initialization to mirror only the root bone
+	FMirrorDataCache(const UMirrorDataTable* InMirrorDataTable);
 	FMirrorDataCache(const UMirrorDataTable* InMirrorDataTable, const FBoneContainer& BoneContainer);
 	
+	// fast initialization to mirror only the root bone
+	void Init(const UMirrorDataTable* InMirrorDataTable);
 	void Init(const UMirrorDataTable* InMirrorDataTable, const FBoneContainer& BoneContainer);
 	void Reset();
 
diff --git a/Engine/Plugins/Animation/PoseSearch/Source/Runtime/Public/PoseSearch/PoseSearchSchema.h b/Engine/Plugins/Animation/PoseSearch/Source/Runtime/Public/PoseSearch/PoseSearchSchema.h
--- a/Engine/Plugins/Animation/PoseSearch/Source/Runtime/Public/PoseSearch/PoseSearchSchema.h
+++ b/Engine/Plugins/Animation/PoseSearch/Source/Runtime/Public/PoseSearch/PoseSearchSchema.h
@@ -177,6 +177,7 @@ public:
 	// to avoid confusing the Role to index mapping in the Schema versus the Role to index mapping in other assets
 	// please use GetRoledSkeleton(Role)
 	void InitBoneContainersFromRoledSkeleton(TMap<FName, FBoneContainer>& RoledBoneContainers) const;
+	UE_DEPRECATED(5.6, "access the Skeletons via GetRoledSkeletons() and query for their MirrorDataTable instead")
 	bool AllRoledSkeletonHaveMirrorDataTable() const;
 	const FPoseSearchRoledSkeleton* GetRoledSkeleton(const UE::PoseSearch::FRole& Role) const;
 	FPoseSearchRoledSkeleton* GetRoledSkeleton(const UE::PoseSearch::FRole& Role);
