On behalf of Samuele.Rigomonti

PoseSearch - fix for FAnimNode_MotionMatching::UpdateAssetPlayer not calling MotionMatchingState.Reset when Context.GetFinalBlendWeight was not relevant enought. This was causing stale continuing pose from MotionMatchingState being reselected when the node weight was getting relevant once again

diff --git a/Engine/Plugins/Animation/PoseSearch/Source/Runtime/Private/AnimNode_MotionMatching.cpp b/Engine/Plugins/Animation/PoseSearch/Source/Runtime/Private/AnimNode_MotionMatching.cpp
--- a/Engine/Plugins/Animation/PoseSearch/Source/Runtime/Private/AnimNode_MotionMatching.cpp
+++ b/Engine/Plugins/Animation/PoseSearch/Source/Runtime/Private/AnimNode_MotionMatching.cpp
@@ -67,9 +67,7 @@ void FAnimNode_MotionMatching::Initialize_AnyThread(const FAnimationInitializeCo
 	GetEvaluateGraphExposedInputs().Execute(Context);
 
 	FAnimNode_BlendStack_Standalone::Initialize_AnyThread(Context);
-	PRAGMA_DISABLE_DEPRECATION_WARNINGS
-	MotionMatchingState.Reset(Context.AnimInstanceProxy->GetComponentTransform());
-	PRAGMA_ENABLE_DEPRECATION_WARNINGS
+	MotionMatchingState.Reset();
 }
 
 void FAnimNode_MotionMatching::Evaluate_AnyThread(FPoseContext& Output)
@@ -126,9 +124,7 @@ void FAnimNode_MotionMatching::UpdateAssetPlayer(const FAnimationUpdateContext&
 	// If we just became relevant and haven't been initialized yet, then reset motion matching state, otherwise update the asset time using the player node.
 	if (bNeedsReset)
 	{
-		PRAGMA_DISABLE_DEPRECATION_WARNINGS
-		MotionMatchingState.Reset(Context.AnimInstanceProxy->GetComponentTransform());
-		PRAGMA_ENABLE_DEPRECATION_WARNINGS
+		MotionMatchingState.Reset();
 		FAnimNode_BlendStack_Standalone::Reset();
 	}
 	else
@@ -214,6 +210,10 @@ void FAnimNode_MotionMatching::UpdateAssetPlayer(const FAnimationUpdateContext&
 #endif // ENABLE_ANIM_DEBUG
 		);
 	}
+	else
+	{
+		MotionMatchingState.Reset();
+	}
 
 	const bool bJumpToPose = bIsRelevantEnough && MotionMatchingState.CurrentSearchResult.IsValid() && !MotionMatchingState.CurrentSearchResult.bIsContinuingPoseSearch;
 
diff --git a/Engine/Plugins/Animation/PoseSearch/Source/Runtime/Private/PoseSearchLibrary.cpp b/Engine/Plugins/Animation/PoseSearch/Source/Runtime/Private/PoseSearchLibrary.cpp
--- a/Engine/Plugins/Animation/PoseSearch/Source/Runtime/Private/PoseSearchLibrary.cpp
+++ b/Engine/Plugins/Animation/PoseSearch/Source/Runtime/Private/PoseSearchLibrary.cpp
@@ -301,7 +301,6 @@ namespace UE::PoseSearch
 void FMotionMatchingState::Reset(const FTransform& ComponentTransform)
 {
 PRAGMA_DISABLE_DEPRECATION_WARNINGS
-	bJumpedToPose = false;
 PRAGMA_ENABLE_DEPRECATION_WARNINGS
 
 	Reset();
@@ -309,6 +308,10 @@ PRAGMA_ENABLE_DEPRECATION_WARNINGS
 
 void FMotionMatchingState::Reset()
 {
+PRAGMA_DISABLE_DEPRECATION_WARNINGS
+	bJumpedToPose = false;
+PRAGMA_ENABLE_DEPRECATION_WARNINGS
+
 	CurrentSearchResult.Reset();
 	// Set the elapsed time to INFINITY to trigger a search right away
 	ElapsedPoseSearchTime = std::numeric_limits<float>::infinity();
