PoseSearch - fix for UPoseSearchLibrary::MotionMatch continuing pose reconstruction logic

diff --git a/Engine/Plugins/Animation/PoseSearch/Source/Runtime/Private/PoseSearchLibrary.cpp b/Engine/Plugins/Animation/PoseSearch/Source/Runtime/Private/PoseSearchLibrary.cpp
--- a/Engine/Plugins/Animation/PoseSearch/Source/Runtime/Private/PoseSearchLibrary.cpp
+++ b/Engine/Plugins/Animation/PoseSearch/Source/Runtime/Private/PoseSearchLibrary.cpp
@@ -300,9 +300,6 @@ namespace UE::PoseSearch
 
 void FMotionMatchingState::Reset(const FTransform& ComponentTransform)
 {
-PRAGMA_DISABLE_DEPRECATION_WARNINGS
-PRAGMA_ENABLE_DEPRECATION_WARNINGS
-
 	Reset();
 }
 
@@ -1117,11 +1114,10 @@ UE::PoseSearch::FSearchResult UPoseSearchLibrary::MotionMatch(
 				ReconstructedPreviousSearchResult.Database = Database;
 				ReconstructedPreviousSearchResult.PoseIdx = Database->GetPoseIndexFromTime(ContinuingProperties.PlayingAssetAccumulatedTime, SearchIndexAsset);
 
-				FSearchResult NewSearchResult;
 				if (bCanAdvance && !bForceInterrupt)
 				{
 					SearchContext.UpdateCurrentResultPoseVector();
-					NewSearchResult = Database->SearchContinuingPose(SearchContext);
+					const FSearchResult NewSearchResult = Database->SearchContinuingPose(SearchContext);
 
 #if WITH_EDITOR && ENABLE_ANIM_DEBUG && UE_POSE_SEARCH_TRACE_ENABLED
 					const FPoseSearchCost BestBruteForcePoseCost = NewSearchResult.BruteForcePoseCost < SearchResult.BruteForcePoseCost ? NewSearchResult.BruteForcePoseCost : SearchResult.BruteForcePoseCost;
@@ -1135,41 +1131,44 @@ UE::PoseSearch::FSearchResult UPoseSearchLibrary::MotionMatch(
 #if WITH_EDITOR && ENABLE_ANIM_DEBUG && UE_POSE_SEARCH_TRACE_ENABLED
 					SearchResult.BruteForcePoseCost = BestBruteForcePoseCost;
 #endif // WITH_EDITOR && ENABLE_ANIM_DEBUG && UE_POSE_SEARCH_TRACE_ENABLED
-				}
 
-				// keeping track of the best ReconstructedPreviousSearchResult, as best continuing pose search result for the full Database search happening after
-				if (FSearchResult* FoundReconstructedPreviousSearchBestResult = ReconstructedPreviousSearchBestResultMap.Find(Database))
-				{
-					if (NewSearchResult.PoseCost == FoundReconstructedPreviousSearchBestResult->PoseCost)
+					if (NewSearchResult.IsValid())
 					{
-						check(NewSearchResult.PoseIdx != FoundReconstructedPreviousSearchBestResult->PoseIdx);
-						if (NewSearchResult.PoseIdx < FoundReconstructedPreviousSearchBestResult->PoseIdx)
+						// keeping track of the best ReconstructedPreviousSearchResult, as best continuing pose search result for the full Database search happening after
+						if (FSearchResult* FoundReconstructedPreviousSearchBestResult = ReconstructedPreviousSearchBestResultMap.Find(Database))
 						{
-							// choosing to update with the lowest PoseIdx to avoid indeterminism 
-							// since ReconstructedPreviousSearchBestResultMap is an unsorted Map
-							*FoundReconstructedPreviousSearchBestResult = ReconstructedPreviousSearchResult;
-							FoundReconstructedPreviousSearchBestResult->PoseCost = NewSearchResult.PoseCost;
+							if (NewSearchResult.PoseCost == FoundReconstructedPreviousSearchBestResult->PoseCost)
+							{
+								check(NewSearchResult.PoseIdx != FoundReconstructedPreviousSearchBestResult->PoseIdx);
+								if (NewSearchResult.PoseIdx < FoundReconstructedPreviousSearchBestResult->PoseIdx)
+								{
+									// choosing to update with the lowest PoseIdx to avoid indeterminism 
+									// since ReconstructedPreviousSearchBestResultMap is an unsorted Map
+									*FoundReconstructedPreviousSearchBestResult = ReconstructedPreviousSearchResult;
+									FoundReconstructedPreviousSearchBestResult->PoseCost = NewSearchResult.PoseCost;
+								}
+							}
+							else if (NewSearchResult.PoseCost < FoundReconstructedPreviousSearchBestResult->PoseCost)
+							{
+								*FoundReconstructedPreviousSearchBestResult = ReconstructedPreviousSearchResult;
+								FoundReconstructedPreviousSearchBestResult->PoseCost = NewSearchResult.PoseCost;
+							}
+						}
+						else
+						{
+							FSearchResult& NewReconstructedPreviousSearchBestResult = ReconstructedPreviousSearchBestResultMap.Add(Database);
+							NewReconstructedPreviousSearchBestResult = ReconstructedPreviousSearchResult;
+							NewReconstructedPreviousSearchBestResult.PoseCost = NewSearchResult.PoseCost;
+
+							if (bForceInterrupt)
+							{
+								// no point to continue if it's bForceInterrupt: no continuing pose searches will be performed 
+								// so all the ReconstructedPreviousSearchResult will lead to the same pose with invalid PoseCost
+								// so any of the pair AssetTime / PoseIdx will be the same
+								break;
+							}
 						}
 					}
-					else if (NewSearchResult.PoseCost < FoundReconstructedPreviousSearchBestResult->PoseCost)
-					{
-						*FoundReconstructedPreviousSearchBestResult = ReconstructedPreviousSearchResult;
-						FoundReconstructedPreviousSearchBestResult->PoseCost = NewSearchResult.PoseCost;
-					}
-				}
-				else
-				{
-					FSearchResult& NewReconstructedPreviousSearchBestResult = ReconstructedPreviousSearchBestResultMap.Add(Database);
-					NewReconstructedPreviousSearchBestResult = ReconstructedPreviousSearchResult;
-					NewReconstructedPreviousSearchBestResult.PoseCost = NewSearchResult.PoseCost;
-				}
-
-				if (bForceInterrupt)
-				{
-					// no point to continue if it's bForceInterrupt: no continuing pose searches will be performed 
-					// so all the ReconstructedPreviousSearchResult will lead to the same pose with invalid PoseCost
-					// so any of the pair AssetTime / PoseIdx will be the same
-					break;
 				}
 			}
 		}
