PoseSearch - fix for FAnimationAssetSampler::ExtractPose invalid check when using AnimConposites

diff --git a/Engine/Plugins/Animation/PoseSearch/Source/Runtime/Private/PoseSearchAssetSampler.cpp b/Engine/Plugins/Animation/PoseSearch/Source/Runtime/Private/PoseSearchAssetSampler.cpp
--- a/Engine/Plugins/Animation/PoseSearch/Source/Runtime/Private/PoseSearchAssetSampler.cpp
+++ b/Engine/Plugins/Animation/PoseSearch/Source/Runtime/Private/PoseSearchAssetSampler.cpp
@@ -254,7 +254,7 @@ void FAnimationAssetSampler::Init(const UAnimationAsset* InAnimationAsset, const
 #if WITH_EDITOR
 	ITargetPlatform* TargetPlatform = GetTargetPlatformManagerRef().GetRunningTargetPlatform();		
 	PlatformHash = FIoHash::Zero;
-	if (UAnimSequence* AnimSequence = Cast<UAnimSequence>(const_cast<UAnimationAsset*>(InAnimationAsset)))
+	if (const UAnimSequence* AnimSequence = Cast<UAnimSequence>(InAnimationAsset))
 	{		
 		PlatformHash = AnimSequence->GetDerivedDataKeyHash(TargetPlatform);
 	}
@@ -420,8 +420,11 @@ void FAnimationAssetSampler::ExtractPose(const FAnimExtractContext& ExtractionCt
 	{
 		FAnimExtractContext SequenceContext = ExtractionCtx;
 #if WITH_EDITOR
-		check(PlatformHash != FIoHash::Zero);
-		SequenceContext.TargetPlatformHash = PlatformHash;
+		if (Cast<UAnimSequence>(SequenceBase))
+		{
+			check(PlatformHash != FIoHash::Zero);
+			SequenceContext.TargetPlatformHash = PlatformHash;
+		}
 #endif // WITH_EDITOR
 		SequenceBase->GetAnimationPose(OutAnimPoseData, SequenceContext);
 	}
