PR #12656: [PoseSearch] Expose EPoseSearchInterruptMode to Blueprint, display weights in editor database details, improved editor only search index data serialization, fixes for running editor with cooked dbs

diff --git a/Engine/Plugins/Animation/PoseSearch/Source/Editor/Private/PoseSearchDatabaseDataDetails.cpp b/Engine/Plugins/Animation/PoseSearch/Source/Editor/Private/PoseSearchDatabaseDataDetails.cpp
--- a/Engine/Plugins/Animation/PoseSearch/Source/Editor/Private/PoseSearchDatabaseDataDetails.cpp
+++ b/Engine/Plugins/Animation/PoseSearch/Source/Editor/Private/PoseSearchDatabaseDataDetails.cpp
@@ -16,11 +16,12 @@ namespace UE::PoseSearch
 class FChannelItem
 {
 public:
-	FChannelItem(const UPoseSearchFeatureChannel* Channel, int32 ChannelComponentIdx = -1)
+	FChannelItem(const UPoseSearchFeatureChannel* Channel, int32 ChannelComponentIdx, const FSearchIndex& SearchIndex)
 	: FullLabel(ComputeLabel(Channel, ChannelComponentIdx, true))
 	, CompactLabel(ComputeLabel(Channel, ChannelComponentIdx, false))
 	, DataOffset(ComputeDataOffset(Channel, ChannelComponentIdx))
 	, Cardinality(ComputeCardinality(Channel, ChannelComponentIdx))
+	, Weight(ComputeWeight(Channel, ChannelComponentIdx, SearchIndex))
 	{
 	}
 
@@ -29,6 +30,7 @@ public:
 	, CompactLabel(InCompactLabel)
 	, DataOffset(INDEX_NONE)
 	, Cardinality(INDEX_NONE)
+	, Weight(-1.0f)
 	{
 	}
 
@@ -57,6 +59,11 @@ public:
 		return Cardinality;
 	}
 
+	float GetWeight() const
+	{
+		return Weight;
+	}
+
 	bool IsExpanded() const
 	{
 		return bExpanded;
@@ -124,10 +131,31 @@ private:
 		return Cardinality;
 	}
 
+	static float ComputeWeight(const UPoseSearchFeatureChannel* Channel, int32 ChannelComponentIdx, const FSearchIndex& SearchIndex)
+	{
+		// Using double precision for accumulation of component weights
+		float Weight = 0.0f;
+
+		const int32 DataOffset = ComputeDataOffset(Channel, ChannelComponentIdx);
+		const int32 Cardinality = ComputeCardinality(Channel, ChannelComponentIdx);
+
+		for (int32 Component = 0; Component != Cardinality; ++Component)
+		{
+			// Weights are pre-multiplied by inverse variance in PreprocessSearchIndexWeights, so we undo that
+			// here to reconstitute the original weights
+			const float WeightSqrt = SearchIndex.WeightsSqrt[DataOffset + Component];
+			const float Deviation = SearchIndex.DeviationEditorOnly[DataOffset + Component];
+			Weight += FMath::Square(WeightSqrt * Deviation);
+		}
+
+		return Weight;
+	}
+
 	const FString FullLabel;
 	const FString CompactLabel;
 	const int32 DataOffset = 0;
 	const int32 Cardinality = 0;
+	const float Weight = 0.f;
 	bool bExpanded = false;
 	TArray<FChannelItemPtr> ChannelItems;
 };
@@ -179,6 +207,16 @@ public:
 				.Text(FText::FromString(FString::Printf(TEXT("%d"), ChannelItem->GetDataOffset())));
 		}
 
+		if (ColumnName == FName("Weight"))
+		{
+			if (ChannelItem->GetDataOffset() == INDEX_NONE)
+			{
+				return SNew(STextBlock);
+			}
+			return SNew(STextBlock)
+				.Text(FText::FromString(FString::Printf(TEXT("%.3f"), ChannelItem->GetWeight())));
+		}
+
 		if (ColumnName == FName("Query"))
 		{
 			if (ChannelItem->GetDataOffset() == INDEX_NONE)
@@ -463,7 +501,7 @@ void SDatabaseDataDetails::Reconstruct(int32 MaxPreviewActors)
 
 	ChannelItems.Reset();
 
-	RebuildChannelItemsTreeRecursively(ChannelItems, PoseSearchDatabase->Schema->GetChannels());
+	RebuildChannelItemsTreeRecursively(ChannelItems, PoseSearchDatabase->Schema->GetChannels(), PoseSearchDatabase->GetSearchIndex());
 	RebuildChannelItemsStats(ChannelItems);
 
 	// generating the header
@@ -471,13 +509,19 @@ void SDatabaseDataDetails::Reconstruct(int32 MaxPreviewActors)
 	HeaderRow->AddColumn(
 		SHeaderRow::Column(TEXT("ChannelName"))
 			.DefaultLabel(LOCTEXT("ChannelName_Header", "Channel Name"))
-			.ToolTipText(LOCTEXT("ChannelName_ToolTip", "Channel Name"))
+			.DefaultTooltip(LOCTEXT("ChannelName_ToolTip", "Channel Name"))
 			.FillWidth(0.2f));
 
 	HeaderRow->AddColumn(
 		SHeaderRow::Column(TEXT("DataOffset"))
 		.DefaultLabel(LOCTEXT("DataOffset_Header", "Data Offset"))
-		.ToolTipText(LOCTEXT("DataOffset_ToolTip", "Offset from the beginning of the features data"))
+		.DefaultTooltip(LOCTEXT("DataOffset_ToolTip", "Offset from the beginning of the features data"))
+		.FillWidth(0.1f));
+
+	HeaderRow->AddColumn(
+		SHeaderRow::Column(TEXT("Weight"))
+		.DefaultLabel(LOCTEXT("Weight_Header", "Weight"))
+		.DefaultTooltip(LOCTEXT("Weight_ToolTip", "Normalized weight of feature's contribution to pose comparisons.\nWeights are set on channels in the schema, and channels distribute weights over their features.\nWeights of all individual features are then normalized together so they sum to one."))
 		.FillWidth(0.1f));
 
 	if (ViewModel->ShouldDrawQueryVector())
@@ -485,7 +529,7 @@ void SDatabaseDataDetails::Reconstruct(int32 MaxPreviewActors)
 		HeaderRow->AddColumn(
 			SHeaderRow::Column(TEXT("Query"))
 			.DefaultLabel(LOCTEXT("Query_Header", "Query"))
-			.ToolTipText(LOCTEXT("Query_ToolTip", "Query Values")));
+			.DefaultTooltip(LOCTEXT("Query_ToolTip", "Query Values")));
 	}
 
 	int32 AddedColumns = 0;
@@ -539,26 +583,26 @@ void SDatabaseDataDetails::Reconstruct(int32 MaxPreviewActors)
 	SetExpandedItems(ChannelItems, ExpandedItems, ChannelItemsTreeView.Get());
 }
 
-void SDatabaseDataDetails::RebuildChannelItemsTreeRecursively(TArray<FChannelItemPtr>& ChannelItems, TConstArrayView<TObjectPtr<UPoseSearchFeatureChannel>> Channels)
+void SDatabaseDataDetails::RebuildChannelItemsTreeRecursively(TArray<FChannelItemPtr>& ChannelItems, TConstArrayView<TObjectPtr<UPoseSearchFeatureChannel>> Channels, const FSearchIndex& SearchIndex)
 {
 	for (const TObjectPtr<UPoseSearchFeatureChannel>& ChannelPtr : Channels)
 	{
 		if (ChannelPtr)
 		{
-			TSharedRef<FChannelItem> ChannelItem = MakeShareable(new FChannelItem(ChannelPtr));
+			TSharedRef<FChannelItem> ChannelItem = MakeShareable(new FChannelItem(ChannelPtr, INDEX_NONE, SearchIndex));
 			ChannelItems.Add(ChannelItem);
 
 			if (ChannelPtr->GetSubChannels().IsEmpty())
 			{
 				for (int32 ChannelComponentIdx = 0; ChannelComponentIdx < ChannelPtr->GetChannelCardinality(); ++ChannelComponentIdx)
 				{
-					TSharedRef<FChannelItem> SubChannelItem = MakeShareable(new FChannelItem(ChannelPtr, ChannelComponentIdx));
+					TSharedRef<FChannelItem> SubChannelItem = MakeShareable(new FChannelItem(ChannelPtr, ChannelComponentIdx, SearchIndex));
 					ChannelItem->GetChannelItems().Add(SubChannelItem);
 				}
 			}
 			else
 			{
-				RebuildChannelItemsTreeRecursively(ChannelItem->GetChannelItems(), ChannelPtr->GetSubChannels());
+				RebuildChannelItemsTreeRecursively(ChannelItem->GetChannelItems(), ChannelPtr->GetSubChannels(), SearchIndex);
 			}
 		}
 	}
diff --git a/Engine/Plugins/Animation/PoseSearch/Source/Editor/Private/PoseSearchDatabaseEditorReflection.cpp b/Engine/Plugins/Animation/PoseSearch/Source/Editor/Private/PoseSearchDatabaseEditorReflection.cpp
--- a/Engine/Plugins/Animation/PoseSearch/Source/Editor/Private/PoseSearchDatabaseEditorReflection.cpp
+++ b/Engine/Plugins/Animation/PoseSearch/Source/Editor/Private/PoseSearchDatabaseEditorReflection.cpp
@@ -247,7 +247,7 @@ void UPoseSearchDatabaseStatistics::Initialize(const UPoseSearchDatabase* PoseSe
 
 		// Principal Component Analysis
 			
-		ExplainedVariance = SearchIndex.PCAExplainedVariance * 100.f;
+		ExplainedVariance = SearchIndex.PCAExplainedVarianceEditorOnly * 100.f;
 			
 		// Memory Information
 			
diff --git a/Engine/Plugins/Animation/PoseSearch/Source/Editor/Public/PoseSearchDatabaseDataDetails.h b/Engine/Plugins/Animation/PoseSearch/Source/Editor/Public/PoseSearchDatabaseDataDetails.h
--- a/Engine/Plugins/Animation/PoseSearch/Source/Editor/Public/PoseSearchDatabaseDataDetails.h
+++ b/Engine/Plugins/Animation/PoseSearch/Source/Editor/Public/PoseSearchDatabaseDataDetails.h
@@ -12,6 +12,7 @@ namespace UE::PoseSearch
 	class FDatabaseViewModel;
 	typedef TSharedPtr<class FChannelItem> FChannelItemPtr;
 	typedef STreeView<FChannelItemPtr> SChannelItemsTreeView;
+	struct FSearchIndex;
 
 	class SDatabaseDataDetails : public SCompoundWidget
 	{
@@ -23,7 +24,7 @@ namespace UE::PoseSearch
 		void Reconstruct(int32 MaxPreviewActors = 15);
 
 	private:
-		static void RebuildChannelItemsTreeRecursively(TArray<FChannelItemPtr>& ChannelItems, TConstArrayView<TObjectPtr<UPoseSearchFeatureChannel>> Channels);
+		static void RebuildChannelItemsTreeRecursively(TArray<FChannelItemPtr>& ChannelItems, TConstArrayView<TObjectPtr<UPoseSearchFeatureChannel>> Channels, const FSearchIndex& SearchIndex);
 		static void RebuildChannelItemsStats(TArray<FChannelItemPtr>& ChannelItems);
 		static void TrackExpandedItems(const TArray<FChannelItemPtr>& ChannelItems, TMap<const FString, bool>& ExpandedItems);
 		static void SetExpandedItems(TArray<FChannelItemPtr>& ChannelItems, const TMap<const FString, bool>& ExpandedItems, SChannelItemsTreeView* ChannelItemsTreeView);
diff --git a/Engine/Plugins/Animation/PoseSearch/Source/Runtime/Private/PoseSearchDatabase.cpp b/Engine/Plugins/Animation/PoseSearch/Source/Runtime/Private/PoseSearchDatabase.cpp
--- a/Engine/Plugins/Animation/PoseSearch/Source/Runtime/Private/PoseSearchDatabase.cpp
+++ b/Engine/Plugins/Animation/PoseSearch/Source/Runtime/Private/PoseSearchDatabase.cpp
@@ -1093,12 +1093,13 @@ void UPoseSearchDatabase::Serialize(FArchive& Ar)
 {
 	Super::Serialize(Ar);
 
-	if (Ar.IsFilterEditorOnly())
+	if (!IsTemplate())
 	{
-		if (Ar.IsLoading() || Ar.IsCooking())
+		const bool bSavingCooked = Ar.IsSaving() && Ar.IsCooking();
+		const bool bLoadingCooked = Ar.IsLoadingFromCookedPackage();
+		if (bSavingCooked || bLoadingCooked)
 		{
 			Ar << SearchIndexPrivate;
-
 			UpdateCachedProperties();
 		}
 	}
diff --git a/Engine/Plugins/Animation/PoseSearch/Source/Runtime/Private/PoseSearchDerivedData.cpp b/Engine/Plugins/Animation/PoseSearch/Source/Runtime/Private/PoseSearchDerivedData.cpp
--- a/Engine/Plugins/Animation/PoseSearch/Source/Runtime/Private/PoseSearchDerivedData.cpp
+++ b/Engine/Plugins/Animation/PoseSearch/Source/Runtime/Private/PoseSearchDerivedData.cpp
@@ -551,6 +551,11 @@ static bool InitSearchIndexAssets(FSearchIndexBase& SearchIndex, const UPoseSear
 
 static void PreprocessSearchIndexWeights(FSearchIndex& SearchIndex, const UPoseSearchSchema* Schema, TConstArrayView<float> Deviation)
 {
+#if WITH_EDITORONLY_DATA
+	// Copy deviation into search index for weights display in editor
+	SearchIndex.DeviationEditorOnly = Deviation;
+#endif // WITH_EDITORONLY_DATA
+
 	const int32 NumDimensions = Schema->SchemaCardinality;
 	SearchIndex.WeightsSqrt.Init(1.f, NumDimensions);
 
@@ -615,7 +620,12 @@ static Eigen::ComputationInfo PreprocessSearchIndexPCAData(FSearchIndex& SearchI
 	// binding SearchIndex.Values and SearchIndex.PCAValues Eigen row major matrix maps
 	const int32 NumPoses = SearchIndex.GetNumPoses();
 
+#if WITH_EDITORONLY_DATA
+	SearchIndex.PCAExplainedVarianceEditorOnly = 0.f;
+#endif
+	PRAGMA_DISABLE_DEPRECATION_WARNINGS
 	SearchIndex.PCAExplainedVariance = 0.f;
+	PRAGMA_ENABLE_DEPRECATION_WARNINGS
 
 	SearchIndex.PCAValues.Reset();
 	SearchIndex.Mean.Reset();
@@ -697,11 +707,18 @@ static Eigen::ComputationInfo PreprocessSearchIndexPCAData(FSearchIndex& SearchI
 				AccumulatedVariance += EigenValues[Indexer[PCAComponentIndex]];
 			}
 
+			PRAGMA_DISABLE_DEPRECATION_WARNINGS
+			// @todo: move this code under WITH_EDITORONLY_DATA once SearchIndex.PCAExplainedVariance has been removed
 			// calculating the total variance knowing that eigen values measure variance along the principal components:
 			const float TotalVariance = EigenValues.sum();
 			// and explained variance as ratio between AccumulatedVariance and TotalVariance: https://ro-che.info/articles/2017-12-11-pca-explained-variance
 			SearchIndex.PCAExplainedVariance = TotalVariance > UE_KINDA_SMALL_NUMBER ? AccumulatedVariance / TotalVariance : 0.f;
 
+#if WITH_EDITORONLY_DATA
+			SearchIndex.PCAExplainedVarianceEditorOnly = SearchIndex.PCAExplainedVariance;
+#endif // WITH_EDITORONLY_DATA
+			PRAGMA_ENABLE_DEPRECATION_WARNINGS
+
 			MapPCAValues = CenteredValues * PCAProjectionMatrix;
 
 #if ENABLE_ANIM_DEBUG
@@ -1315,11 +1332,6 @@ static void CompareSearchIndex(const FSearchIndex& A, const FSearchIndex& B, con
 		StringBuilder.Append(TEXT("Mean mismatch\n"));
 	}
 
-	if (A.PCAExplainedVariance != B.PCAExplainedVariance)
-	{
-		StringBuilder.Append(TEXT("PCAExplainedVariance mismatch\n"));
-	}
-
 	if (A.KDTree != B.KDTree)
 	{
 		StringBuilder.Append(TEXT("KDTree mismatch\n"));
@@ -1329,6 +1341,25 @@ static void CompareSearchIndex(const FSearchIndex& A, const FSearchIndex& B, con
 	{
 		StringBuilder.Append(TEXT("VPTree mismatch\n"));
 	}
+
+#if WITH_EDITORONLY_DATA
+	if (A.DeviationEditorOnly != B.DeviationEditorOnly)
+	{
+		StringBuilder.Append(TEXT("DeviationEditorOnly mismatch\n"));
+	}
+
+	if (A.PCAExplainedVarianceEditorOnly != B.PCAExplainedVarianceEditorOnly)
+	{
+		StringBuilder.Append(TEXT("PCAExplainedVarianceEditorOnly mismatch\n"));
+	}
+#endif // WITH_EDITORONLY_DATA
+
+	PRAGMA_DISABLE_DEPRECATION_WARNINGS
+	if (A.PCAExplainedVariance != B.PCAExplainedVariance)
+	{
+		StringBuilder.Append(TEXT("PCAExplainedVariance mismatch\n"));
+	}
+	PRAGMA_ENABLE_DEPRECATION_WARNINGS
 }
 
 #endif // ENABLE_ANIM_DEBUG
@@ -2596,6 +2627,12 @@ EAsyncBuildIndexResult FAsyncPoseSearchDatabasesManagement::RequestAsyncBuildInd
 		return EAsyncBuildIndexResult::Failed;
 	}
 
+	if (Database->GetPackage()->bIsCookedForEditor)
+	{
+		// Don't cache for cooked packages
+		return EAsyncBuildIndexResult::Success;
+	}
+
 	FScopeLock Lock(&Mutex);
 
 	check(EnumHasAnyFlags(Flag, ERequestAsyncBuildFlag::NewRequest | ERequestAsyncBuildFlag::ContinueRequest));
diff --git a/Engine/Plugins/Animation/PoseSearch/Source/Runtime/Private/PoseSearchDerivedDataKey.cpp b/Engine/Plugins/Animation/PoseSearch/Source/Runtime/Private/PoseSearchDerivedDataKey.cpp
--- a/Engine/Plugins/Animation/PoseSearch/Source/Runtime/Private/PoseSearchDerivedDataKey.cpp
+++ b/Engine/Plugins/Animation/PoseSearch/Source/Runtime/Private/PoseSearchDerivedDataKey.cpp
@@ -73,7 +73,7 @@ FKeyBuilder::FKeyBuilder(const UObject* Object, bool bUseDataVer, bool bPerformC
 		Hasher.Reset();
 
 		// used to invalidate the key without having to change POSESEARCHDB_DERIVEDDATA_VER all the times
-		int32 POSESEARCHDB_DERIVEDDATA_VER_SMALL = 2660;
+		int32 POSESEARCHDB_DERIVEDDATA_VER_SMALL = 2663;
 		FGuid VersionGuid = FDevSystemGuids::GetSystemGuid(FDevSystemGuids::Get().POSESEARCHDB_DERIVEDDATA_VER);
 		FString AnimationCompressionVersionString = UE::Anim::Compression::AnimationCompressionVersionString;
 
diff --git a/Engine/Plugins/Animation/PoseSearch/Source/Runtime/Private/PoseSearchIndex.cpp b/Engine/Plugins/Animation/PoseSearch/Source/Runtime/Private/PoseSearchIndex.cpp
--- a/Engine/Plugins/Animation/PoseSearch/Source/Runtime/Private/PoseSearchIndex.cpp
+++ b/Engine/Plugins/Animation/PoseSearch/Source/Runtime/Private/PoseSearchIndex.cpp
@@ -436,7 +436,13 @@ FSearchIndex::FSearchIndex(const FSearchIndex& Other)
 	, Mean(Other.Mean)
 	, KDTree(Other.KDTree)
 	, VPTree(Other.VPTree)
+	#if WITH_EDITORONLY_DATA
+	, DeviationEditorOnly(Other.DeviationEditorOnly)
+	, PCAExplainedVarianceEditorOnly(Other.PCAExplainedVarianceEditorOnly)
+	#endif
+	PRAGMA_DISABLE_DEPRECATION_WARNINGS
 	, PCAExplainedVariance(Other.PCAExplainedVariance)
+	PRAGMA_ENABLE_DEPRECATION_WARNINGS
 {
 	check(!PCAValues.IsEmpty() || KDTree.DataSource.PointCount == 0);
 	KDTree.DataSource.Data = PCAValues.IsEmpty() ? nullptr : PCAValues.GetData();
@@ -775,16 +781,22 @@ void FSearchIndex::GetPoseToPCAValuesVectorIndexes(TArray<uint32>& PoseToPCAValu
 
 bool FSearchIndex::operator==(const FSearchIndex& Other) const
 {
-	return
-		FSearchIndexBase::operator==(Other) &&
-		WeightsSqrt == Other.WeightsSqrt &&
-		PCAValues == Other.PCAValues &&
-		PCAValuesVectorToPoseIndexes == Other.PCAValuesVectorToPoseIndexes &&
-		PCAProjectionMatrix == Other.PCAProjectionMatrix &&
-		Mean == Other.Mean &&
-		PCAExplainedVariance == Other.PCAExplainedVariance &&
-		KDTree == Other.KDTree &&
-		VPTree == Other.VPTree;
+	return FSearchIndexBase::operator==(Other)
+		&& WeightsSqrt == Other.WeightsSqrt
+		&& PCAValues == Other.PCAValues
+		&& PCAValuesVectorToPoseIndexes == Other.PCAValuesVectorToPoseIndexes
+		&& PCAProjectionMatrix == Other.PCAProjectionMatrix
+		&& Mean == Other.Mean
+		&& KDTree == Other.KDTree
+		&& VPTree == Other.VPTree
+		PRAGMA_DISABLE_DEPRECATION_WARNINGS
+		&& PCAExplainedVariance == Other.PCAExplainedVariance
+		PRAGMA_ENABLE_DEPRECATION_WARNINGS
+#if WITH_EDITORONLY_DATA
+		&& DeviationEditorOnly == Other.DeviationEditorOnly
+		&& PCAExplainedVarianceEditorOnly == Other.PCAExplainedVarianceEditorOnly
+#endif
+		;
 }
 
 FArchive& operator<<(FArchive& Ar, FSearchIndex& Index)
@@ -796,10 +808,22 @@ FArchive& operator<<(FArchive& Ar, FSearchIndex& Index)
 	Ar << Index.PCAValuesVectorToPoseIndexes;
 	Ar << Index.PCAProjectionMatrix;
 	Ar << Index.Mean;
-	Ar << Index.PCAExplainedVariance;
 	Ar << Index.VPTree;
 
 	Serialize(Ar, Index.KDTree, Index.PCAValues.GetData());
+
+	PRAGMA_DISABLE_DEPRECATION_WARNINGS
+	Ar << Index.PCAExplainedVariance;
+	PRAGMA_ENABLE_DEPRECATION_WARNINGS
+
+#if WITH_EDITORONLY_DATA
+	if (!Ar.IsFilterEditorOnly())
+	{
+		Ar << Index.DeviationEditorOnly;
+		Ar << Index.PCAExplainedVarianceEditorOnly;
+	}
+#endif
+
 	return Ar;
 }
 
diff --git a/Engine/Plugins/Animation/PoseSearch/Source/Runtime/Public/PoseSearch/PoseSearchIndex.h b/Engine/Plugins/Animation/PoseSearch/Source/Runtime/Public/PoseSearch/PoseSearchIndex.h
--- a/Engine/Plugins/Animation/PoseSearch/Source/Runtime/Public/PoseSearch/PoseSearchIndex.h
+++ b/Engine/Plugins/Animation/PoseSearch/Source/Runtime/Public/PoseSearch/PoseSearchIndex.h
@@ -426,7 +426,11 @@ struct FSearchIndex : public FSearchIndexBase
 	FKDTree KDTree;
 	FVPTree VPTree;
 
-	// @todo: this property should be editor only
+#if WITH_EDITORONLY_DATA
+	TAlignedArray<float> DeviationEditorOnly;
+	float PCAExplainedVarianceEditorOnly = 0.f;
+#endif
+	UE_DEPRECATED(5.6, "Use PCAExplainedVarianceEditorOnly")
 	float PCAExplainedVariance = 0.f;
 
 	FSearchIndex() = default;
diff --git a/Engine/Plugins/Animation/PoseSearch/Source/Runtime/Public/PoseSearch/PoseSearchLibrary.h b/Engine/Plugins/Animation/PoseSearch/Source/Runtime/Public/PoseSearch/PoseSearchLibrary.h
--- a/Engine/Plugins/Animation/PoseSearch/Source/Runtime/Public/PoseSearch/PoseSearchLibrary.h
+++ b/Engine/Plugins/Animation/PoseSearch/Source/Runtime/Public/PoseSearch/PoseSearchLibrary.h
@@ -21,7 +21,7 @@ namespace UE::PoseSearch
 struct FAnimationUpdateContext;
 struct FAnimNode_PoseSearchHistoryCollector_Base;
 
-UENUM()
+UENUM(BlueprintType)
 enum class EPoseSearchInterruptMode : uint8
 {
 	// continuing pose search will be performed if valid
