PoseSearch - AnimNext support for UPoseSearchLibrary::MotionMatch FPoseSearchFutureProperties

diff --git a/Engine/Plugins/Animation/PoseSearch/Source/Runtime/Private/PoseSearchContext.cpp b/Engine/Plugins/Animation/PoseSearch/Source/Runtime/Private/PoseSearchContext.cpp
--- a/Engine/Plugins/Animation/PoseSearch/Source/Runtime/Private/PoseSearchContext.cpp
+++ b/Engine/Plugins/Animation/PoseSearch/Source/Runtime/Private/PoseSearchContext.cpp
@@ -15,6 +15,77 @@
 namespace UE::PoseSearch
 {
 	
+const FTransform& GetContextTransform(const UObject* AnimContext)
+{
+	if (const UAnimInstance* AnimInstance = Cast<UAnimInstance>(AnimContext))
+	{
+		return AnimInstance->GetSkelMeshComponent()->GetComponentTransform();
+	}
+	
+	if (const UActorComponent* ActorComponent = Cast<UActorComponent>(AnimContext))
+	{
+		const AActor* Actor = ActorComponent->GetOwner();
+		check(Actor);
+		// @todo: this code depends on how AnimNext gather its context object, and will likely change in the future
+		if (const USkeletalMeshComponent* SkeletalMeshComponent = Actor->GetComponentByClass<USkeletalMeshComponent>())
+		{
+			return SkeletalMeshComponent->GetComponentTransform();
+		}
+
+		return ActorComponent->GetOwner()->GetTransform();
+	}
+
+	unimplemented();
+	return FTransform::Identity;
+}
+
+const USkeleton* GetContextSkeleton(const UObject* AnimContext)
+{
+	if (const UAnimInstance* AnimInstance = Cast<UAnimInstance>(AnimContext))
+	{
+		return AnimInstance->GetRequiredBonesOnAnyThread().GetSkeletonAsset();
+	}
+
+	if (const UActorComponent* ActorComponent = Cast<UActorComponent>(AnimContext))
+	{
+		const AActor* Actor = ActorComponent->GetOwner();
+		check(Actor);
+		// @todo: this code depends on how AnimNext gather its context object, and will likely change in the future
+		if (const USkeletalMeshComponent* SkeletalMeshComponent = Actor->GetComponentByClass<USkeletalMeshComponent>())
+		{
+			if (const USkeletalMesh* SkeletalMesh = SkeletalMeshComponent->GetSkeletalMeshAsset())
+			{
+				return SkeletalMesh->GetSkeleton();
+			}
+		}
+		return nullptr;
+	}
+
+	unimplemented();
+	return nullptr;
+}
+
+const AActor* GetContextOwningActor(const UObject* AnimContext)
+{
+	if (const UAnimInstance* AnimInstance = Cast<UAnimInstance>(AnimContext))
+	{
+		return AnimInstance->GetOwningActor();
+	}
+	
+	if (const UActorComponent* AnimNextComponent = Cast<UActorComponent>(AnimContext))
+	{
+		return AnimNextComponent->GetOwner();
+	}
+	
+	unimplemented();
+	return nullptr;
+}
+
+FVector GetContextLocation(const UObject* AnimContext)
+{
+	return GetContextTransform(AnimContext).GetLocation();
+}
+
 #if ENABLE_DRAW_DEBUG
 
 class UAnimInstanceProxyProvider : public UAnimInstance
diff --git a/Engine/Plugins/Animation/PoseSearch/Source/Runtime/Private/PoseSearchInteractionIsland.cpp b/Engine/Plugins/Animation/PoseSearch/Source/Runtime/Private/PoseSearchInteractionIsland.cpp
--- a/Engine/Plugins/Animation/PoseSearch/Source/Runtime/Private/PoseSearchInteractionIsland.cpp
+++ b/Engine/Plugins/Animation/PoseSearch/Source/Runtime/Private/PoseSearchInteractionIsland.cpp
@@ -4,95 +4,23 @@
 #include "Animation/AnimClassInterface.h"
 #include "Animation/AnimInstance.h"
 #include "Components/SkeletalMeshComponent.h"
-#include "PoseSearch/PoseSearchInteractionLibrary.h"
 #include "DrawDebugHelpers.h"
 #include "Engine/SkeletalMesh.h"
 #include "GameFramework/CharacterMovementComponent.h"
 #include "PoseSearch/AnimNode_PoseSearchHistoryCollector.h"
 #include "PoseSearch/MultiAnimAsset.h"
+#include "PoseSearch/PoseSearchContext.h"
+#include "PoseSearch/PoseSearchDatabase.h"
+#include "PoseSearch/PoseSearchInteractionLibrary.h"
 #include "PoseSearch/PoseSearchInteractionSubsystem.h"
 #include "PoseSearch/PoseSearchInteractionValidator.h"
 #include "PoseSearch/PoseSearchLibrary.h"
-#include "PoseSearch/PoseSearchDatabase.h"
 #include "PoseSearch/PoseSearchSchema.h"
 #include "VisualLogger/VisualLogger.h"
 
 namespace UE::PoseSearch
 {
 
-// Utils functions
-///////////////////////////////////////////////////////////
-const FTransform& GetContextTransform(const UObject* AnimContext)
-{
-	if (const UAnimInstance* AnimInstance = Cast<UAnimInstance>(AnimContext))
-	{
-		return AnimInstance->GetSkelMeshComponent()->GetComponentTransform();
-	}
-	
-	if (const UActorComponent* ActorComponent = Cast<UActorComponent>(AnimContext))
-	{
-		const AActor* Actor = ActorComponent->GetOwner();
-		check(Actor);
-		// @todo: this code depends on how AnimNext gather its context object, and will likely change in the future
-		if (const USkeletalMeshComponent* SkeletalMeshComponent = Actor->GetComponentByClass<USkeletalMeshComponent>())
-		{
-			return SkeletalMeshComponent->GetComponentTransform();
-		}
-
-		return ActorComponent->GetOwner()->GetTransform();
-	}
-
-	unimplemented();
-	return FTransform::Identity;
-}
-
-const USkeleton* GetContextSkeleton(const UObject* AnimContext)
-{
-	if (const UAnimInstance* AnimInstance = Cast<UAnimInstance>(AnimContext))
-	{
-		return AnimInstance->GetRequiredBonesOnAnyThread().GetSkeletonAsset();
-	}
-
-	if (const UActorComponent* ActorComponent = Cast<UActorComponent>(AnimContext))
-	{
-		const AActor* Actor = ActorComponent->GetOwner();
-		check(Actor);
-		// @todo: this code depends on how AnimNext gather its context object, and will likely change in the future
-		if (const USkeletalMeshComponent* SkeletalMeshComponent = Actor->GetComponentByClass<USkeletalMeshComponent>())
-		{
-			if (const USkeletalMesh* SkeletalMesh = SkeletalMeshComponent->GetSkeletalMeshAsset())
-			{
-				return SkeletalMesh->GetSkeleton();
-			}
-		}
-		return nullptr;
-	}
-
-	unimplemented();
-	return nullptr;
-}
-
-const AActor* GetContextOwningActor(const UObject* AnimContext)
-{
-	if (const UAnimInstance* AnimInstance = Cast<UAnimInstance>(AnimContext))
-	{
-		return AnimInstance->GetOwningActor();
-	}
-	
-	if (const UActorComponent* AnimNextComponent = Cast<UActorComponent>(AnimContext))
-	{
-		return AnimNextComponent->GetOwner();
-	}
-	
-	unimplemented();
-	return nullptr;
-}
-
-FVector GetContextLocation(const UObject* AnimContext)
-{
-	return GetContextTransform(AnimContext).GetLocation();
-}
-
 static FInteractionSearchResult InitSearchResult(const FSearchResult& SearchResult, const FInteractionSearchContext& SearchContext, int32 SearchIndex)
 {
 	FInteractionSearchResult InteractionSearchResult;
diff --git a/Engine/Plugins/Animation/PoseSearch/Source/Runtime/Private/PoseSearchInteractionSubsystem.cpp b/Engine/Plugins/Animation/PoseSearch/Source/Runtime/Private/PoseSearchInteractionSubsystem.cpp
--- a/Engine/Plugins/Animation/PoseSearch/Source/Runtime/Private/PoseSearchInteractionSubsystem.cpp
+++ b/Engine/Plugins/Animation/PoseSearch/Source/Runtime/Private/PoseSearchInteractionSubsystem.cpp
@@ -2,8 +2,9 @@
 
 #include "PoseSearch/PoseSearchInteractionSubsystem.h"
 #include "PoseSearch/AnimNode_PoseSearchHistoryCollector.h"
-#include "PoseSearch/PoseSearchInteractionIsland.h"
+#include "PoseSearch/PoseSearchContext.h"
 #include "PoseSearch/PoseSearchDatabase.h"
+#include "PoseSearch/PoseSearchInteractionIsland.h"
 #include "PoseSearch/PoseSearchRole.h"
 #include "PoseSearch/PoseSearchSchema.h"
 #include "VisualLogger/VisualLogger.h"
diff --git a/Engine/Plugins/Animation/PoseSearch/Source/Runtime/Private/PoseSearchInteractionValidator.cpp b/Engine/Plugins/Animation/PoseSearch/Source/Runtime/Private/PoseSearchInteractionValidator.cpp
--- a/Engine/Plugins/Animation/PoseSearch/Source/Runtime/Private/PoseSearchInteractionValidator.cpp
+++ b/Engine/Plugins/Animation/PoseSearch/Source/Runtime/Private/PoseSearchInteractionValidator.cpp
@@ -1,6 +1,7 @@
 // Copyright Epic Games, Inc. All Rights Reserved.
 
 #include "PoseSearch/PoseSearchInteractionValidator.h"
+#include "PoseSearch/PoseSearchContext.h"
 #include "PoseSearch/PoseSearchInteractionSubsystem.h"
 
 namespace UE::PoseSearch
diff --git a/Engine/Plugins/Animation/PoseSearch/Source/Runtime/Private/PoseSearchLibrary.cpp b/Engine/Plugins/Animation/PoseSearch/Source/Runtime/Private/PoseSearchLibrary.cpp
--- a/Engine/Plugins/Animation/PoseSearch/Source/Runtime/Private/PoseSearchLibrary.cpp
+++ b/Engine/Plugins/Animation/PoseSearch/Source/Runtime/Private/PoseSearchLibrary.cpp
@@ -847,7 +847,7 @@ void UPoseSearchLibrary::MotionMatch(
 	AnimInstances.Add(AnimInstance);
 
 	TArray<FName, TInlineAllocator<PreallocatedRolesNum, TMemStackAllocator<>>> Roles;
-	Roles.Add(UE::PoseSearch::DefaultRole);
+	Roles.Add(DefaultRole);
 
 	TArray<const UObject*>& AssetsToSearchConst = reinterpret_cast<TArray<const UObject*>&>(AssetsToSearch);
 	PRAGMA_DISABLE_DEPRECATION_WARNINGS
@@ -899,7 +899,7 @@ void UPoseSearchLibrary::MotionMatch(
 
 	FMemMark Mark(FMemStack::Get());
 
-	TArray<const UE::PoseSearch::IPoseHistory*, TInlineAllocator<PreallocatedRolesNum, TMemStackAllocator<>>> PoseHistories;
+	TArray<const IPoseHistory*, TInlineAllocator<PreallocatedRolesNum, TMemStackAllocator<>>> PoseHistories;
 	TArray<const UObject*, TInlineAllocator<PreallocatedRolesNum, TMemStackAllocator<>>> AnimContexts;
 	for (UAnimInstance* AnimInstance : AnimInstances)
 	{
@@ -968,7 +968,7 @@ UE::PoseSearch::FSearchResult UPoseSearchLibrary::MotionMatch(
 
 	FMemMark Mark(FMemStack::Get());
 
-	TArray<const UE::PoseSearch::IPoseHistory*, TInlineAllocator<PreallocatedRolesNum, TMemStackAllocator<>>> InternalPoseHistories;
+	TArray<const IPoseHistory*, TInlineAllocator<PreallocatedRolesNum, TMemStackAllocator<>>> InternalPoseHistories;
 	InternalPoseHistories = PoseHistories;
 
 	// MemStackPoseHistories will hold future poses to match AssetSamplerBase (at FutureAnimationStartTime) TimeToFutureAnimationStart seconds in the future
@@ -994,48 +994,71 @@ UE::PoseSearch::FSearchResult UPoseSearchLibrary::MotionMatch(
 
 		for (int32 RoleIndex = 0; RoleIndex < Roles.Num(); ++RoleIndex)
 		{
-			if (const UAnimInstance* AnimInstance = Cast<UAnimInstance>(AnimContexts[RoleIndex]))
+			if (const IPoseHistory* PoseHistory = InternalPoseHistories[RoleIndex])
 			{
-				MemStackPoseHistories[RoleIndex].Init(InternalPoseHistories[RoleIndex]);
-
-				// extracting 2 poses to be able to calculate velocities
-				FCSPose<FCompactPose> ComponentSpacePose;
-				FCompactPose Pose;
-				FBlendedCurve Curves;
-				Pose.SetBoneContainer(&AnimInstance->GetRequiredBonesOnAnyThread());
-
-				// @todo: add input BlendParameters to support sampling FutureAnimation blendspaces and support for multi character
-				const UAnimationAsset* AnimationAsset = Cast<UAnimationAsset>(Future.Animation);
-				if (!AnimationAsset)
+				if (USkeleton* Skeleton = const_cast<USkeleton*>(GetContextSkeleton(AnimContexts[RoleIndex])))
 				{
-					if (const UMultiAnimAsset* MultiAnimAsset = Cast<UMultiAnimAsset>(Future.Animation))
+					MemStackPoseHistories[RoleIndex].Init(InternalPoseHistories[RoleIndex]);
+
+					TArray<uint16> BoneIndices;
+					const FBoneToTransformMap& BoneToTransformMap = PoseHistory->GetBoneToTransformMap();
+					if (BoneToTransformMap.IsEmpty())
 					{
-						AnimationAsset = MultiAnimAsset->GetAnimationAsset(Roles[RoleIndex]);
+						const int32 NumBones = Skeleton->GetReferenceSkeleton().GetNum();
+						BoneIndices.SetNum(NumBones);
+						for (int32 BoneIndex = 0; BoneIndex < NumBones; ++BoneIndex)
+						{
+							BoneIndices[BoneIndex] = BoneIndex;
+						}
 					}
 					else
 					{
-						checkNoEntry();
+						for (const FBoneToTransformPair& BoneToTransformPair : BoneToTransformMap)
+						{
+							BoneIndices.Add(BoneToTransformPair.Key);
+						}
+
+						BoneIndices.Sort();
+						FAnimationRuntime::EnsureParentsPresent(BoneIndices, Skeleton->GetReferenceSkeleton());
 					}
-				}
 
-				const FAnimationAssetSampler Sampler(AnimationAsset, FTransform::Identity, FVector::ZeroVector, FAnimationAssetSampler::DefaultRootTransformSamplingRate, false);
-				for (int32 i = 0; i < 2; ++i)
-				{
-					const float FuturePoseExtractionTime = FutureAnimationTime + (i - 1) * FiniteDelta;
-					const float FuturePoseAnimationTime = FutureIntervalTime + (i - 1) * FiniteDelta;
+					FBoneContainer BoneContainer;
+					BoneContainer.InitializeTo(BoneIndices, UE::Anim::FCurveFilterSettings(UE::Anim::ECurveFilterMode::DisallowAll), *Skeleton);
 
-					Sampler.ExtractPose(FuturePoseExtractionTime, Pose, Curves);
-					ComponentSpacePose.InitPose(Pose);
-					FComponentSpacePoseProvider ComponentSpacePoseProvider(ComponentSpacePose);
-					MemStackPoseHistories[RoleIndex].AddFuturePose(FuturePoseAnimationTime, ComponentSpacePoseProvider);
-				}
+					// extracting 2 poses to be able to calculate velocities
+					FCSPose<FCompactPose> ComponentSpacePose;
+					FCompactPose Pose;
+					FBlendedCurve Curves;
+					Pose.SetBoneContainer(&BoneContainer);
 
-				InternalPoseHistories[RoleIndex] = MemStackPoseHistories[RoleIndex].GetThisOrPoseHistory();
-			}
-			else
-			{
-				// @todo: implement FPoseSearchFutureProperties support for AnimNext
-				UE_LOG(LogPoseSearch, Error, TEXT("UPoseSearchLibrary::MotionMatch - FPoseSearchFutureProperties is supported only for UAnimInstance AnimContexts! (Unimplemented)"));
+					// @todo: add input BlendParameters to support sampling FutureAnimation blendspaces and support for multi character
+					const UAnimationAsset* AnimationAsset = Cast<UAnimationAsset>(Future.Animation);
+					if (!AnimationAsset)
+					{
+						if (const UMultiAnimAsset* MultiAnimAsset = Cast<UMultiAnimAsset>(Future.Animation))
+						{
+							AnimationAsset = MultiAnimAsset->GetAnimationAsset(Roles[RoleIndex]);
+						}
+						else
+						{
+							checkNoEntry();
+						}
+					}
+
+					const FAnimationAssetSampler Sampler(AnimationAsset, FTransform::Identity, FVector::ZeroVector, FAnimationAssetSampler::DefaultRootTransformSamplingRate, false);
+					for (int32 i = 0; i < 2; ++i)
+					{
+						const float FuturePoseExtractionTime = FutureAnimationTime + (i - 1) * FiniteDelta;
+						const float FuturePoseAnimationTime = FutureIntervalTime + (i - 1) * FiniteDelta;
+
+						Sampler.ExtractPose(FuturePoseExtractionTime, Pose, Curves);
+						ComponentSpacePose.InitPose(Pose);
+						FComponentSpacePoseProvider ComponentSpacePoseProvider(ComponentSpacePose);
+						MemStackPoseHistories[RoleIndex].AddFuturePose(FuturePoseAnimationTime, ComponentSpacePoseProvider);
+					}
+
+					InternalPoseHistories[RoleIndex] = MemStackPoseHistories[RoleIndex].GetThisOrPoseHistory();
+				}
 			}
 		}		
 	}
diff --git a/Engine/Plugins/Animation/PoseSearch/Source/Runtime/Public/PoseSearch/PoseSearchContext.h b/Engine/Plugins/Animation/PoseSearch/Source/Runtime/Public/PoseSearch/PoseSearchContext.h
--- a/Engine/Plugins/Animation/PoseSearch/Source/Runtime/Public/PoseSearch/PoseSearchContext.h
+++ b/Engine/Plugins/Animation/PoseSearch/Source/Runtime/Public/PoseSearch/PoseSearchContext.h
@@ -16,6 +16,12 @@ class UPoseSearchFeatureChannel_Position;
 namespace UE::PoseSearch
 {
 
+// Utils functions
+const FTransform& GetContextTransform(const UObject* AnimContext);
+const USkeleton* GetContextSkeleton(const UObject* AnimContext);
+const AActor* GetContextOwningActor(const UObject* AnimContext);
+FVector GetContextLocation(const UObject* AnimContext);
+
 struct FPoseIndicesHistory;
 struct IPoseHistory;
 
diff --git a/Engine/Plugins/Animation/PoseSearch/Source/Runtime/Public/PoseSearch/PoseSearchInteractionIsland.h b/Engine/Plugins/Animation/PoseSearch/Source/Runtime/Public/PoseSearch/PoseSearchInteractionIsland.h
--- a/Engine/Plugins/Animation/PoseSearch/Source/Runtime/Public/PoseSearch/PoseSearchInteractionIsland.h
+++ b/Engine/Plugins/Animation/PoseSearch/Source/Runtime/Public/PoseSearch/PoseSearchInteractionIsland.h
@@ -15,18 +15,6 @@ namespace UE::PoseSearch
 {
 struct FInteractionIsland;
 
-// Experimental, this feature might be removed without warning, not for production use
-const FTransform& GetContextTransform(const UObject* AnimContext);
-
-// Experimental, this feature might be removed without warning, not for production use
-const USkeleton* GetContextSkeleton(const UObject* AnimContext);
-
-// Experimental, this feature might be removed without warning, not for production use
-const AActor* GetContextOwningActor(const UObject* AnimContext);
-
-// Experimental, this feature might be removed without warning, not for production use
-FVector GetContextLocation(const UObject* AnimContext);
-
 // Experimental, this feature might be removed without warning, not for production use
 struct FInteractionSearchContextBase
 {
