AnimNext: External component dependency management API

Allows other components etc to depend on AnimNext component tick functions. Does not expose the tick functions directly so we can still internally batch work as required in future.
Initial implementation for PoseSearch interaction islands

#virtualized

diff --git a/Engine/Plugins/Animation/PoseSearch/Source/Runtime/Private/PoseSearchInteractionIsland.cpp b/Engine/Plugins/Animation/PoseSearch/Source/Runtime/Private/PoseSearchInteractionIsland.cpp
--- a/Engine/Plugins/Animation/PoseSearch/Source/Runtime/Private/PoseSearchInteractionIsland.cpp
+++ b/Engine/Plugins/Animation/PoseSearch/Source/Runtime/Private/PoseSearchInteractionIsland.cpp
@@ -6,6 +6,7 @@
 #include "Components/SkeletalMeshComponent.h"
 #include "DrawDebugHelpers.h"
 #include "Engine/SkeletalMesh.h"
+#include "Features/IModularFeatures.h"
 #include "GameFramework/CharacterMovementComponent.h"
 #include "PoseSearch/AnimNode_PoseSearchHistoryCollector.h"
 #include "PoseSearch/MultiAnimAsset.h"
@@ -461,6 +462,77 @@ UActorComponent* FInteractionIsland::FindPostTickComponent(const AActor* Actor)
 	return nullptr;
 }
 
+IInteractionIslandDependency* FInteractionIsland::FindCustomDependency(UActorComponent* InTickComponent) const
+{
+	IModularFeatures& ModularFeatures = IModularFeatures::Get();
+	int32 NumFeatures = ModularFeatures.GetModularFeatureImplementationCount(IInteractionIslandDependency::FeatureName);
+
+	// Add pre-tick function dependencies
+	for(int32 FeatureIndex = 0; FeatureIndex < NumFeatures; ++FeatureIndex)
+	{
+		IInteractionIslandDependency* IslandDependency = static_cast<IInteractionIslandDependency*>(ModularFeatures.GetModularFeatureImplementation(IInteractionIslandDependency::FeatureName, FeatureIndex));
+		if(IslandDependency == nullptr)
+		{
+			continue;
+		}
+
+		if(IslandDependency->CanMakeDependency(InteractionSubsystem, InTickComponent))
+		{
+			return IslandDependency;
+		}
+	}
+	return nullptr;
+}
+
+void FInteractionIsland::AddTickDependencies(UActorComponent* InPreTickComponent, UActorComponent* InPostTickComponent, bool bInIsMainActor)
+{
+	// Add pre-tick function dependencies
+	if(IInteractionIslandDependency* IslandDependency = FindCustomDependency(InPreTickComponent))
+	{
+		IslandDependency->AddPrerequisite(InteractionSubsystem, PreTickFunction, InPreTickComponent);
+	}
+	else
+	{
+		PreTickFunction.AddPrerequisite(InPreTickComponent, InPreTickComponent->PrimaryComponentTick);
+	}
+
+	// Add post-tick function dependencies
+	if(IInteractionIslandDependency* IslandDependency = FindCustomDependency(InPostTickComponent))
+	{
+		if(bInIsMainActor)
+		{
+			IslandDependency->AddSubsequent(InteractionSubsystem, PreTickFunction, InPostTickComponent);
+			IslandDependency->AddPrerequisite(InteractionSubsystem, PostTickFunction, InPostTickComponent);
+
+			check(!bHasTickDependencies);
+			bHasTickDependencies = true;
+		}
+		else
+		{
+			IslandDependency->AddSubsequent(InteractionSubsystem, PostTickFunction, InPostTickComponent);
+			check(bHasTickDependencies);
+		}
+	}
+	else
+	{
+		if(bInIsMainActor)
+		{
+			// it should be this island, but it's not a UObject, so we use the InteractionSubsystem
+			InPostTickComponent->PrimaryComponentTick.AddPrerequisite(InteractionSubsystem, PreTickFunction);
+			PostTickFunction.AddPrerequisite(InPostTickComponent, InPostTickComponent->PrimaryComponentTick);
+
+			check(!bHasTickDependencies);
+			bHasTickDependencies = true;
+		}
+		else
+		{
+			// it should be this island, but it's not a UObject, so we use the InteractionSubsystem
+			InPostTickComponent->PrimaryComponentTick.AddPrerequisite(InteractionSubsystem, PostTickFunction);
+			check(bHasTickDependencies);
+		}
+	}
+}
+
 void FInteractionIsland::InjectToActor(const UObject* AnimContext, bool bAddTickDependencies)
 {
 	check(IsInGameThread());
@@ -506,23 +578,7 @@ void FInteractionIsland::InjectToActor(const UObject* AnimContext, bool bAddTick
 
 				if (bAddTickDependencies)
 				{
-					PreTickFunction.AddPrerequisite(PreTickComponent, PreTickComponent->PrimaryComponentTick);
-
-					if (bIsMainActor)
-					{
-						// it should be this island, but it's not a UObject, so we use the InteractionSubsystem
-						PostTickComponent->PrimaryComponentTick.AddPrerequisite(InteractionSubsystem, PreTickFunction);
-						PostTickFunction.AddPrerequisite(PostTickComponent, PostTickComponent->PrimaryComponentTick);
-
-						check(!bHasTickDependencies);
-						bHasTickDependencies = bAddTickDependencies;
-					}
-					else
-					{
-						// it should be this island, but it's not a UObject, so we use the InteractionSubsystem
-						PostTickComponent->PrimaryComponentTick.AddPrerequisite(InteractionSubsystem, PostTickFunction);
-						check(bHasTickDependencies);
-					}
+					AddTickDependencies(PreTickComponent, PostTickComponent, bIsMainActor);
 				}
 				else
 				{
@@ -589,7 +645,14 @@ void FInteractionIsland::RemoveTickDependencies(bool bValidateTickDependencies)
 		{
 			if (UActorComponent* PreTickComponent = PreTickComponentPtr.Get())
 			{
-				PreTickFunction.RemovePrerequisite(PreTickComponent, PreTickComponent->PrimaryComponentTick);
+				if(IInteractionIslandDependency* IslandDependency = FindCustomDependency(PreTickComponent))
+				{
+					IslandDependency->RemovePrerequisite(InteractionSubsystem, PreTickFunction, PreTickComponent);
+				}
+				else
+				{
+					PreTickFunction.RemovePrerequisite(PreTickComponent, PreTickComponent->PrimaryComponentTick);
+				}
 			}
 		}
 
@@ -598,14 +661,29 @@ void FInteractionIsland::RemoveTickDependencies(bool bValidateTickDependencies)
 		{
 			if (UActorComponent* PostTickComponent = PostTickComponentPtr.Get())
 			{
-				if (bMainActor)
+				if(IInteractionIslandDependency* IslandDependency = FindCustomDependency(PostTickComponent))
 				{
-					PostTickComponent->PrimaryComponentTick.RemovePrerequisite(InteractionSubsystem, PreTickFunction);
-					PostTickFunction.RemovePrerequisite(PostTickComponent, PostTickComponent->PrimaryComponentTick);
+					if (bMainActor)
+					{
+						IslandDependency->RemoveSubsequent(InteractionSubsystem, PreTickFunction, PostTickComponent);
+						IslandDependency->RemovePrerequisite(InteractionSubsystem, PostTickFunction, PostTickComponent);
+					}
+					else
+					{
+						IslandDependency->RemoveSubsequent(InteractionSubsystem, PostTickFunction, PostTickComponent);
+					}
 				}
 				else
 				{
-					PostTickComponent->PrimaryComponentTick.RemovePrerequisite(InteractionSubsystem, PostTickFunction);
+					if (bMainActor)
+					{
+						PostTickComponent->PrimaryComponentTick.RemovePrerequisite(InteractionSubsystem, PreTickFunction);
+						PostTickFunction.RemovePrerequisite(PostTickComponent, PostTickComponent->PrimaryComponentTick);
+					}
+					else
+					{
+						PostTickComponent->PrimaryComponentTick.RemovePrerequisite(InteractionSubsystem, PostTickFunction);
+					}
 				}
 			}
 			bMainActor = false;
diff --git a/Engine/Plugins/Animation/PoseSearch/Source/Runtime/Public/PoseSearch/PoseSearchInteractionIsland.h b/Engine/Plugins/Animation/PoseSearch/Source/Runtime/Public/PoseSearch/PoseSearchInteractionIsland.h
--- a/Engine/Plugins/Animation/PoseSearch/Source/Runtime/Public/PoseSearch/PoseSearchInteractionIsland.h
+++ b/Engine/Plugins/Animation/PoseSearch/Source/Runtime/Public/PoseSearch/PoseSearchInteractionIsland.h
@@ -14,6 +14,7 @@ class UPoseSearchInteractionSubsystem;
 namespace UE::PoseSearch
 {
 struct FInteractionIsland;
+struct IInteractionIslandDependency;
 
 // Experimental, this feature might be removed without warning, not for production use
 struct FInteractionSearchContextBase
@@ -107,7 +108,9 @@ struct FInteractionIsland
 	void InjectToActor(const UObject* AnimContext, bool bAddTickDependencies);
 
 private:
+	void AddTickDependencies(UActorComponent* InPreTickComponent, UActorComponent* InPostTickComponent, bool bInIsMainActor);
 	void RemoveTickDependencies(bool bValidateTickDependencies);
+	IInteractionIslandDependency* FindCustomDependency(UActorComponent* InTickComponent) const;
 
 	bool GetResult_AnyThread(const UObject* AnimContext, FPoseSearchInteractionBlueprintResult& Result);
 	const AActor* GetMainActor() const;
@@ -116,7 +119,7 @@ private:
 #if ENABLE_ANIM_DEBUG
 	void LogTickDependencies() const;
 #endif // ENABLE_ANIM_DEBUG
-
+	
 	struct FPreTickFunction : public FTickFunction
 	{
 		virtual void ExecuteTick(float DeltaTime, enum ELevelTick TickType, ENamedThreads::Type CurrentThread, const FGraphEventRef& MyCompletionGraphEvent) override;
@@ -164,4 +167,18 @@ private:
 #endif // ENABLE_ANIM_DEBUG
 };
 
+// Allows systems other than regular actor components to hook into interaction island dependencies 
+struct IInteractionIslandDependency : public IModularFeature
+{
+	POSESEARCH_API static inline FName FeatureName = "IInteractionIslandDependency";
+
+	virtual ~IInteractionIslandDependency() = default;
+
+	virtual bool CanMakeDependency(const UObject* InIslandObject, const UObject* InObject) const = 0;
+	virtual void AddPrerequisite(UObject* InIslandObject, FTickFunction& InIslandTickFunction, UObject* InObject) const = 0;
+	virtual void AddSubsequent(UObject* InIslandObject, FTickFunction& InIslandTickFunction, UObject* InObject) const = 0;
+	virtual void RemovePrerequisite(UObject* InIslandObject, FTickFunction& InIslandTickFunction, UObject* InObject) const = 0;
+	virtual void RemoveSubsequent(UObject* InIslandObject, FTickFunction& InIslandTickFunction, UObject* InObject) const = 0;
+};
+
 } // namespace UE::PoseSearch
diff --git a/Engine/Plugins/Experimental/AnimNext/Source/AnimNext/Internal/Component/AnimNextWorldSubsystem.h b/Engine/Plugins/Experimental/AnimNext/Source/AnimNext/Internal/Component/AnimNextWorldSubsystem.h
--- a/Engine/Plugins/Experimental/AnimNext/Source/AnimNext/Internal/Component/AnimNextWorldSubsystem.h
+++ b/Engine/Plugins/Experimental/AnimNext/Source/AnimNext/Internal/Component/AnimNextWorldSubsystem.h
@@ -82,6 +82,15 @@ public:
 	
 	static void AddReferencedObjects(UObject* InThis, FReferenceCollector& Collector);
 
+	// Dependency type passed to AddDependencyHandle etc.
+	enum class EDependency : uint8
+	{
+		// Dependency runs before the specified event
+		Prerequisite,
+		// Dependency runs after the specified event
+		Subsequent
+	};
+
 private:
 	// UWorldSubsystem interface
 	virtual bool DoesSupportWorldType(const EWorldType::Type WorldType) const override;
@@ -139,6 +148,22 @@ protected:
 	// Disables the initialize event for a module once it has run once
 	void DisableInitializeEvent(UE::AnimNext::FModuleHandle InHandle);
 
+	// Add a dependency on a tick function to the specified event
+	// @param	InHandle			The handle of the module whose dependencies are being modified
+	// @param	InObject			The object that owns the tick function
+	// @param	InTickFunction		The tick function to depend on
+	// @param	InEventName			The event to add the dependency to
+	// @param	InDependency			The kind of dependency to add
+	void AddDependencyHandle(UE::AnimNext::FModuleHandle InHandle, UObject* InObject, FTickFunction& InTickFunction, FName InEventName, EDependency InDependency);
+	
+	// Remove a dependency on a tick function from the specified event
+	// @param	InHandle			The handle of the module whose dependencies are being modified
+	// @param	InObject			The object that owns the tick function
+	// @param	InTickFunction		The tick function that was depended on
+	// @param	InEventName			The event to remove the dependency from
+	// @param	InDependency			The kind of dependency to remove
+	void RemoveDependencyHandle(UE::AnimNext::FModuleHandle InHandle, UObject* InObject, FTickFunction& InTickFunction, FName InEventName, EDependency InDependency);
+
 protected:
 	friend struct FRigUnit_AnimNextInitializeEvent;
 	friend struct UE::AnimNext::FModuleEndTickFunction;
diff --git a/Engine/Plugins/Experimental/AnimNext/Source/AnimNext/Internal/Module/AnimNextModuleInstance.h b/Engine/Plugins/Experimental/AnimNext/Source/AnimNext/Internal/Module/AnimNextModuleInstance.h
--- a/Engine/Plugins/Experimental/AnimNext/Source/AnimNext/Internal/Module/AnimNextModuleInstance.h
+++ b/Engine/Plugins/Experimental/AnimNext/Source/AnimNext/Internal/Module/AnimNextModuleInstance.h
@@ -157,6 +157,9 @@ private:
 	// Clear data that binds the schedule to a runtime (e.g. tick functions) and any instance data
 	void ResetBindingsAndInstanceData();
 
+	// Remove all of our tick function's dependencies
+	void RemoveAllTickDependencies();
+
 	// Flip proxy variable buffers then copy any dirty values
 	void CopyProxyVariables();
 
diff --git a/Engine/Plugins/Experimental/AnimNext/Source/AnimNext/Internal/Module/ModuleTickFunction.h b/Engine/Plugins/Experimental/AnimNext/Source/AnimNext/Internal/Module/ModuleTickFunction.h
--- a/Engine/Plugins/Experimental/AnimNext/Source/AnimNext/Internal/Module/ModuleTickFunction.h
+++ b/Engine/Plugins/Experimental/AnimNext/Source/AnimNext/Internal/Module/ModuleTickFunction.h
@@ -43,6 +43,10 @@ struct FModuleEventTickFunction : public FTickFunction
 	static void InitializeAndRunModule(FAnimNextModuleInstance& InModuleInstance);
 #endif
 
+	void AddSubsequent(UObject* InObject, FTickFunction& InTickFunction);
+	void RemoveSubsequent(UObject* InObject, FTickFunction& InTickFunction);
+	void RemoveAllExternalSubsequents();
+
 	FAnimNextModuleInstance* ModuleInstance = nullptr;
 
 	// Tasks queued to execute once per event
@@ -55,6 +59,9 @@ struct FModuleEventTickFunction : public FTickFunction
 	FName EventName;
 
 	bool bUserEvent = false;
+
+	// External dependencies that this tick function needs to unregister when it is destroyed
+	TArray<FTickPrerequisite> ExternalSubsequents;
 };
 
 struct FModuleEndTickFunction : public FTickFunction
diff --git a/Engine/Plugins/Experimental/AnimNext/Source/AnimNext/Private/Component/AnimNextComponent.cpp b/Engine/Plugins/Experimental/AnimNext/Source/AnimNext/Private/Component/AnimNextComponent.cpp
--- a/Engine/Plugins/Experimental/AnimNext/Source/AnimNext/Private/Component/AnimNextComponent.cpp
+++ b/Engine/Plugins/Experimental/AnimNext/Source/AnimNext/Private/Component/AnimNextComponent.cpp
@@ -904,3 +904,35 @@ void UAnimNextComponent::QueueInputTraitEvent(FAnimNextTraitEventPtr Event)
 		Subsystem->QueueInputTraitEvent(this, Event);
 	}
 }
+
+void UAnimNextComponent::AddPrerequisite(UObject* InObject, FTickFunction& InTickFunction, FName InEventName)
+{
+	if(Subsystem)
+	{
+		Subsystem->AddDependency(this, InObject, InTickFunction, InEventName, UAnimNextWorldSubsystem::EDependency::Prerequisite);
+	}
+}
+
+void UAnimNextComponent::AddSubsequent(UObject* InObject, FTickFunction& InTickFunction, FName InEventName)
+{
+	if(Subsystem)
+	{
+		Subsystem->AddDependency(this, InObject, InTickFunction, InEventName, UAnimNextWorldSubsystem::EDependency::Subsequent);
+	}
+}
+
+void UAnimNextComponent::RemovePrerequisite(UObject* InObject, FTickFunction& InTickFunction, FName InEventName)
+{
+	if(Subsystem)
+	{
+		Subsystem->RemoveDependency(this, InObject, InTickFunction, InEventName, UAnimNextWorldSubsystem::EDependency::Prerequisite);
+	}
+}
+
+void UAnimNextComponent::RemoveSubsequent(UObject* InObject, FTickFunction& InTickFunction, FName InEventName)
+{
+	if(Subsystem)
+	{
+		Subsystem->RemoveDependency(this, InObject, InTickFunction, InEventName, UAnimNextWorldSubsystem::EDependency::Subsequent);
+	}
+}
diff --git a/Engine/Plugins/Experimental/AnimNext/Source/AnimNext/Private/Component/AnimNextComponentWorldSubsystem.cpp b/Engine/Plugins/Experimental/AnimNext/Source/AnimNext/Private/Component/AnimNextComponentWorldSubsystem.cpp
--- a/Engine/Plugins/Experimental/AnimNext/Source/AnimNext/Private/Component/AnimNextComponentWorldSubsystem.cpp
+++ b/Engine/Plugins/Experimental/AnimNext/Source/AnimNext/Private/Component/AnimNextComponentWorldSubsystem.cpp
@@ -67,6 +67,22 @@ void UAnimNextComponentWorldSubsystem::QueueInputTraitEvent(UAnimNextComponent*
 	QueueInputTraitEventHandle(InComponent->ModuleHandle, Event);
 }
 
+void UAnimNextComponentWorldSubsystem::AddDependency(UAnimNextComponent* InComponent, UObject* InObject, FTickFunction& InTickFunction, FName InEventName, EDependency InDependency)
+{
+	using namespace UE::AnimNext;
+	check(IsInGameThread());
+	check(InComponent);
+	AddDependencyHandle(InComponent->ModuleHandle, InObject, InTickFunction, InEventName, InDependency);
+}
+
+void UAnimNextComponentWorldSubsystem::RemoveDependency(UAnimNextComponent* InComponent, UObject* InObject, FTickFunction& InTickFunction, FName InEventName, EDependency InDependency)
+{
+	using namespace UE::AnimNext;
+	check(IsInGameThread());
+	check(InComponent);
+	RemoveDependencyHandle(InComponent->ModuleHandle, InObject, InTickFunction, InEventName, InDependency);
+}
+
 #if WITH_EDITOR
 
 void UAnimNextComponentWorldSubsystem::OnModuleCompiled(UAnimNextModule* InModule)
diff --git a/Engine/Plugins/Experimental/AnimNext/Source/AnimNext/Private/Component/AnimNextComponentWorldSubsystem.h b/Engine/Plugins/Experimental/AnimNext/Source/AnimNext/Private/Component/AnimNextComponentWorldSubsystem.h
--- a/Engine/Plugins/Experimental/AnimNext/Source/AnimNext/Private/Component/AnimNextComponentWorldSubsystem.h
+++ b/Engine/Plugins/Experimental/AnimNext/Source/AnimNext/Private/Component/AnimNextComponentWorldSubsystem.h
@@ -49,6 +49,20 @@ class UAnimNextComponentWorldSubsystem : public UAnimNextWorldSubsystem
 	// Queue an input trait event
 	void QueueInputTraitEvent(UAnimNextComponent* InComponent, FAnimNextTraitEventPtr Event);
 
+	// Add a dependency on a tick function to the specified event
+	// @param	InComponent			The component whose dependencies are being modified
+	// @param	InObject			The object that owns the tick function
+	// @param	InTickFunction		The tick function to depend on
+	// @param	InEventName			The event to add the dependency to
+	void AddDependency(UAnimNextComponent* InComponent, UObject* InObject, FTickFunction& InTickFunction, FName InEventName, EDependency InDependency);
+
+	// Remove a dependency on a tick function from the specified event
+	// @param	InComponent			The component whose dependencies are being modified
+	// @param	InObject			The object that owns the tick function
+	// @param	InTickFunction		The tick function that was depended on
+	// @param	InEventName			The event to remove the dependency from
+	void RemoveDependency(UAnimNextComponent* InComponent, UObject* InObject, FTickFunction& InTickFunction, FName InEventName, EDependency InDependency);
+
 #if WITH_EDITOR
 	// UAnimNextWorldSubsystem interface
 	virtual void OnModuleCompiled(UAnimNextModule* InModule) override;
diff --git a/Engine/Plugins/Experimental/AnimNext/Source/AnimNext/Private/Component/AnimNextWorldSubsystem.cpp b/Engine/Plugins/Experimental/AnimNext/Source/AnimNext/Private/Component/AnimNextWorldSubsystem.cpp
--- a/Engine/Plugins/Experimental/AnimNext/Source/AnimNext/Private/Component/AnimNextWorldSubsystem.cpp
+++ b/Engine/Plugins/Experimental/AnimNext/Source/AnimNext/Private/Component/AnimNextWorldSubsystem.cpp
@@ -168,16 +168,17 @@ void UAnimNextWorldSubsystem::UnregisterHandle(UE::AnimNext::FModuleHandle& InOu
 
 	if(IsValidHandle(InOutHandle))
 	{
+		FWriteScopeLock InstancesLockScope(InstancesLock);
+		FAnimNextModuleInstance& Instance = Instances.Get(InOutHandle);
+		FModuleWriteGuard Guard(&Instance);
 #if UE_ENABLE_DEBUG_DRAWING
 		// Release debug drawing immediately as the renderer will need to know about this before EOF
-		{
-			FWriteScopeLock InstancesLockScope(InstancesLock);
-			FAnimNextModuleInstance& Instance = Instances.Get(InOutHandle);
-			FModuleWriteGuard Guard(&Instance);
-			Instance.DebugDraw.Reset();
-		}
+		Instance.DebugDraw.Reset();
 #endif
 
+		// Remove all tick dependencies immediately, as once the handle has been invalidated there is no way for external systems to remove their dependencies
+		Instance.RemoveAllTickDependencies();
+
 		PendingActions.Emplace(InOutHandle, FModulePendingAction::EType::ReleaseHandle);
 		InOutHandle.Reset();
 	}
@@ -352,4 +353,54 @@ void UAnimNextWorldSubsystem::RemoveModuleDependencyHandle(UE::AnimNext::FModule
 	}
 
 	SubsequentInstance->RemovePrerequisite(*PrerequisiteInstance);
+}
+
+void UAnimNextWorldSubsystem::AddDependencyHandle(UE::AnimNext::FModuleHandle InHandle, UObject* InObject, FTickFunction& InTickFunction, FName InEventName, EDependency InDependency)
+{
+	using namespace UE::AnimNext;
+	check(IsInGameThread());
+	if (IsValidHandle(InHandle))
+	{
+		FAnimNextModuleInstance& Instance = Instances.Get(InHandle);
+		FModuleEventTickFunction* TickFunction = Instance.FindTickFunction(InEventName);
+		if(TickFunction == nullptr)
+		{
+			UE_LOGFMT(LogAnimation, Warning, "AddDependencyHandle: Could not find event '{EventName}' in module '{ModuleName}'", InEventName, Instance.GetDataInterfaceName());
+			return;
+		}
+
+		if(InDependency == EDependency::Prerequisite)
+		{
+			TickFunction->AddPrerequisite(InObject, InTickFunction);
+		}
+		else
+		{
+			TickFunction->AddSubsequent(InObject, InTickFunction);
+		}
+	}
+}
+
+void UAnimNextWorldSubsystem::RemoveDependencyHandle(UE::AnimNext::FModuleHandle InHandle, UObject* InObject, FTickFunction& InTickFunction, FName InEventName, EDependency InDependency)
+{
+	using namespace UE::AnimNext;
+	check(IsInGameThread());
+	if (IsValidHandle(InHandle))
+	{
+		FAnimNextModuleInstance& Instance = Instances.Get(InHandle);
+		FModuleEventTickFunction* TickFunction = Instance.FindTickFunction(InEventName);
+		if(TickFunction == nullptr)
+		{
+			UE_LOGFMT(LogAnimation, Warning, "RemoveDependencyHandle: Could not find event '{EventName}' in module '{ModuleName}'", InEventName, Instance.GetDataInterfaceName());
+			return;
+		}
+
+		if(InDependency == EDependency::Prerequisite)
+		{
+			TickFunction->RemovePrerequisite(InObject, InTickFunction);
+		}
+		else
+		{
+			TickFunction->RemoveSubsequent(InObject, InTickFunction);
+		}
+	}
 }
\ No newline at end of file
diff --git a/Engine/Plugins/Experimental/AnimNext/Source/AnimNext/Private/Module/AnimNextModuleInstance.cpp b/Engine/Plugins/Experimental/AnimNext/Source/AnimNext/Private/Module/AnimNextModuleInstance.cpp
--- a/Engine/Plugins/Experimental/AnimNext/Source/AnimNext/Private/Module/AnimNextModuleInstance.cpp
+++ b/Engine/Plugins/Experimental/AnimNext/Source/AnimNext/Private/Module/AnimNextModuleInstance.cpp
@@ -273,6 +273,18 @@ void FAnimNextModuleInstance::Uninitialize()
 	}
 }
 
+void FAnimNextModuleInstance::RemoveAllTickDependencies()
+{
+	using namespace UE::AnimNext;
+
+	check(IsInGameThread());
+
+	for (FModuleEventTickFunction& TickFunction : TickFunctions)
+	{
+		TickFunction.RemoveAllExternalSubsequents();
+	}
+}
+
 void FAnimNextModuleInstance::ResetBindingsAndInstanceData()
 {
 	using namespace UE::AnimNext;
@@ -285,6 +297,8 @@ void FAnimNextModuleInstance::ResetBindingsAndInstanceData()
 
 	for (FModuleEventTickFunction& TickFunction : TickFunctions)
 	{
+		// We should have released all external dependencies by now via RemoveAllTickDependencies
+		check(TickFunction.ExternalSubsequents.Num() == 0);
 		TickFunction.UnRegisterTickFunction();
 	}
 
diff --git a/Engine/Plugins/Experimental/AnimNext/Source/AnimNext/Private/Module/ModuleTickFunction.cpp b/Engine/Plugins/Experimental/AnimNext/Source/AnimNext/Private/Module/ModuleTickFunction.cpp
--- a/Engine/Plugins/Experimental/AnimNext/Source/AnimNext/Private/Module/ModuleTickFunction.cpp
+++ b/Engine/Plugins/Experimental/AnimNext/Source/AnimNext/Private/Module/ModuleTickFunction.cpp
@@ -186,6 +186,44 @@ void FModuleEventTickFunction::InitializeAndRunModule(FAnimNextModuleInstance& I
 
 #endif
 
+void FModuleEventTickFunction::AddSubsequent(UObject* InObject, FTickFunction& InTickFunction)
+{
+	auto FindExistingSubsequent = [InObject, &InTickFunction](const FTickPrerequisite& InSubsequent)
+	{
+		return InSubsequent.PrerequisiteObject == InObject && InSubsequent.PrerequisiteTickFunction == &InTickFunction;
+	};
+
+	InTickFunction.AddPrerequisite(ModuleInstance->GetObject(), *this);
+	if(!ExternalSubsequents.ContainsByPredicate(FindExistingSubsequent))
+	{
+		ExternalSubsequents.Emplace(InObject, InTickFunction);
+	}
+}
+
+void FModuleEventTickFunction::RemoveSubsequent(UObject* InObject, FTickFunction& InTickFunction)
+{
+	auto FindExistingSubsequent = [InObject, &InTickFunction](const FTickPrerequisite& InSubsequent)
+	{
+		return InSubsequent.PrerequisiteObject == InObject && InSubsequent.PrerequisiteTickFunction == &InTickFunction;
+	};
+
+	InTickFunction.RemovePrerequisite(ModuleInstance->GetObject(), *this);
+	ExternalSubsequents.RemoveAllSwap(FindExistingSubsequent, EAllowShrinking::No);
+}
+
+void FModuleEventTickFunction::RemoveAllExternalSubsequents()
+{
+	for(FTickPrerequisite& Subsequent : ExternalSubsequents)
+	{
+		if(FTickFunction* TickFunction = Subsequent.Get())
+		{
+			TickFunction->RemovePrerequisite(ModuleInstance->GetObject(), *this);
+		}
+	}
+
+	ExternalSubsequents.Reset();
+}
+
 FString FModuleEventTickFunction::DiagnosticMessage()
 {
 	TStringBuilder<256> Builder;
diff --git a/Engine/Plugins/Experimental/AnimNext/Source/AnimNext/Public/Component/AnimNextComponent.h b/Engine/Plugins/Experimental/AnimNext/Source/AnimNext/Public/Component/AnimNextComponent.h
--- a/Engine/Plugins/Experimental/AnimNext/Source/AnimNext/Public/Component/AnimNextComponent.h
+++ b/Engine/Plugins/Experimental/AnimNext/Source/AnimNext/Public/Component/AnimNextComponent.h
@@ -205,6 +205,30 @@ public:
 	// Get the handle to the registered module
 	ANIMNEXT_API UE::AnimNext::FModuleHandle GetModuleHandle() const { return ModuleHandle; }
 
+	// Add a prerequisite tick function dependency to the specified event
+	// @param	InObject			The object that owns the tick function
+	// @param	InTickFunction		The tick function to depend on
+	// @param	InEventName			The event to add the dependency to
+	ANIMNEXT_API void AddPrerequisite(UObject* InObject, FTickFunction& InTickFunction, FName InEventName);
+
+	// Add a subsequent tick function dependency to the specified event
+	// @param	InObject			The object that owns the tick function
+	// @param	InTickFunction		The tick function to depend on
+	// @param	InEventName			The event to add the dependency to
+	ANIMNEXT_API void AddSubsequent(UObject* InObject, FTickFunction& InTickFunction, FName InEventName);
+
+	// Remove a prerequisite tick function dependency from the specified event
+	// @param	InObject			The object that owns the tick function
+	// @param	InTickFunction		The tick function that was depended on
+	// @param	InEventName			The event to remove the dependency from
+	ANIMNEXT_API void RemovePrerequisite(UObject* InObject, FTickFunction& InTickFunction, FName InEventName);
+
+	// Remove a prerequisite tick function dependency from the specified event
+	// @param	InObject			The object that owns the tick function
+	// @param	InTickFunction		The tick function that was depended on
+	// @param	InEventName			The event to remove the dependency from
+	ANIMNEXT_API void RemoveSubsequent(UObject* InObject, FTickFunction& InTickFunction, FName InEventName);
+
 private:
 	DECLARE_FUNCTION(execSetVariable);
 
diff --git a/Engine/Plugins/Experimental/AnimNextPoseSearch/Source/AnimNextPoseSearch/Private/AnimNextInteractionIslandDependency.cpp b/Engine/Plugins/Experimental/AnimNextPoseSearch/Source/AnimNextPoseSearch/Private/AnimNextInteractionIslandDependency.cpp
new file mode 100644
--- /dev/null
+++ b/Engine/Plugins/Experimental/AnimNextPoseSearch/Source/AnimNextPoseSearch/Private/AnimNextInteractionIslandDependency.cpp
@@ -0,0 +1,42 @@
+﻿// Copyright Epic Games, Inc. All Rights Reserved.
+
+#include "AnimNextInteractionIslandDependency.h"
+
+#include "Component/AnimNextComponent.h"
+#include "Module/RigUnit_AnimNextModuleEvents.h"
+
+namespace UE::PoseSearch
+{
+
+FAnimNextInteractionIslandDependency FAnimNextInteractionIslandDependency::ModularFeature;
+
+bool FAnimNextInteractionIslandDependency::CanMakeDependency(const UObject* InIslandObject, const UObject* InObject) const
+{
+	return InObject->IsA<UAnimNextComponent>();
+}
+
+void FAnimNextInteractionIslandDependency::AddPrerequisite(UObject* InIslandObject, FTickFunction& InIslandTickFunction, UObject* InObject) const
+{
+	UAnimNextComponent* AnimNextComponent = CastChecked<UAnimNextComponent>(InObject);
+	AnimNextComponent->AddSubsequent(InIslandObject, InIslandTickFunction, FRigUnit_AnimNextPrePhysicsEvent::EventName);
+}
+
+void FAnimNextInteractionIslandDependency::AddSubsequent(UObject* InIslandObject, FTickFunction& InIslandTickFunction, UObject* InObject) const
+{
+	UAnimNextComponent* AnimNextComponent = CastChecked<UAnimNextComponent>(InObject);
+	AnimNextComponent->AddPrerequisite(InIslandObject, InIslandTickFunction, FRigUnit_AnimNextPrePhysicsEvent::EventName);
+}
+
+void FAnimNextInteractionIslandDependency::RemovePrerequisite(UObject* InIslandObject, FTickFunction& InIslandTickFunction, UObject* InObject) const
+{
+	UAnimNextComponent* AnimNextComponent = CastChecked<UAnimNextComponent>(InObject);
+	AnimNextComponent->RemoveSubsequent(InIslandObject, InIslandTickFunction, FRigUnit_AnimNextPrePhysicsEvent::EventName);
+}
+
+void FAnimNextInteractionIslandDependency::RemoveSubsequent(UObject* InIslandObject, FTickFunction& InIslandTickFunction, UObject* InObject) const
+{
+	UAnimNextComponent* AnimNextComponent = CastChecked<UAnimNextComponent>(InObject);
+	AnimNextComponent->RemovePrerequisite(InIslandObject, InIslandTickFunction, FRigUnit_AnimNextPrePhysicsEvent::EventName);
+}
+
+}
diff --git a/Engine/Plugins/Experimental/AnimNextPoseSearch/Source/AnimNextPoseSearch/Private/AnimNextInteractionIslandDependency.h b/Engine/Plugins/Experimental/AnimNextPoseSearch/Source/AnimNextPoseSearch/Private/AnimNextInteractionIslandDependency.h
new file mode 100644
--- /dev/null
+++ b/Engine/Plugins/Experimental/AnimNextPoseSearch/Source/AnimNextPoseSearch/Private/AnimNextInteractionIslandDependency.h
@@ -0,0 +1,31 @@
+﻿// Copyright Epic Games, Inc. All Rights Reserved.
+
+#pragma once
+
+#include "PoseSearch/PoseSearchInteractionIsland.h"
+
+namespace UE::AnimNext::PoseSearch
+{
+	class FModule;
+}
+
+namespace UE::PoseSearch
+{
+
+struct FAnimNextInteractionIslandDependency : public IInteractionIslandDependency
+{
+private:
+	friend UE::AnimNext::PoseSearch::FModule;
+	
+	virtual bool CanMakeDependency(const UObject* InIslandObject, const UObject* InObject) const override;
+	virtual void AddPrerequisite(UObject* InIslandObject, FTickFunction& InIslandTickFunction, UObject* InObject) const override;
+	virtual void AddSubsequent(UObject* InIslandObject, FTickFunction& InIslandTickFunction, UObject* InObject) const override;
+	virtual void RemovePrerequisite(UObject* InIslandObject, FTickFunction& InIslandTickFunction, UObject* InObject) const override;
+	virtual void RemoveSubsequent(UObject* InIslandObject, FTickFunction& InIslandTickFunction, UObject* InObject) const override;
+
+	// Single impl for modular feature
+	static FAnimNextInteractionIslandDependency ModularFeature;
+};
+
+}
+
diff --git a/Engine/Plugins/Experimental/AnimNextPoseSearch/Source/AnimNextPoseSearch/Private/Module.cpp b/Engine/Plugins/Experimental/AnimNextPoseSearch/Source/AnimNextPoseSearch/Private/Module.cpp
--- a/Engine/Plugins/Experimental/AnimNextPoseSearch/Source/AnimNextPoseSearch/Private/Module.cpp
+++ b/Engine/Plugins/Experimental/AnimNextPoseSearch/Source/AnimNextPoseSearch/Private/Module.cpp
@@ -7,6 +7,8 @@
 #include "RigVMCore/RigVMRegistry.h"
 
 #include "AlphaBlend.h"
+#include "AnimNextInteractionIslandDependency.h"
+#include "Features/IModularFeatures.h"
 #include "PoseSearch/PoseSearchDatabase.h"
 #include "PoseSearch/PoseSearchInteractionLibrary.h"
 #include "PoseSearch/PoseSearchSchema.h"
@@ -37,10 +39,13 @@ public:
 		FRigVMRegistry& RigVMRegistry = FRigVMRegistry::Get();
 		RigVMRegistry.RegisterObjectTypes(AllowedObjectTypes);
 		RigVMRegistry.RegisterStructTypes(AllowedStructTypes);
+
+		IModularFeatures::Get().RegisterModularFeature(UE::PoseSearch::IInteractionIslandDependency::FeatureName, &UE::PoseSearch::FAnimNextInteractionIslandDependency::ModularFeature);
 	}
 
 	virtual void ShutdownModule() override
 	{
+		IModularFeatures::Get().UnregisterModularFeature(UE::PoseSearch::IInteractionIslandDependency::FeatureName, &UE::PoseSearch::FAnimNextInteractionIslandDependency::ModularFeature);
 	}
 };
 
