PoseSearch - fix for FAIPComponentSpacePoseProvider providing non identity transforms during the first update

diff --git a/Engine/Plugins/Animation/PoseSearch/Source/Runtime/Private/AnimNode_PoseSearchHistoryCollector.cpp b/Engine/Plugins/Animation/PoseSearch/Source/Runtime/Private/AnimNode_PoseSearchHistoryCollector.cpp
--- a/Engine/Plugins/Animation/PoseSearch/Source/Runtime/Private/AnimNode_PoseSearchHistoryCollector.cpp
+++ b/Engine/Plugins/Animation/PoseSearch/Source/Runtime/Private/AnimNode_PoseSearchHistoryCollector.cpp
@@ -89,11 +89,11 @@ void FAnimNode_PoseSearchHistoryCollector_Base::Initialize_AnyThread(const FAnim
 
 	PoseHistory.Initialize_AnyThread(PoseCount, SamplingInterval);
 
-	const TArray<FBoneIndexType> RequiredBones = GetRequiredBones(Context.AnimInstanceProxy);
+	FMemMark Mark(FMemStack::Get());
 	UE::PoseSearch::FAIPComponentSpacePoseProvider ComponentSpacePoseProvider(Context.AnimInstanceProxy);
-			PoseHistory.EvaluateComponentSpace_AnyThread(0.f, ComponentSpacePoseProvider, bStoreScales,
-				RootBoneRecoveryTime, RootBoneTranslationRecoveryRatio, RootBoneRotationRecoveryRatio, true, true, 
-		RequiredBones, FBlendedCurve(), MakeConstArrayView(CollectedCurves));
+	PoseHistory.EvaluateComponentSpace_AnyThread(0.f, ComponentSpacePoseProvider, bStoreScales,
+		RootBoneRecoveryTime, RootBoneTranslationRecoveryRatio, RootBoneRotationRecoveryRatio, true, true, 
+		GetRequiredBones(Context.AnimInstanceProxy), FBlendedCurve(), MakeConstArrayView(CollectedCurves));
 }
 
 TArray<FBoneIndexType> FAnimNode_PoseSearchHistoryCollector_Base::GetRequiredBones(const FAnimInstanceProxy* AnimInstanceProxy) const
diff --git a/Engine/Plugins/Animation/PoseSearch/Source/Runtime/Private/PoseSearchHistory.cpp b/Engine/Plugins/Animation/PoseSearch/Source/Runtime/Private/PoseSearchHistory.cpp
--- a/Engine/Plugins/Animation/PoseSearch/Source/Runtime/Private/PoseSearchHistory.cpp
+++ b/Engine/Plugins/Animation/PoseSearch/Source/Runtime/Private/PoseSearchHistory.cpp
@@ -275,20 +275,57 @@ const USkeleton* FComponentSpacePoseProvider::GetSkeletonAsset() const
 //////////////////////////////////////////////////////////////////////////
 // FAIPComponentSpacePoseProvider
 FAIPComponentSpacePoseProvider::FAIPComponentSpacePoseProvider(const FAnimInstanceProxy* InAnimInstanceProxy)
-	: AnimInstanceProxy(InAnimInstanceProxy)
 {
-	check(AnimInstanceProxy && GetSkeletonAsset());
+	check(InAnimInstanceProxy);
+	// initializing PoseHistory with a ref pose at FAnimInstanceProxy location/facing
+	const FBoneContainer& BoneContainer = InAnimInstanceProxy->GetRequiredBones();
+
+	// BoneContainer can be invalid when recompiling ABP while PIE is running
+	if (BoneContainer.IsValid())
+	{
+		ComponentSpacePose.InitPose(&BoneContainer);
+	}
 }
 
 FTransform FAIPComponentSpacePoseProvider::CalculateComponentSpaceTransform(const FSkeletonPoseBoneIndex SkeletonBoneIdx)
 {
-	const FMeshPoseBoneIndex MeshPoseBoneIndex = AnimInstanceProxy->GetRequiredBones().GetMeshPoseIndexFromSkeletonPoseIndex(SkeletonBoneIdx);
-	return AnimInstanceProxy->GetSkelMeshComponent()->GetBoneTransform(MeshPoseBoneIndex.GetInt(), FTransform::Identity);
+	// NoTe: calling GetBoneTransform on the mesh returns Identity on the first frame of simulation, so we use a different approach
+	// const FMeshPoseBoneIndex MeshPoseBoneIndex = AnimInstanceProxy->GetRequiredBones().GetMeshPoseIndexFromSkeletonPoseIndex(SkeletonBoneIdx);
+	// return AnimInstanceProxy->GetSkelMeshComponent()->GetBoneTransform(MeshPoseBoneIndex.GetInt(), FTransform::Identity);
+
+	if (!ComponentSpacePose.GetPose().IsValid())
+	{
+		return FTransform::Identity;
+	}
+
+	const FBoneContainer& BoneContainer = ComponentSpacePose.GetPose().GetBoneContainer();
+	const FCompactPoseBoneIndex CompactBoneIdx = BoneContainer.GetCompactPoseIndexFromSkeletonPoseIndex(SkeletonBoneIdx);
+	if (CompactBoneIdx.IsValid())
+	{
+		return ComponentSpacePose.GetComponentSpaceTransform(CompactBoneIdx);
+	}
+	
+	// NoTe: this chunk of code is very unlikely to be called, but in case:
+	// @todo: cache any transform outside the domain of ComponentSpacePose if needed
+	// @todo: use the skeletal mesh reference pose instead of the one from the skeleton if needed
+	const USkeleton* Skeleton = BoneContainer.GetSkeletonAsset();
+	const FReferenceSkeleton& ReferenceSkeleton = Skeleton->GetReferenceSkeleton();
+	const int32 ParentIndex = ReferenceSkeleton.GetParentIndex(SkeletonBoneIdx.GetInt());
+	check(ParentIndex >= 0);
+
+	const TArray<FTransform>& RefBonePose = Skeleton->GetReferenceSkeleton().GetRefBonePose();
+	return RefBonePose[SkeletonBoneIdx.GetInt()] * CalculateComponentSpaceTransform(FSkeletonPoseBoneIndex(ParentIndex));
+
 }
 
 const USkeleton* FAIPComponentSpacePoseProvider::GetSkeletonAsset() const
 {
-	return AnimInstanceProxy->GetSkelMeshComponent()->GetSkeletalMeshAsset()->GetSkeleton();
+	// return AnimInstanceProxy->GetSkelMeshComponent()->GetSkeletalMeshAsset()->GetSkeleton();
+	if (ComponentSpacePose.GetPose().IsValid())
+	{
+		return ComponentSpacePose.GetPose().GetBoneContainer().GetSkeletonAsset();
+	}
+	return nullptr;
 }
 
 //////////////////////////////////////////////////////////////////////////
diff --git a/Engine/Plugins/Animation/PoseSearch/Source/Runtime/Public/PoseSearch/PoseSearchHistory.h b/Engine/Plugins/Animation/PoseSearch/Source/Runtime/Public/PoseSearch/PoseSearchHistory.h
--- a/Engine/Plugins/Animation/PoseSearch/Source/Runtime/Public/PoseSearch/PoseSearchHistory.h
+++ b/Engine/Plugins/Animation/PoseSearch/Source/Runtime/Public/PoseSearch/PoseSearchHistory.h
@@ -49,7 +49,7 @@ struct FAIPComponentSpacePoseProvider : public IComponentSpacePoseProvider
 	virtual const USkeleton* GetSkeletonAsset() const override;
 
 private:
-	const FAnimInstanceProxy* AnimInstanceProxy = nullptr;
+	FCSPose<FCompactPose> ComponentSpacePose;
 };
 
 struct FPoseHistoryEntry
