PoseSearch - FPoseSearchTrajectoryData::UpdateData support for zero DeltaTime

diff --git a/Engine/Plugins/Animation/PoseSearch/Source/Runtime/Private/PoseSearchTrajectoryLibrary.cpp b/Engine/Plugins/Animation/PoseSearch/Source/Runtime/Private/PoseSearchTrajectoryLibrary.cpp
--- a/Engine/Plugins/Animation/PoseSearch/Source/Runtime/Private/PoseSearchTrajectoryLibrary.cpp
+++ b/Engine/Plugins/Animation/PoseSearch/Source/Runtime/Private/PoseSearchTrajectoryLibrary.cpp
@@ -25,12 +25,6 @@ bool FPoseSearchTrajectoryData::UpdateData(
 	FDerived& TrajectoryDataDerived,
 	FState& TrajectoryDataState) const
 {
-	// An AnimInstance might call this during an AnimBP recompile with 0 delta time.
-	if (DeltaTime <= 0.f)
-	{
-		return false;
-	}
-
 	const ACharacter* Character = Cast<ACharacter>(Context);
 	if (!Character)
 	{
@@ -82,10 +76,14 @@ bool FPoseSearchTrajectoryData::UpdateData(
 	const float DesiredYawDelta = DesiredControllerYaw - TrajectoryDataState.DesiredControllerYawLastUpdate;
 	TrajectoryDataState.DesiredControllerYawLastUpdate = DesiredControllerYaw;
 
-	TrajectoryDataDerived.ControllerYawRate = FRotator::NormalizeAxis(DesiredYawDelta) * (1.f / DeltaTime);
-	if (MaxControllerYawRate >= 0.f)
+	if (DeltaTime > UE_SMALL_NUMBER)
 	{
-		TrajectoryDataDerived.ControllerYawRate = FMath::Sign(TrajectoryDataDerived.ControllerYawRate) * FMath::Min(FMath::Abs(TrajectoryDataDerived.ControllerYawRate), MaxControllerYawRate);
+		// An AnimInstance might call this during an AnimBP recompile with 0 delta time, so we don't update ControllerYawRate
+		TrajectoryDataDerived.ControllerYawRate = FRotator::NormalizeAxis(DesiredYawDelta) / DeltaTime;
+		if (MaxControllerYawRate >= 0.f)
+		{
+			TrajectoryDataDerived.ControllerYawRate = FMath::Sign(TrajectoryDataDerived.ControllerYawRate) * FMath::Min(FMath::Abs(TrajectoryDataDerived.ControllerYawRate), MaxControllerYawRate);
+		}
 	}
 
 	TrajectoryDataDerived.Position = MeshComp->GetComponentLocation();
