PoseSearch - adding MM support for searching events by tagging animation assets with UAnimNotify_PoseSearchEvent and relying on UPoseSearchDatabase::SearchEvent API

diff --git a/Engine/Plugins/Animation/PoseSearch/Source/Editor/Private/PoseSearchDatabaseEditorCommands.cpp b/Engine/Plugins/Animation/PoseSearch/Source/Editor/Private/PoseSearchDatabaseEditorCommands.cpp
--- a/Engine/Plugins/Animation/PoseSearch/Source/Editor/Private/PoseSearchDatabaseEditorCommands.cpp
+++ b/Engine/Plugins/Animation/PoseSearch/Source/Editor/Private/PoseSearchDatabaseEditorCommands.cpp
@@ -12,6 +12,7 @@ namespace UE::PoseSearch
 		UI_COMMAND(ShowQuantizeAnimationToPoseData, "Quantize Animation To Pose Data", "Animations will not playing continuously, only using the associated pose quantized time", EUserInterfaceActionType::ToggleButton, FInputChord());
 		UI_COMMAND(ShowBones, "Show All Bones", "Debug Draw All Bones", EUserInterfaceActionType::ToggleButton, FInputChord());
 		UI_COMMAND(ShowDisplayBlockTransition, "Display Block Transition", "Show root motion samples (green) with block transition (red)", EUserInterfaceActionType::ToggleButton, FInputChord());
+		UI_COMMAND(ShowDisplayEventData, "Display Events", "Show Events (blue)", EUserInterfaceActionType::ToggleButton, FInputChord());
 	}
 }
 
diff --git a/Engine/Plugins/Animation/PoseSearch/Source/Editor/Private/PoseSearchDatabaseEditorReflection.cpp b/Engine/Plugins/Animation/PoseSearch/Source/Editor/Private/PoseSearchDatabaseEditorReflection.cpp
--- a/Engine/Plugins/Animation/PoseSearch/Source/Editor/Private/PoseSearchDatabaseEditorReflection.cpp
+++ b/Engine/Plugins/Animation/PoseSearch/Source/Editor/Private/PoseSearchDatabaseEditorReflection.cpp
@@ -278,8 +278,9 @@ void UPoseSearchDatabaseStatistics::Initialize(const UPoseSearchDatabase* PoseSe
 
 			const uint32 PoseMetadataBytesSize = SearchIndex.PoseMetadata.GetAllocatedSize();
 			const uint32 AssetsBytesSize = SearchIndex.Assets.GetAllocatedSize();
+			const uint32 EventDataBytesSize = SearchIndex.EventData.GetAllocatedSize();
 			const uint32 OtherBytesSize = SearchIndex.PCAProjectionMatrix.GetAllocatedSize() + SearchIndex.Mean.GetAllocatedSize() + SearchIndex.WeightsSqrt.GetAllocatedSize();
-			const uint32 EstimatedDatabaseBytesSize = ValuesBytesSize + PCAValuesBytesSize + KDTreeBytesSize + VPTreeBytesSize + ValuesVectorToPoseIndexesBytesSize + PCAValuesVectorToPoseIndexesBytesSize + PoseMetadataBytesSize + AssetsBytesSize + OtherBytesSize + SourceAnimAssetsSizeCookedEstimateInBytes;
+			const uint32 EstimatedDatabaseBytesSize = ValuesBytesSize + PCAValuesBytesSize + KDTreeBytesSize + VPTreeBytesSize + ValuesVectorToPoseIndexesBytesSize + PCAValuesVectorToPoseIndexesBytesSize + PoseMetadataBytesSize + AssetsBytesSize + EventDataBytesSize + OtherBytesSize + SourceAnimAssetsSizeCookedEstimateInBytes;
 				
 			ValuesSize = FText::AsMemory(ValuesBytesSize);
 			PCAValuesSize = FText::AsMemory(PCAValuesBytesSize);
@@ -287,6 +288,7 @@ void UPoseSearchDatabaseStatistics::Initialize(const UPoseSearchDatabase* PoseSe
 			VPTreeSize = FText::AsMemory(VPTreeBytesSize);
 			PoseMetadataSize = FText::AsMemory(PoseMetadataBytesSize);
 			AssetsSize = FText::AsMemory(AssetsBytesSize);
+			EventDataSize = FText::AsMemory(EventDataBytesSize);
 			EstimatedDatabaseSize = FText::AsMemory(EstimatedDatabaseBytesSize);
 			SourceAnimAssetsSizeCookedEstimate = FText::AsMemory(SourceAnimAssetsSizeCookedEstimateInBytes);
 		}
diff --git a/Engine/Plugins/Animation/PoseSearch/Source/Editor/Private/PoseSearchDatabaseViewModel.cpp b/Engine/Plugins/Animation/PoseSearch/Source/Editor/Private/PoseSearchDatabaseViewModel.cpp
--- a/Engine/Plugins/Animation/PoseSearch/Source/Editor/Private/PoseSearchDatabaseViewModel.cpp
+++ b/Engine/Plugins/Animation/PoseSearch/Source/Editor/Private/PoseSearchDatabaseViewModel.cpp
@@ -238,7 +238,7 @@ void FDatabasePreviewActor::Destroy()
 	}
 }
 
-bool FDatabasePreviewActor::DrawPreviewActors(TConstArrayView<FDatabasePreviewActor> PreviewActors, const UPoseSearchDatabase* PoseSearchDatabase, bool bDisplayRootMotionSpeed, bool bDisplayBlockTransition, TConstArrayView<float> QueryVector)
+bool FDatabasePreviewActor::DrawPreviewActors(TConstArrayView<FDatabasePreviewActor> PreviewActors, const UPoseSearchDatabase* PoseSearchDatabase, bool bDisplayRootMotionSpeed, bool bDisplayBlockTransition, bool bDisplayEventData, TConstArrayView<float> QueryVector)
 {
 	using namespace UE::PoseSearch;
 
@@ -403,6 +403,31 @@ bool FDatabasePreviewActor::DrawPreviewActors(TConstArrayView<FDatabasePreviewAc
 			}
 		}
 
+		if (bDisplayEventData && !SearchIndex.EventData.GetData().IsEmpty())
+		{
+			const int NumPoses = IndexAsset.GetNumPoses();
+			if (NumPoses == SamplesNum)
+			{
+				TSet<int32, DefaultKeyFuncs<int32>, TInlineSetAllocator<256>> AllEventEventDataPoseIndexes;
+				for (const FEventData::FTagToPoseIndexes& TagToPoseIndexes : SearchIndex.EventData.GetData())
+				{
+					for (int32 PoseIdx : TagToPoseIndexes.Value)
+					{
+						AllEventEventDataPoseIndexes.Add(PoseIdx);
+					}
+				}
+				
+				for (int32 Index = 0; Index < SamplesNum; ++Index)
+				{
+					const int32 IndexAssetPoseIdx = Index + IndexAsset.GetFirstPoseIdx();
+					if (AllEventEventDataPoseIndexes.Find(IndexAssetPoseIdx))
+					{
+						DrawParams.DrawPoint(PreviewActor.Trajectory.Samples[Index].Position, FColor::Blue, 8.f);
+					}
+				}
+			}
+		}
+
 #if ENABLE_ANIM_DEBUG
 		const float DebugDrawSamplerSize = GVarDatabasePreviewDebugDrawSamplerSize;
 		if (DebugDrawSamplerSize > UE_KINDA_SMALL_NUMBER)
@@ -581,7 +606,7 @@ void FDatabaseViewModel::Tick(float DeltaSeconds)
 			bool bShouldDrawQueryVector = ShouldDrawQueryVector();
 			for (TArray<FDatabasePreviewActor>& PreviewActorGroup : PreviewActors)
 			{
-				bShouldDrawQueryVector &= !FDatabasePreviewActor::DrawPreviewActors(PreviewActorGroup, Database, bDisplayRootMotionSpeed, bDisplayBlockTransition, bShouldDrawQueryVector ? GetQueryVector() : TConstArrayView<float>());
+				bShouldDrawQueryVector &= !FDatabasePreviewActor::DrawPreviewActors(PreviewActorGroup, Database, bDisplayRootMotionSpeed, bDisplayBlockTransition, bDisplayEventData, bShouldDrawQueryVector ? GetQueryVector() : TConstArrayView<float>());
 			}
 		}
 	}
diff --git a/Engine/Plugins/Animation/PoseSearch/Source/Editor/Private/SPoseSearchDatabaseViewport.cpp b/Engine/Plugins/Animation/PoseSearch/Source/Editor/Private/SPoseSearchDatabaseViewport.cpp
--- a/Engine/Plugins/Animation/PoseSearch/Source/Editor/Private/SPoseSearchDatabaseViewport.cpp
+++ b/Engine/Plugins/Animation/PoseSearch/Source/Editor/Private/SPoseSearchDatabaseViewport.cpp
@@ -73,6 +73,17 @@ namespace UE::PoseSearch
 		return AssetEditorPtr.Pin()->GetViewModelSharedPtr().ToSharedRef()->IsDisplayBlockTransitionChecked();
 	}
 
+	void SDatabaseViewport::OnToggleDisplayEventData()
+	{
+		AssetEditorPtr.Pin()->GetViewModelSharedPtr().ToSharedRef()->ToggleDisplayEventData();
+		Invalidate();
+	}
+
+	bool SDatabaseViewport::IsDisplayEventDataChecked() const
+	{
+		return AssetEditorPtr.Pin()->GetViewModelSharedPtr().ToSharedRef()->IsDisplayEventDataChecked();
+	}
+
 	void SDatabaseViewport::BindCommands()
 	{
 		SEditorViewport::BindCommands();
@@ -102,6 +113,12 @@ namespace UE::PoseSearch
 			FExecuteAction::CreateSP(this, &SDatabaseViewport::OnToggleDisplayBlockTransition),
 			FCanExecuteAction(),
 			FIsActionChecked::CreateSP(this, &SDatabaseViewport::IsDisplayBlockTransitionChecked));
+
+		CommandList->MapAction(
+			Commands.ShowDisplayEventData,
+			FExecuteAction::CreateSP(this, &SDatabaseViewport::OnToggleDisplayEventData),
+			FCanExecuteAction(),
+			FIsActionChecked::CreateSP(this, &SDatabaseViewport::IsDisplayEventDataChecked));
 	}
 
 	TSharedRef<FEditorViewportClient> SDatabaseViewport::MakeEditorViewportClient()
diff --git a/Engine/Plugins/Animation/PoseSearch/Source/Editor/Private/SPoseSearchDatabaseViewportToolbar.cpp b/Engine/Plugins/Animation/PoseSearch/Source/Editor/Private/SPoseSearchDatabaseViewportToolbar.cpp
--- a/Engine/Plugins/Animation/PoseSearch/Source/Editor/Private/SPoseSearchDatabaseViewportToolbar.cpp
+++ b/Engine/Plugins/Animation/PoseSearch/Source/Editor/Private/SPoseSearchDatabaseViewportToolbar.cpp
@@ -35,6 +35,7 @@ namespace UE::PoseSearch
 			ShowMenuBuilder.AddMenuEntry(Commands.ShowQuantizeAnimationToPoseData);
 			ShowMenuBuilder.AddMenuEntry(Commands.ShowBones);
 			ShowMenuBuilder.AddMenuEntry(Commands.ShowDisplayBlockTransition);
+			ShowMenuBuilder.AddMenuEntry(Commands.ShowDisplayEventData);
 			ShowMenuBuilder.EndSection();
 		}
 
diff --git a/Engine/Plugins/Animation/PoseSearch/Source/Editor/Public/PoseSearchDatabaseEditorCommands.h b/Engine/Plugins/Animation/PoseSearch/Source/Editor/Public/PoseSearchDatabaseEditorCommands.h
--- a/Engine/Plugins/Animation/PoseSearch/Source/Editor/Public/PoseSearchDatabaseEditorCommands.h
+++ b/Engine/Plugins/Animation/PoseSearch/Source/Editor/Public/PoseSearchDatabaseEditorCommands.h
@@ -28,6 +28,8 @@ namespace UE::PoseSearch
 		TSharedPtr<FUICommandInfo> ShowQuantizeAnimationToPoseData;
 		TSharedPtr<FUICommandInfo> ShowBones;
 		TSharedPtr<FUICommandInfo> ShowDisplayBlockTransition;
+		// Experimental, this feature might be removed without warning, not for production use
+		TSharedPtr<FUICommandInfo> ShowDisplayEventData;
 	};
 }
 
diff --git a/Engine/Plugins/Animation/PoseSearch/Source/Editor/Public/PoseSearchDatabaseEditorReflection.h b/Engine/Plugins/Animation/PoseSearch/Source/Editor/Public/PoseSearchDatabaseEditorReflection.h
--- a/Engine/Plugins/Animation/PoseSearch/Source/Editor/Public/PoseSearchDatabaseEditorReflection.h
+++ b/Engine/Plugins/Animation/PoseSearch/Source/Editor/Public/PoseSearchDatabaseEditorReflection.h
@@ -281,6 +281,12 @@ public:
 	UPROPERTY(VisibleAnywhere, Category = "Memory Information")
 	FText AssetsSize;
 
+	// Experimental, this feature might be removed without warning, not for production use
+	// partial memory size used for events.
+	UPROPERTY(VisibleAnywhere, Category = "Memory Information")
+	FText EventDataSize;
+
+	// Experimental, this feature might be removed without warning, not for production use
 	// partial memory size used by the cooked animation assets.
 	UPROPERTY(VisibleAnywhere, Category = "Memory Information")
  	FText SourceAnimAssetsSizeCookedEstimate;
diff --git a/Engine/Plugins/Animation/PoseSearch/Source/Editor/Public/PoseSearchDatabaseViewModel.h b/Engine/Plugins/Animation/PoseSearch/Source/Editor/Public/PoseSearchDatabaseViewModel.h
--- a/Engine/Plugins/Animation/PoseSearch/Source/Editor/Public/PoseSearchDatabaseViewModel.h
+++ b/Engine/Plugins/Animation/PoseSearch/Source/Editor/Public/PoseSearchDatabaseViewModel.h
@@ -34,7 +34,7 @@ namespace UE::PoseSearch
 	public:
 		bool SpawnPreviewActor(UWorld* World, const UPoseSearchDatabase* PoseSearchDatabase, int32 IndexAssetIdx, const FRole& Role, const FTransform& SamplerRootTransformOrigin, int32 PoseIdxForTimeOffset = INDEX_NONE);
 		void UpdatePreviewActor(const UPoseSearchDatabase* PoseSearchDatabase, float PlayTime, bool bQuantizeAnimationToPoseData);
-		static bool DrawPreviewActors(TConstArrayView<FDatabasePreviewActor> PreviewActors, const UPoseSearchDatabase* PoseSearchDatabase, bool bDisplayRootMotionSpeed, bool bDisplayBlockTransition, TConstArrayView<float> QueryVector);
+		static bool DrawPreviewActors(TConstArrayView<FDatabasePreviewActor> PreviewActors, const UPoseSearchDatabase* PoseSearchDatabase, bool bDisplayRootMotionSpeed, bool bDisplayBlockTransition, bool bDisplayEventData, TConstArrayView<float> QueryVector);
 
 		void Destroy();
 
@@ -106,6 +106,9 @@ namespace UE::PoseSearch
 		void ToggleDisplayBlockTransition() { bDisplayBlockTransition = !bDisplayBlockTransition; }
 		bool IsDisplayBlockTransitionChecked() const { return bDisplayBlockTransition; }
 
+		void ToggleDisplayEventData() { bDisplayEventData = !bDisplayEventData; }
+		bool IsDisplayEventDataChecked() const { return bDisplayEventData; }
+
 		void AddSequenceToDatabase(UAnimSequence* AnimSequence);
 		void AddBlendSpaceToDatabase(UBlendSpace* BlendSpace);
 		void AddAnimCompositeToDatabase(UAnimComposite* AnimComposite);
@@ -174,6 +177,7 @@ namespace UE::PoseSearch
 		bool bQuantizeAnimationToPoseData = false;
 		bool bShowBones = false;
 		bool bDisplayBlockTransition = false;
+		bool bDisplayEventData = false;
 
 		int32 SelectedActorIndexAssetIndex = INDEX_NONE;
 	};
diff --git a/Engine/Plugins/Animation/PoseSearch/Source/Editor/Public/SPoseSearchDatabaseViewport.h b/Engine/Plugins/Animation/PoseSearch/Source/Editor/Public/SPoseSearchDatabaseViewport.h
--- a/Engine/Plugins/Animation/PoseSearch/Source/Editor/Public/SPoseSearchDatabaseViewport.h
+++ b/Engine/Plugins/Animation/PoseSearch/Source/Editor/Public/SPoseSearchDatabaseViewport.h
@@ -60,6 +60,10 @@ namespace UE::PoseSearch
 		bool IsShowBonesChecked() const;
 		void OnToggleDisplayBlockTransition();
 		bool IsDisplayBlockTransitionChecked() const;
+		// Experimental, this feature might be removed without warning, not for production use
+		void OnToggleDisplayEventData();
+		// Experimental, this feature might be removed without warning, not for production use
+		bool IsDisplayEventDataChecked() const;
 
 		/** The viewport toolbar */
 		TSharedPtr<SPoseSearchDatabaseViewportToolBar> ViewportToolbar;
diff --git a/Engine/Plugins/Animation/PoseSearch/Source/Runtime/PoseSearch.Build.cs b/Engine/Plugins/Animation/PoseSearch/Source/Runtime/PoseSearch.Build.cs
--- a/Engine/Plugins/Animation/PoseSearch/Source/Runtime/PoseSearch.Build.cs
+++ b/Engine/Plugins/Animation/PoseSearch/Source/Runtime/PoseSearch.Build.cs
@@ -29,6 +29,7 @@ public class PoseSearch : ModuleRules
 
 		PrivateDependencyModuleNames.AddRange(
 		new string[]{
+			"GameplayTags",
 			"RewindDebuggerRuntimeInterface",
 		}
 		);
diff --git a/Engine/Plugins/Animation/PoseSearch/Source/Runtime/Private/AnimNode_MotionMatching.cpp b/Engine/Plugins/Animation/PoseSearch/Source/Runtime/Private/AnimNode_MotionMatching.cpp
--- a/Engine/Plugins/Animation/PoseSearch/Source/Runtime/Private/AnimNode_MotionMatching.cpp
+++ b/Engine/Plugins/Animation/PoseSearch/Source/Runtime/Private/AnimNode_MotionMatching.cpp
@@ -207,7 +207,11 @@ void FAnimNode_MotionMatching::UpdateAssetPlayer(const FAnimationUpdateContext&
 #if ENABLE_ANIM_DEBUG
 			, GVarAnimNodeMotionMatchingDrawQuery
 			, GVarAnimNodeMotionMatchingDrawCurResult
+#else // ENABLE_ANIM_DEBUG
+			, false
+			, false
 #endif // ENABLE_ANIM_DEBUG
+			, EventToSearch
 		);
 	}
 	else
diff --git a/Engine/Plugins/Animation/PoseSearch/Source/Runtime/Private/AnimNode_MotionMatchingInteraction.cpp b/Engine/Plugins/Animation/PoseSearch/Source/Runtime/Private/AnimNode_MotionMatchingInteraction.cpp
--- a/Engine/Plugins/Animation/PoseSearch/Source/Runtime/Private/AnimNode_MotionMatchingInteraction.cpp
+++ b/Engine/Plugins/Animation/PoseSearch/Source/Runtime/Private/AnimNode_MotionMatchingInteraction.cpp
@@ -260,9 +260,10 @@ void FAnimNode_MotionMatchingInteraction::Evaluate_AnyThread(FPoseContext& Outpu
 	TSet<FName, DefaultKeyFuncs<FName>, TInlineSetAllocator<16, TMemStackSetAllocator<>>> GoalNames;
 	if (bUsePoseSearchIKWindowNotifyStates)
 	{
-		Sampler.ExtractPoseSearchNotifyStates(AccumulatedTime, [&GoalNames](const UAnimNotifyState_PoseSearchBase* PoseSearchBase)
+		FAnimNotifyContext PreAllocatedNotifyContext;
+		Sampler.ExtractAnimNotifyStates(AccumulatedTime, PreAllocatedNotifyContext, [&GoalNames](const UAnimNotifyState* AnimNotifyState)
 			{
-				if (const UAnimNotifyState_PoseSearchIKWindow* PoseSearchIKWindow = Cast<const UAnimNotifyState_PoseSearchIKWindow>(PoseSearchBase))
+				if (const UAnimNotifyState_PoseSearchIKWindow* PoseSearchIKWindow = Cast<const UAnimNotifyState_PoseSearchIKWindow>(AnimNotifyState))
 				{
 					for (FName GoalName : PoseSearchIKWindow->GoalNames)
 					{
diff --git a/Engine/Plugins/Animation/PoseSearch/Source/Runtime/Private/AnimNode_PoseSearchHistoryCollector.cpp b/Engine/Plugins/Animation/PoseSearch/Source/Runtime/Private/AnimNode_PoseSearchHistoryCollector.cpp
--- a/Engine/Plugins/Animation/PoseSearch/Source/Runtime/Private/AnimNode_PoseSearchHistoryCollector.cpp
+++ b/Engine/Plugins/Animation/PoseSearch/Source/Runtime/Private/AnimNode_PoseSearchHistoryCollector.cpp
@@ -91,9 +91,12 @@ void FAnimNode_PoseSearchHistoryCollector_Base::Initialize_AnyThread(const FAnim
 
 	FMemMark Mark(FMemStack::Get());
 	UE::PoseSearch::FAIPComponentSpacePoseProvider ComponentSpacePoseProvider(Context.AnimInstanceProxy);
-	PoseHistory.EvaluateComponentSpace_AnyThread(0.f, ComponentSpacePoseProvider, bStoreScales,
-		RootBoneRecoveryTime, RootBoneTranslationRecoveryRatio, RootBoneRotationRecoveryRatio, true, true, 
-		GetRequiredBones(Context.AnimInstanceProxy), FBlendedCurve(), MakeConstArrayView(CollectedCurves));
+	if (ComponentSpacePoseProvider.GetSkeletonAsset())
+	{
+		PoseHistory.EvaluateComponentSpace_AnyThread(0.f, ComponentSpacePoseProvider, bStoreScales,
+			RootBoneRecoveryTime, RootBoneTranslationRecoveryRatio, RootBoneRotationRecoveryRatio, true, true,
+			GetRequiredBones(Context.AnimInstanceProxy), FBlendedCurve(), MakeConstArrayView(CollectedCurves));
+	}
 }
 
 TArray<FBoneIndexType> FAnimNode_PoseSearchHistoryCollector_Base::GetRequiredBones(const FAnimInstanceProxy* AnimInstanceProxy) const
diff --git a/Engine/Plugins/Animation/PoseSearch/Source/Runtime/Private/PoseSearchAssetIndexer.cpp b/Engine/Plugins/Animation/PoseSearch/Source/Runtime/Private/PoseSearchAssetIndexer.cpp
--- a/Engine/Plugins/Animation/PoseSearch/Source/Runtime/Private/PoseSearchAssetIndexer.cpp
+++ b/Engine/Plugins/Animation/PoseSearch/Source/Runtime/Private/PoseSearchAssetIndexer.cpp
@@ -136,11 +136,11 @@ bool FAnimationAssetSamplers::IsLoopable() const
 	return true;
 }
 
-void FAnimationAssetSamplers::ExtractPoseSearchNotifyStates(float Time, const TFunction<bool(UAnimNotifyState_PoseSearchBase*)>& ProcessPoseSearchBase) const
+void FAnimationAssetSamplers::ExtractAnimNotifyStates(float Time, FAnimNotifyContext& PreAllocatedNotifyContext, const TFunction<bool(UAnimNotifyState*)>& ProcessAnimNotifyState) const
 {
 	for (const FAnimationAssetSampler* Sampler : AnimationAssetSamplers)
 	{
-		Sampler->ExtractPoseSearchNotifyStates(Time, ProcessPoseSearchBase);
+		Sampler->ExtractAnimNotifyStates(Time, PreAllocatedNotifyContext, ProcessAnimNotifyState);
 	}
 }
 
@@ -241,23 +241,26 @@ void FAssetIndexer::Process(int32 AssetIdx)
 
 	// Generate pose metadata
 	const float PlayLength = GetPlayLength();
+	FAnimNotifyContext PreAllocatedNotifyContext;
+
+	// @todo: optimize this code, by extracting ALL the notify states, and then perform time overlap with the poses
 	for (int32 SampleIdx = GetBeginSampleIdx(); SampleIdx != GetEndSampleIdx(); ++SampleIdx)
 	{
 		const float SampleTime = FMath::Min(CalculateSampleTime(SampleIdx), PlayLength);
 		float CostAddend = SamplingContext.BaseCostBias;
 		bool bBlockTransition = false;
 
-		AssetSamplers.ExtractPoseSearchNotifyStates(SampleTime, [&bBlockTransition, &CostAddend](const UAnimNotifyState_PoseSearchBase* PoseSearchBase)
+		AssetSamplers.ExtractAnimNotifyStates(SampleTime, PreAllocatedNotifyContext, [&bBlockTransition, &CostAddend](const UAnimNotifyState* AnimNotifyState)
 			{
-				if (PoseSearchBase->GetClass()->IsChildOf<UAnimNotifyState_PoseSearchBlockTransition>())
+				if (AnimNotifyState->GetClass()->IsChildOf<UAnimNotifyState_PoseSearchBlockTransition>())
 				{
 					bBlockTransition = true;
 				}
-				else if (PoseSearchBase->GetClass()->IsChildOf<UAnimNotifyState_PoseSearchModifyCost>())
+				else if (AnimNotifyState->GetClass()->IsChildOf<UAnimNotifyState_PoseSearchModifyCost>())
 				{
-					const UAnimNotifyState_PoseSearchModifyCost* ModifyCostNotify = Cast<const UAnimNotifyState_PoseSearchModifyCost>(PoseSearchBase);
-					check(ModifyCostNotify);
-					CostAddend = ModifyCostNotify->CostAddend;
+					const UAnimNotifyState_PoseSearchModifyCost* ModifyCostNotifyState = Cast<const UAnimNotifyState_PoseSearchModifyCost>(AnimNotifyState);
+					check(ModifyCostNotifyState);
+					CostAddend = ModifyCostNotifyState->CostAddend;
 				}
 				return true;
 			});
@@ -274,6 +277,25 @@ void FAssetIndexer::Process(int32 AssetIdx)
 		PoseMetadata[VectorIdx] = FPoseMetadata(ValueOffset, AssetIdx, bBlockTransition, CostAddend);
 	}
 
+	
+	AssetSamplers.ProcessAllAnimNotifyEvents([this](const TConstArrayView<FAnimNotifyEvent> AnimNotifyEvents)
+		{
+			for (const FAnimNotifyEvent& AnimNotifyEvent : AnimNotifyEvents)
+			{
+				if (AnimNotifyEvent.Notify && AnimNotifyEvent.Notify->GetClass()->IsChildOf<UAnimNotify_PoseSearchEvent>())
+				{
+					const UAnimNotify_PoseSearchEvent* EventNotify = Cast<const UAnimNotify_PoseSearchEvent>(AnimNotifyEvent.Notify);
+					check(EventNotify);
+					if (EventNotify->EventTag.IsValid())
+					{
+						const int32 PoseIdx = SearchIndexAsset.GetPoseIndexFromTime(AnimNotifyEvent.GetTime(), Schema.SampleRate);
+						EventDataCollector.Emplace(EventNotify->EventTag, PoseIdx);
+					}
+				}
+			}
+			return true;
+		});
+
 	// Generate pose features data
 	if (Schema.SchemaCardinality > 0)
 	{
diff --git a/Engine/Plugins/Animation/PoseSearch/Source/Runtime/Private/PoseSearchAssetSampler.cpp b/Engine/Plugins/Animation/PoseSearch/Source/Runtime/Private/PoseSearchAssetSampler.cpp
--- a/Engine/Plugins/Animation/PoseSearch/Source/Runtime/Private/PoseSearchAssetSampler.cpp
+++ b/Engine/Plugins/Animation/PoseSearch/Source/Runtime/Private/PoseSearchAssetSampler.cpp
@@ -624,9 +624,26 @@ void FAnimationAssetSampler::Process()
 }
 
 void FAnimationAssetSampler::ExtractPoseSearchNotifyStates(float Time, const TFunction<bool(UAnimNotifyState_PoseSearchBase*)>& ProcessPoseSearchBase) const
+{
+	FAnimNotifyContext PreAllocatedNotifyContext;
+	ExtractAnimNotifyStates(Time, PreAllocatedNotifyContext, [&ProcessPoseSearchBase](UAnimNotifyState* AnimNotifyState)
+		{
+			if (UAnimNotifyState_PoseSearchBase* PoseSearchAnimNotify = Cast<UAnimNotifyState_PoseSearchBase>(AnimNotifyState))
+			{
+				return ProcessPoseSearchBase(PoseSearchAnimNotify);
+			}
+			return true;
+		});
+}
+
+void FAnimationAssetSampler::ExtractAnimNotifyStates(float Time, FAnimNotifyContext& PreAllocatedNotifyContext, const TFunction<bool(UAnimNotifyState*)>& ProcessAnimNotifyState) const
 {
 	float SampleTime = Time;
-	FAnimNotifyContext NotifyContext;
+
+	// resetting PreAllocatedNotifyContext without losing ActiveNotifies allocation 
+	PreAllocatedNotifyContext.TickRecord = nullptr;
+	PreAllocatedNotifyContext.ActiveNotifies.Reset();
+
 	if (const UBlendSpace* BlendSpace = Cast<UBlendSpace>(AnimationAssetPtr.Get()))
 	{
 		if (BlendSpace->NotifyTriggerMode == ENotifyTriggerMode::HighestWeightedAnimation)
@@ -649,7 +666,7 @@ void FAnimationAssetSampler::ExtractPoseSearchNotifyStates(float Time, const TFu
 
 					// Get notifies for highest weighted
 					const float ExtractionStartTime = FMath::Min(SampleTime, BlendSample.Animation->GetPlayLength()) - (ExtractionInterval * 0.5f);
-					BlendSample.Animation->GetAnimNotifies(ExtractionStartTime, ExtractionInterval, NotifyContext);
+					BlendSample.Animation->GetAnimNotifies(ExtractionStartTime, ExtractionInterval, PreAllocatedNotifyContext);
 				}
 			}
 		}
@@ -662,7 +679,7 @@ void FAnimationAssetSampler::ExtractPoseSearchNotifyStates(float Time, const TFu
 	{
 		// getting pose search notifies in an interval of size ExtractionInterval, centered on Time
 		const float ExtractionStartTime = FMath::Min(Time, SequenceBase->GetPlayLength()) - (ExtractionInterval * 0.5f);
-		SequenceBase->GetAnimNotifies(ExtractionStartTime, ExtractionInterval, NotifyContext);
+		SequenceBase->GetAnimNotifies(ExtractionStartTime, ExtractionInterval, PreAllocatedNotifyContext);
 	}
 	else
 	{
@@ -670,16 +687,13 @@ void FAnimationAssetSampler::ExtractPoseSearchNotifyStates(float Time, const TFu
 	}
 
 	// check which notifies actually overlap Time and are of the right base type
-	for (const FAnimNotifyEventReference& EventReference : NotifyContext.ActiveNotifies)
+	for (const FAnimNotifyEventReference& EventReference : PreAllocatedNotifyContext.ActiveNotifies)
 	{
 		if (const FAnimNotifyEvent* NotifyEvent = EventReference.GetNotify())
 		{
-			if (UAnimNotifyState_PoseSearchBase* PoseSearchAnimNotify = Cast<UAnimNotifyState_PoseSearchBase>(NotifyEvent->NotifyStateClass))
+			if (!ProcessAnimNotifyState(NotifyEvent->NotifyStateClass))
 			{
-				if (!ProcessPoseSearchBase(PoseSearchAnimNotify))
-				{
-					break;
-				}
+				break;
 			}
 		}
 	}
diff --git a/Engine/Plugins/Animation/PoseSearch/Source/Runtime/Private/PoseSearchContext.cpp b/Engine/Plugins/Animation/PoseSearch/Source/Runtime/Private/PoseSearchContext.cpp
--- a/Engine/Plugins/Animation/PoseSearch/Source/Runtime/Private/PoseSearchContext.cpp
+++ b/Engine/Plugins/Animation/PoseSearch/Source/Runtime/Private/PoseSearchContext.cpp
@@ -651,15 +651,22 @@ FCachedQuery::FCachedQuery(const UPoseSearchSchema* InSchema)
 // FSearchContext
 FSearchContext::FSearchContext(float InDesiredPermutationTimeOffset, const FPoseIndicesHistory* InPoseIndicesHistory,
 	const FSearchResult& InCurrentResult, const FFloatInterval& InPoseJumpThresholdTime, bool bInUseCachedChannelData)
+: FSearchContext(InDesiredPermutationTimeOffset, InPoseIndicesHistory, InCurrentResult, InPoseJumpThresholdTime, FPoseSearchEvent())
+{
+}
+
+FSearchContext::FSearchContext(float InDesiredPermutationTimeOffset, const FPoseIndicesHistory* InPoseIndicesHistory,
+	const FSearchResult& InCurrentResult, const FFloatInterval& InPoseJumpThresholdTime, const FPoseSearchEvent& InEventToSearch)
 : AnimContexts()
 , PoseHistories()
 , RoleToIndex()
 , AssetsToConsider()
+, EventToSearch(InEventToSearch)
 , DesiredPermutationTimeOffset(InDesiredPermutationTimeOffset)
 , PoseIndicesHistory(InPoseIndicesHistory)
 , CurrentResult(InCurrentResult)
 , PoseJumpThresholdTime(InPoseJumpThresholdTime)
-, bUseCachedChannelData(bInUseCachedChannelData)
+, bUseCachedChannelData(false)
 {
 	UpdateCurrentResultPoseVector();
 
diff --git a/Engine/Plugins/Animation/PoseSearch/Source/Runtime/Private/PoseSearchDatabase.cpp b/Engine/Plugins/Animation/PoseSearch/Source/Runtime/Private/PoseSearchDatabase.cpp
--- a/Engine/Plugins/Animation/PoseSearch/Source/Runtime/Private/PoseSearchDatabase.cpp
+++ b/Engine/Plugins/Animation/PoseSearch/Source/Runtime/Private/PoseSearchDatabase.cpp
@@ -31,9 +31,12 @@ DECLARE_STATS_GROUP(TEXT("PoseSearch"), STATGROUP_PoseSearch, STATCAT_Advanced);
 DECLARE_CYCLE_STAT_EXTERN(TEXT("Search Brute Force"), STAT_PoseSearch_BruteForce, STATGROUP_PoseSearch, );
 DECLARE_CYCLE_STAT_EXTERN(TEXT("Search PCA/KNN"), STAT_PoseSearch_PCAKNN, STATGROUP_PoseSearch, );
 DECLARE_CYCLE_STAT_EXTERN(TEXT("Search VPTree"), STAT_PoseSearch_VPTree, STATGROUP_PoseSearch, );
+DECLARE_CYCLE_STAT_EXTERN(TEXT("Search Event"), STAT_PoseSearch_Event, STATGROUP_PoseSearch, );
 DEFINE_STAT(STAT_PoseSearch_BruteForce);
 DEFINE_STAT(STAT_PoseSearch_PCAKNN);
 DEFINE_STAT(STAT_PoseSearch_VPTree);
+DEFINE_STAT(STAT_PoseSearch_Event);
+
 
 namespace UE::PoseSearch
 {
@@ -1174,45 +1177,41 @@ UE::PoseSearch::FSearchResult UPoseSearchDatabase::Search(UE::PoseSearch::FSearc
 		SearchContext.SetAssetsToConsider(ChooserAssetToConsider);
 	}
 
-	switch (PoseSearchMode)
+	if (SearchContext.GetEventToSearch().IsValid())
+	{
+		Result = SearchEvent(SearchContext);
+	}
+	else
 	{
-	case EPoseSearchMode::BruteForce:
+		switch (PoseSearchMode)
 		{
+		case EPoseSearchMode::BruteForce:
 			Result = SearchBruteForce(SearchContext);
-
-#if WITH_EDITOR && ENABLE_ANIM_DEBUG && UE_POSE_SEARCH_TRACE_ENABLED
-			if (GVarMotionMatchCompareAgainstBruteForce)
-			{
-				Result.BruteForcePoseCost = Result.PoseCost;
-			}
-#endif // WITH_EDITOR && ENABLE_ANIM_DEBUG && UE_POSE_SEARCH_TRACE_ENABLED
 			break;
-		}
-	case EPoseSearchMode::VPTree:
-		{
+		case EPoseSearchMode::VPTree:
 			Result = SearchVPTree(SearchContext);
-
-#if WITH_EDITOR && ENABLE_ANIM_DEBUG && UE_POSE_SEARCH_TRACE_ENABLED
-			if (GVarMotionMatchCompareAgainstBruteForce)
-			{
-				Result.BruteForcePoseCost = SearchBruteForce(SearchContext).PoseCost;
-			}
-#endif // WITH_EDITOR && ENABLE_ANIM_DEBUG && UE_POSE_SEARCH_TRACE_ENABLED
 			break;
-		}
-	case EPoseSearchMode::PCAKDTree:
-		{
+		case EPoseSearchMode::PCAKDTree:
 			Result = SearchPCAKDTree(SearchContext);
+			break;
+		case EPoseSearchMode::EventOnly:
+			break;
+		}
+	}
 
 #if WITH_EDITOR && ENABLE_ANIM_DEBUG && UE_POSE_SEARCH_TRACE_ENABLED
-			if (GVarMotionMatchCompareAgainstBruteForce)
-			{
-				Result.BruteForcePoseCost = SearchBruteForce(SearchContext).PoseCost;
-			}
-#endif // WITH_EDITOR && ENABLE_ANIM_DEBUG && UE_POSE_SEARCH_TRACE_ENABLED
-			break;
+	if (GVarMotionMatchCompareAgainstBruteForce)
+	{
+		if (PoseSearchMode == EPoseSearchMode::PCAKDTree || PoseSearchMode == EPoseSearchMode::VPTree)
+		{
+			Result.BruteForcePoseCost = SearchBruteForce(SearchContext).PoseCost;
+		}
+		else
+		{
+			Result.BruteForcePoseCost = Result.PoseCost;
 		}
 	}
+#endif // WITH_EDITOR && ENABLE_ANIM_DEBUG && UE_POSE_SEARCH_TRACE_ENABLED
 
 #if UE_POSE_SEARCH_TRACE_ENABLED
 	// in case we skipped the search, or we didn't find any candidates we still have to track we requested to evaluate this database, so we keep track of this
@@ -1394,6 +1393,7 @@ UE::PoseSearch::FSearchResult UPoseSearchDatabase::SearchContinuingPose(UE::Pose
 
 	float UpdatedContinuingPoseCostBias = ContinuingPoseCostBias;
 	const float SampleTime = GetRealAssetTime(PoseIdx);
+	FAnimNotifyContext PreAllocatedNotifyContext;
 	for (int32 RoleIndex = 0; RoleIndex < DatabaseAnimationAssetBase->GetNumRoles(); ++RoleIndex)
 	{
 		if (const UAnimationAsset* AnimationAsset = DatabaseAnimationAssetBase->GetAnimationAssetForRole(DatabaseAnimationAssetBase->GetRole(RoleIndex)))
@@ -1402,9 +1402,9 @@ UE::PoseSearch::FSearchResult UPoseSearchDatabase::SearchContinuingPose(UE::Pose
 			const FAnimationAssetSampler SequenceBaseSampler(AnimationAsset, FTransform::Identity, SearchIndexAsset.GetBlendParameters(), FAnimationAssetSampler::DefaultRootTransformSamplingRate, false);
 
 			bool bDone = false;
-			SequenceBaseSampler.ExtractPoseSearchNotifyStates(SampleTime, [&UpdatedContinuingPoseCostBias, &bDone](const UAnimNotifyState_PoseSearchBase* PoseSearchBase)
+			SequenceBaseSampler.ExtractAnimNotifyStates(SampleTime, PreAllocatedNotifyContext, [&UpdatedContinuingPoseCostBias, &bDone](const UAnimNotifyState* AnimNotifyState)
 				{
-					if (const UAnimNotifyState_PoseSearchOverrideContinuingPoseCostBias* NotifyStateContinuingPoseCostBias = Cast<const UAnimNotifyState_PoseSearchOverrideContinuingPoseCostBias>(PoseSearchBase))
+					if (const UAnimNotifyState_PoseSearchOverrideContinuingPoseCostBias* NotifyStateContinuingPoseCostBias = Cast<const UAnimNotifyState_PoseSearchOverrideContinuingPoseCostBias>(AnimNotifyState))
 					{
 						UpdatedContinuingPoseCostBias = NotifyStateContinuingPoseCostBias->CostAddend;
 						bDone = true;
@@ -2004,4 +2004,83 @@ UE::PoseSearch::FSearchResult UPoseSearchDatabase::SearchBruteForce(UE::PoseSear
 	}
 
 	return Result;
-}
\ No newline at end of file
+}
+
+UE::PoseSearch::FSearchResult UPoseSearchDatabase::SearchEvent(UE::PoseSearch::FSearchContext& SearchContext) const
+{
+	SCOPE_CYCLE_COUNTER(STAT_PoseSearch_Event);
+
+	using namespace UE::PoseSearch;
+
+	FSearchResult Result;
+
+	const FSearchIndex& SearchIndex = GetSearchIndex();
+	if (SearchIndex.IsValuesEmpty())
+	{
+		UE_LOG(LogPoseSearch, Error, TEXT("UPoseSearchDatabase::SearchEvent unimplemented for reconstructed poses"));
+	}
+	else
+	{
+		const FPoseSearchEvent& EventToSearch = SearchContext.GetEventToSearch();
+		const TConstArrayView<int32> PosesWithEvent = SearchIndex.EventData.GetPosesWithEvent(EventToSearch.EventTag);
+		if (!PosesWithEvent.IsEmpty())
+		{
+			TConstArrayView<float> QueryValues = SearchContext.GetOrBuildQuery(Schema);
+
+			const int32 NumDimensions = Schema->SchemaCardinality;
+			TArrayView<float> DynamicWeightsSqrtBuffer((float*)FMemory_Alloca(NumDimensions * sizeof(float)), NumDimensions);
+			const TConstArrayView<float> DynamicWeightsSqrt = CalculateDynamicWeightsSqrt(DynamicWeightsSqrtBuffer);
+
+			FSelectableAssetIdx SelectableAssetIdx;
+			PopulateSelectableAssetIdx(SelectableAssetIdx, SearchContext.GetAssetsToConsider());
+
+			FNonSelectableIdx NonSelectableIdx;
+			PopulateNonSelectableIdx(NonSelectableIdx, SearchContext
+#if UE_POSE_SEARCH_TRACE_ENABLED
+				, QueryValues, DynamicWeightsSqrt
+#endif // UE_POSE_SEARCH_TRACE_ENABLED
+			);
+
+			FSearchFilters SearchFilters(Schema, NonSelectableIdx, SelectableAssetIdx, SearchIndex.bAnyBlockTransition);
+
+			if (FMath::IsNearlyZero(EventToSearch.TimeToEvent))
+			{
+				for (int32 PoseIdx : PosesWithEvent)
+				{
+					EvaluatePoseKernel<false, false>(Result, SearchIndex, QueryValues, TArrayView<float>(), PoseIdx, SearchFilters, SearchContext, this, DynamicWeightsSqrt, true
+#if WITH_EDITOR && ENABLE_ANIM_DEBUG
+						, PoseIdx
+#endif // WITH_EDITOR && ENABLE_ANIM_DEBUG
+					);
+				}
+			}
+			else
+			{
+				for (int32 PoseIdx : PosesWithEvent)
+				{
+					// Calculating the pose which is TimeToEvent seconds before the event
+					const FSearchIndexAsset& SearchIndexAsset = SearchIndex.GetAssetForPose(PoseIdx);
+					const float EventTime = SearchIndexAsset.GetTimeFromPoseIndex(PoseIdx, Schema->SampleRate);
+					const int32 SearchPoseIndex = SearchIndexAsset.GetPoseIndexFromTime(EventTime - EventToSearch.TimeToEvent, Schema->SampleRate);
+					if (SearchPoseIndex != INDEX_NONE)
+					{
+						EvaluatePoseKernel<false, false>(Result, SearchIndex, QueryValues, TArrayView<float>(), SearchPoseIndex, SearchFilters, SearchContext, this, DynamicWeightsSqrt, true
+#if WITH_EDITOR && ENABLE_ANIM_DEBUG
+							, SearchPoseIndex
+#endif // WITH_EDITOR && ENABLE_ANIM_DEBUG
+						);
+					}
+				}
+			}
+		}
+	}
+
+	// finalizing Result properties
+	if (Result.PoseIdx != INDEX_NONE)
+	{
+		Result.AssetTime = GetNormalizedAssetTime(Result.PoseIdx);
+		Result.Database = this;
+	}
+
+	return Result;
+}
diff --git a/Engine/Plugins/Animation/PoseSearch/Source/Runtime/Private/PoseSearchDerivedData.cpp b/Engine/Plugins/Animation/PoseSearch/Source/Runtime/Private/PoseSearchDerivedData.cpp
--- a/Engine/Plugins/Animation/PoseSearch/Source/Runtime/Private/PoseSearchDerivedData.cpp
+++ b/Engine/Plugins/Animation/PoseSearch/Source/Runtime/Private/PoseSearchDerivedData.cpp
@@ -1124,6 +1124,15 @@ static bool IndexDatabase(FSearchIndexBase& SearchIndexBase, const UPoseSearchDa
 		return false;
 	}
 
+	// Joining EventData
+	FEventDataCollector EventDataCollector;
+	for (int32 AssetIdx = 0; AssetIdx != SearchIndexBase.Assets.Num(); ++AssetIdx)
+	{
+		EventDataCollector.MergeWith(Indexers[AssetIdx].GetEventDataCollector());
+	}
+	// sorting EventData to make it deterministic across multiple indexing
+	SearchIndexBase.EventData.Initialize(EventDataCollector);
+
 	// Joining Metadata.Flags into OverallFlags
 	SearchIndexBase.bAnyBlockTransition = false;
 	for (const FPoseMetadata& Metadata : SearchIndexBase.PoseMetadata)
diff --git a/Engine/Plugins/Animation/PoseSearch/Source/Runtime/Private/PoseSearchDerivedDataKey.cpp b/Engine/Plugins/Animation/PoseSearch/Source/Runtime/Private/PoseSearchDerivedDataKey.cpp
--- a/Engine/Plugins/Animation/PoseSearch/Source/Runtime/Private/PoseSearchDerivedDataKey.cpp
+++ b/Engine/Plugins/Animation/PoseSearch/Source/Runtime/Private/PoseSearchDerivedDataKey.cpp
@@ -73,7 +73,7 @@ FKeyBuilder::FKeyBuilder(const UObject* Object, bool bUseDataVer, bool bPerformC
 		Hasher.Reset();
 
 		// used to invalidate the key without having to change POSESEARCHDB_DERIVEDDATA_VER all the times
-		int32 POSESEARCHDB_DERIVEDDATA_VER_SMALL = 2663;
+		int32 POSESEARCHDB_DERIVEDDATA_VER_SMALL = 2667;
 		FGuid VersionGuid = FDevSystemGuids::GetSystemGuid(FDevSystemGuids::Get().POSESEARCHDB_DERIVEDDATA_VER);
 		FString AnimationCompressionVersionString = UE::Anim::Compression::AnimationCompressionVersionString;
 
diff --git a/Engine/Plugins/Animation/PoseSearch/Source/Runtime/Private/PoseSearchHistory.cpp b/Engine/Plugins/Animation/PoseSearch/Source/Runtime/Private/PoseSearchHistory.cpp
--- a/Engine/Plugins/Animation/PoseSearch/Source/Runtime/Private/PoseSearchHistory.cpp
+++ b/Engine/Plugins/Animation/PoseSearch/Source/Runtime/Private/PoseSearchHistory.cpp
@@ -335,6 +335,7 @@ void FPoseHistoryEntry::Update(float Time, IComponentSpacePoseProvider& Componen
 	AccumulatedSeconds = Time;
 
 	const USkeleton* Skeleton = ComponentSpacePoseProvider.GetSkeletonAsset();
+	check(Skeleton);
 	const int32 NumSkeletonBones = Skeleton->GetReferenceSkeleton().GetNum();
 	if (BoneToTransformMap.IsEmpty())
 	{
diff --git a/Engine/Plugins/Animation/PoseSearch/Source/Runtime/Private/PoseSearchIndex.cpp b/Engine/Plugins/Animation/PoseSearch/Source/Runtime/Private/PoseSearchIndex.cpp
--- a/Engine/Plugins/Animation/PoseSearch/Source/Runtime/Private/PoseSearchIndex.cpp
+++ b/Engine/Plugins/Animation/PoseSearch/Source/Runtime/Private/PoseSearchIndex.cpp
@@ -300,6 +300,92 @@ FArchive& operator<<(FArchive& Ar, FSearchStats& Stats)
 	return Ar;
 }
 
+//////////////////////////////////////////////////////////////////////////
+// FEventDataCollector
+#if WITH_EDITOR
+void FEventDataCollector::Emplace(const FGameplayTag& EventTag, int32 PoseIdx)
+{
+	Data.FindOrAdd(EventTag).Add(PoseIdx);
+}
+
+void FEventDataCollector::MergeWith(const FEventDataCollector& Other)
+{
+	for (const FTagToPoseIndexes& OtherTagToPoseIndexes : Other.Data)
+	{
+		FPoseIndexes& PoseIndexes = Data.FindOrAdd(OtherTagToPoseIndexes.Key);
+
+		// adding all the missing poses of OtherTagToPoseIndexes to PoseIndexes
+		for (int32 OtherPoseIdx : OtherTagToPoseIndexes.Value)
+		{
+			PoseIndexes.Add(OtherPoseIdx);
+		}
+	}
+}
+#endif // WITH_EDITOR
+
+//////////////////////////////////////////////////////////////////////////
+// FEventData
+FArchive& operator<<(FArchive& Ar, FEventData& EventData)
+{
+	Ar << EventData.Data;
+	return Ar;
+}
+
+bool FEventData::operator==(const FEventData& Other) const
+{
+	return Data == Other.Data;
+}
+
+const TConstArrayView<int32> FEventData::GetPosesWithEvent(const FGameplayTag& GameplayTag) const
+{
+	// @todo: optmize me: Data is sorted! (or can be sorted differently to accomodate fast searches - maybe a serialized perfect hash map)
+	for (const FEventData::FTagToPoseIndexes& TagToPoseIndexes : Data)
+	{
+		if (TagToPoseIndexes.Key == GameplayTag)
+		{
+			return TagToPoseIndexes.Value;
+		}
+	}
+
+	return TConstArrayView<int32>();
+}
+
+void FEventData::Reset()
+{
+	Data.Reset();
+}
+
+#if WITH_EDITOR
+void FEventData::Initialize(const FEventDataCollector& EventDataCollector)
+{
+	Data.Reset();
+	for (const FEventDataCollector::FTagToPoseIndexes& OtherTagToPoseIndexes : EventDataCollector.GetData())
+	{
+		FTagToPoseIndexes& TagToPoseIndexes = Data.AddDefaulted_GetRef();
+		TagToPoseIndexes.Key = OtherTagToPoseIndexes.Key;
+
+		TagToPoseIndexes.Value.Reserve(OtherTagToPoseIndexes.Value.Num());
+		for (int32 PoseIdx : OtherTagToPoseIndexes.Value)
+		{
+			TagToPoseIndexes.Value.Add(PoseIdx);
+		}
+		TagToPoseIndexes.Value.Sort();
+	}
+
+	Data.Sort([](const FTagToPoseIndexes& A, const FTagToPoseIndexes& B)
+		{
+			// converting FGameplayTag::TagName to string to be deterministic across multiple editor restarts
+			return A.Key.ToString() < B.Key.ToString();
+		});
+}
+
+SIZE_T FEventData::GetAllocatedSize(void) const
+{
+	return Data.GetAllocatedSize();
+}
+
+#endif // WITH_EDITOR
+
 //////////////////////////////////////////////////////////////////////////
 // FPoseSearchBaseIndex
 const FSearchIndexAsset& FSearchIndexBase::GetAssetForPose(int32 PoseIdx) const
@@ -399,6 +485,8 @@ void FSearchIndexBase::AllocateData(int32 DataCardinality, int32 NumPoses)
 
 	Values.SetNumZeroed(DataCardinality * NumPoses);
 	PoseMetadata.SetNumZeroed(NumPoses);
+
+	EventData.Reset();
 }
 
 bool FSearchIndexBase::operator==(const FSearchIndexBase& Other) const
@@ -409,6 +497,7 @@ bool FSearchIndexBase::operator==(const FSearchIndexBase& Other) const
 		PoseMetadata == Other.PoseMetadata &&
 		bAnyBlockTransition == Other.bAnyBlockTransition &&
 		Assets == Other.Assets &&
+		EventData == Other.EventData &&
 		MinCostAddend == Other.MinCostAddend &&
 		Stats == Other.Stats;
 }
@@ -420,6 +509,7 @@ FArchive& operator<<(FArchive& Ar, FSearchIndexBase& Index)
 	Ar << Index.PoseMetadata;
 	Ar << Index.bAnyBlockTransition;
 	Ar << Index.Assets;
+	Ar << Index.EventData;
 	Ar << Index.MinCostAddend;
 	Ar << Index.Stats;
 	return Ar;
diff --git a/Engine/Plugins/Animation/PoseSearch/Source/Runtime/Private/PoseSearchLibrary.cpp b/Engine/Plugins/Animation/PoseSearch/Source/Runtime/Private/PoseSearchLibrary.cpp
--- a/Engine/Plugins/Animation/PoseSearch/Source/Runtime/Private/PoseSearchLibrary.cpp
+++ b/Engine/Plugins/Animation/PoseSearch/Source/Runtime/Private/PoseSearchLibrary.cpp
@@ -314,6 +314,7 @@ PRAGMA_ENABLE_DEPRECATION_WARNINGS
 	ElapsedPoseSearchTime = std::numeric_limits<float>::infinity();
 	WantedPlayRate = 1.f;
 	PoseIndicesHistory.Reset();
+	CurrentEventToSearch.Reset();
 }
 
 void FMotionMatchingState::AdjustAssetTime(float AssetTime)
@@ -631,7 +632,8 @@ void UPoseSearchLibrary::UpdateMotionMatchingState(
 	EPoseSearchInterruptMode InterruptMode,
 	bool bShouldUseCachedChannelData,
 	bool bDebugDrawQuery,
-	bool bDebugDrawCurResult)
+	bool bDebugDrawCurResult,
+	const FPoseSearchEvent& EventToSearch)
 {
 	QUICK_SCOPE_CYCLE_COUNTER(STAT_PoseSearch_Update);
 
@@ -660,7 +662,20 @@ void UPoseSearchLibrary::UpdateMotionMatchingState(
 		InOutMotionMatchingState.CurrentSearchResult.Reset();
 	}
 
-	FSearchContext SearchContext(0.f, &InOutMotionMatchingState.PoseIndicesHistory, InOutMotionMatchingState.CurrentSearchResult, PoseJumpThresholdTime);
+	if (EventToSearch.IsValid())
+	{
+		InOutMotionMatchingState.CurrentEventToSearch = EventToSearch;
+	}
+	else
+	{
+		InOutMotionMatchingState.CurrentEventToSearch.TimeToEvent -= DeltaTime;
+	}
+	if (InOutMotionMatchingState.CurrentEventToSearch.TimeToEvent < 0.f)
+	{
+		InOutMotionMatchingState.CurrentEventToSearch.Reset();
+	}
+
+	FSearchContext SearchContext(0.f, &InOutMotionMatchingState.PoseIndicesHistory, InOutMotionMatchingState.CurrentSearchResult, PoseJumpThresholdTime, InOutMotionMatchingState.CurrentEventToSearch);
 	SearchContext.AddRole(GetCommonDefaultRole(Databases), AnimContext, PoseHistory);
 
 	const bool bCanAdvance = InOutMotionMatchingState.CurrentSearchResult.CanAdvance(DeltaTime);
@@ -1065,7 +1080,7 @@ UE::PoseSearch::FSearchResult UPoseSearchLibrary::MotionMatch(
 
 	FReconstructedPreviousSearchBestResultMap ReconstructedPreviousSearchBestResultMap;
 	FSearchResult ReconstructedPreviousSearchResult;
-	FSearchContext SearchContext(FutureIntervalTime, nullptr, ReconstructedPreviousSearchResult);
+	FSearchContext SearchContext(FutureIntervalTime, nullptr, ReconstructedPreviousSearchResult, FFloatInterval(0.f, 0.f), FPoseSearchEvent());
 
 	for (int32 RoleIndex = 0; RoleIndex < Roles.Num(); ++RoleIndex)
 	{
diff --git a/Engine/Plugins/Animation/PoseSearch/Source/Runtime/Public/PoseSearch/AnimNode_MotionMatching.h b/Engine/Plugins/Animation/PoseSearch/Source/Runtime/Public/PoseSearch/AnimNode_MotionMatching.h
--- a/Engine/Plugins/Animation/PoseSearch/Source/Runtime/Public/PoseSearch/AnimNode_MotionMatching.h
+++ b/Engine/Plugins/Animation/PoseSearch/Source/Runtime/Public/PoseSearch/AnimNode_MotionMatching.h
@@ -9,6 +9,7 @@
 #include "AnimNode_MotionMatching.generated.h"
 
 class UPoseSearchDatabase;
+struct FPoseSearchEvent;
 
 USTRUCT(BlueprintInternalUseOnly)
 struct POSESEARCH_API FAnimNode_MotionMatching : public FAnimNode_BlendStack_Standalone
@@ -155,6 +156,10 @@ private:
 	UPROPERTY(EditAnywhere, Category = Settings, meta = (PinHiddenByDefault))
 	bool bShouldUseCachedChannelData = false;
 	
+	// Experimental, this feature might be removed without warning, not for production use
+	UPROPERTY(EditAnywhere, Category = "Settings|Experimental", meta = (PinHiddenByDefault))
+	FPoseSearchEvent EventToSearch;
+
 	// Encapsulated motion matching algorithm and internal state
 	FMotionMatchingState MotionMatchingState;
 
diff --git a/Engine/Plugins/Animation/PoseSearch/Source/Runtime/Public/PoseSearch/PoseSearchAnimNotifies.h b/Engine/Plugins/Animation/PoseSearch/Source/Runtime/Public/PoseSearch/PoseSearchAnimNotifies.h
--- a/Engine/Plugins/Animation/PoseSearch/Source/Runtime/Public/PoseSearch/PoseSearchAnimNotifies.h
+++ b/Engine/Plugins/Animation/PoseSearch/Source/Runtime/Public/PoseSearch/PoseSearchAnimNotifies.h
@@ -2,8 +2,10 @@
 
 #pragma once
 
+#include "Animation/AnimNotifies/AnimNotify.h"
 #include "Animation/AnimNotifies/AnimNotifyState.h"
 #include "BoneContainer.h"
+#include "GameplayTagContainer.h"
 #include "PoseSearchAnimNotifies.generated.h"
 
 class UPoseSearchDatabase;
@@ -135,3 +137,14 @@ public:
 	UPROPERTY(EditAnywhere, Category = Config)
 	TArray<FName> GoalNames;
 };
+
+UCLASS(Experimental, Blueprintable, meta = (DisplayName = "Pose Search: Event"))
+class POSESEARCH_API UAnimNotify_PoseSearchEvent : public UAnimNotify
+{
+	GENERATED_BODY()
+
+public:
+	// Tag identifying this event
+	UPROPERTY(EditAnywhere, Category = Config)
+	FGameplayTag EventTag;
+};
\ No newline at end of file
diff --git a/Engine/Plugins/Animation/PoseSearch/Source/Runtime/Public/PoseSearch/PoseSearchAssetIndexer.h b/Engine/Plugins/Animation/PoseSearch/Source/Runtime/Public/PoseSearch/PoseSearchAssetIndexer.h
--- a/Engine/Plugins/Animation/PoseSearch/Source/Runtime/Public/PoseSearch/PoseSearchAssetIndexer.h
+++ b/Engine/Plugins/Animation/PoseSearch/Source/Runtime/Public/PoseSearch/PoseSearchAssetIndexer.h
@@ -5,9 +5,8 @@
 #if WITH_EDITOR
 
 #include "BonePose.h"
+#include "PoseSearch/PoseSearchIndex.h"
 #include "PoseSearch/PoseSearchMirrorDataCache.h"
-#include "PoseSearch/PoseSearchDefines.h"
-#include "PoseSearch/PoseSearchFeatureChannel.h"
 #include "PoseSearch/PoseSearchSchema.h"
 
 class UAnimNotifyState_PoseSearchBase;
@@ -35,7 +34,7 @@ struct FAnimationAssetSamplers
 
 	float GetPlayLength() const;
 	bool IsLoopable() const;
-	void ExtractPoseSearchNotifyStates(float Time, const TFunction<bool(UAnimNotifyState_PoseSearchBase*)>& ProcessPoseSearchBase) const;
+	void ExtractAnimNotifyStates(float Time, FAnimNotifyContext& PreAllocatedNotifyContext, const TFunction<bool(UAnimNotifyState*)>& ProcessAnimNotifyState) const;
 	bool ProcessAllAnimNotifyEvents(const TFunction<bool(TConstArrayView<FAnimNotifyEvent>)>& ProcessAnimNotifyEvents) const;
 	const FString GetAssetName() const;
 
@@ -67,6 +66,8 @@ public:
 	void AssignWorkingData(int32 InStartPoseIdx, TArrayView<float> InOutFeatureVectorTable, TArrayView<FPoseMetadata> InOutPoseMetadata);
 	void Process(int32 AssetIdx);
 	const FStats& GetStats() const { return Stats; }
+	// Experimental, this feature might be removed without warning, not for production use
+	const FEventDataCollector& GetEventDataCollector() const { return EventDataCollector; }
 
 	// Returns the value of float curve CurveName at time CalculateSampleTime(SampleIdx) + SampleTimeOffset.
 	bool GetSampleCurveValue(float& OutCurveValue, float SampleTimeOffset, int32 SampleIdx, const FName& CurveName, const FRole& SampleRole);
@@ -175,6 +176,7 @@ private:
 	TArrayView<float> FeatureVectorTable;
 	TArrayView<FPoseMetadata> PoseMetadata;
 
+	FEventDataCollector EventDataCollector;
 	FStats Stats;
 
 	bool bProcessFailed = false;
diff --git a/Engine/Plugins/Animation/PoseSearch/Source/Runtime/Public/PoseSearch/PoseSearchAssetSampler.h b/Engine/Plugins/Animation/PoseSearch/Source/Runtime/Public/PoseSearch/PoseSearchAssetSampler.h
--- a/Engine/Plugins/Animation/PoseSearch/Source/Runtime/Public/PoseSearch/PoseSearchAssetSampler.h
+++ b/Engine/Plugins/Animation/PoseSearch/Source/Runtime/Public/PoseSearch/PoseSearchAssetSampler.h
@@ -7,57 +7,59 @@
 #include "IO/IoHash.h"
 #endif // WITH_EDITOR
 
-struct FAnimationPoseData;
+struct FAnimNotifyContext;
 struct FAnimExtractContext;
-class UAnimationAsset;
+struct FAnimationPoseData;
 class UAnimNotifyState_PoseSearchBase;
+class UAnimationAsset;
 
 namespace UE::PoseSearch
 {
-
 /**
  * Helper for sampling data from animation assets
  */
-struct POSESEARCH_API FAnimationAssetSampler
+struct FAnimationAssetSampler
 {
 	enum { DefaultRootTransformSamplingRate = 30 };
-	FAnimationAssetSampler(const UAnimationAsset* InAnimationAsset = nullptr, const FTransform& InRootTransformOrigin = FTransform::Identity, const FVector& InBlendParameters = FVector::ZeroVector, int32 InRootTransformSamplingRate = DefaultRootTransformSamplingRate, bool bPreProcessRootTransform = true);
-	void Init(const UAnimationAsset* InAnimationAsset, const FTransform& InRootTransformOrigin = FTransform::Identity, const FVector& InBlendParameters = FVector::ZeroVector, int32 InRootTransformSamplingRate = DefaultRootTransformSamplingRate, bool bPreProcessRootTransform = true);
+	POSESEARCH_API FAnimationAssetSampler(const UAnimationAsset* InAnimationAsset = nullptr, const FTransform& InRootTransformOrigin = FTransform::Identity, const FVector& InBlendParameters = FVector::ZeroVector, int32 InRootTransformSamplingRate = DefaultRootTransformSamplingRate, bool bPreProcessRootTransform = true);
+	POSESEARCH_API void Init(const UAnimationAsset* InAnimationAsset, const FTransform& InRootTransformOrigin = FTransform::Identity, const FVector& InBlendParameters = FVector::ZeroVector, int32 InRootTransformSamplingRate = DefaultRootTransformSamplingRate, bool bPreProcessRootTransform = true);
 
-	bool IsInitialized() const;
-	float GetPlayLength() const;
-	float ToRealTime(float NormalizedTime) const;
-	float ToNormalizedTime(float RealTime) const;
-	bool IsLoopable() const;
+	POSESEARCH_API bool IsInitialized() const;
+	POSESEARCH_API float GetPlayLength() const;
+	POSESEARCH_API float ToRealTime(float NormalizedTime) const;
+	POSESEARCH_API float ToNormalizedTime(float RealTime) const;
+	POSESEARCH_API bool IsLoopable() const;
 
 	// Gets the final root transformation at the end of the asset's playback time
-	FTransform GetTotalRootTransform() const;
+	POSESEARCH_API FTransform GetTotalRootTransform() const;
 
 	// Extracts pose for this asset for a given context
-	void ExtractPose(const FAnimExtractContext& ExtractionCtx, FAnimationPoseData& OutAnimPoseData) const;
+	POSESEARCH_API void ExtractPose(const FAnimExtractContext& ExtractionCtx, FAnimationPoseData& OutAnimPoseData) const;
 
 	// Extracts pose for this asset at a given Time
-	void ExtractPose(float Time, FCompactPose& OutPose) const;
-	void ExtractPose(float Time, FCompactPose& OutPose, FBlendedCurve& OutCurve) const;
+	POSESEARCH_API void ExtractPose(float Time, FCompactPose& OutPose) const;
+	POSESEARCH_API void ExtractPose(float Time, FCompactPose& OutPose, FBlendedCurve& OutCurve) const;
 
 	// Extracts root transform at the given time, using the extremities of the sequence to extrapolate beyond the 
 	// sequence limits when Time is less than zero or greater than the sequence length.
-	FTransform ExtractRootTransform(float Time) const;
+	POSESEARCH_API FTransform ExtractRootTransform(float Time) const;
+	
+	UE_DEPRECATED(5.6, "Use ExtractAnimNotifyStates instead")
+	POSESEARCH_API void ExtractPoseSearchNotifyStates(float Time, const TFunction<bool(UAnimNotifyState_PoseSearchBase*)>& ProcessPoseSearchBase) const;
 
-	// Extracts notify states inheriting from UAnimNotifyState_PoseSearchBase present in the sequence at Time.
-	// The function does not empty NotifyStates before adding new notifies!
-	void ExtractPoseSearchNotifyStates(float Time, const TFunction<bool(UAnimNotifyState_PoseSearchBase*)>& ProcessPoseSearchBase) const;
+	// Extracts notify states present in the sequence at Time
+	POSESEARCH_API void ExtractAnimNotifyStates(float Time, FAnimNotifyContext& PreAllocatedNotifyContext, const TFunction<bool(UAnimNotifyState*)>& ProcessAnimNotifyState) const;
 	
-	TConstArrayView<FAnimNotifyEvent> GetAllAnimNotifyEvents() const;
+	POSESEARCH_API TConstArrayView<FAnimNotifyEvent> GetAllAnimNotifyEvents() const;
 
-	const UAnimationAsset* GetAsset() const;
+	POSESEARCH_API const UAnimationAsset* GetAsset() const;
 
-	void Process();
+	POSESEARCH_API void Process();
 
-	static float GetPlayLength(const UAnimationAsset* AnimAsset, const FVector& BlendParameters = FVector::ZeroVector);
+	POSESEARCH_API static float GetPlayLength(const UAnimationAsset* AnimAsset, const FVector& BlendParameters = FVector::ZeroVector);
 
-	void SetRootTransformOrigin(const FTransform& InRootTransformOrigin) { RootTransformOrigin = InRootTransformOrigin; }
-	const FTransform& GetRootTransformOrigin() const { return RootTransformOrigin; }
+	POSESEARCH_API void SetRootTransformOrigin(const FTransform& InRootTransformOrigin) { RootTransformOrigin = InRootTransformOrigin; }
+	POSESEARCH_API const FTransform& GetRootTransformOrigin() const { return RootTransformOrigin; }
 
 protected:
 	TWeakObjectPtr<const UAnimationAsset> AnimationAssetPtr;
diff --git a/Engine/Plugins/Animation/PoseSearch/Source/Runtime/Public/PoseSearch/PoseSearchContext.h b/Engine/Plugins/Animation/PoseSearch/Source/Runtime/Public/PoseSearch/PoseSearchContext.h
--- a/Engine/Plugins/Animation/PoseSearch/Source/Runtime/Public/PoseSearch/PoseSearchContext.h
+++ b/Engine/Plugins/Animation/PoseSearch/Source/Runtime/Public/PoseSearch/PoseSearchContext.h
@@ -2,9 +2,10 @@
 
 #pragma once
 
-#include "DrawDebugHelpers.h"
 #include "Animation/AnimInstance.h"
+#include "DrawDebugHelpers.h"
 #include "PoseSearch/PoseSearchDefines.h"
+#include "PoseSearch/PoseSearchEvent.h"
 #include "PoseSearch/PoseSearchFeatureChannel.h"
 #include "PoseSearch/PoseSearchIndex.h"
 #include "PoseSearch/PoseSearchResult.h"
@@ -135,9 +136,13 @@ struct FCachedChannel
 
 struct POSESEARCH_API FSearchContext
 {
+	UE_DEPRECATED(5.6, "Use other signaure instead")
 	FSearchContext(float InDesiredPermutationTimeOffset = 0.f, const FPoseIndicesHistory* InPoseIndicesHistory = nullptr,
 		const FSearchResult& InCurrentResult = FSearchResult(), const FFloatInterval& InPoseJumpThresholdTime = FFloatInterval(0.f, 0.f), bool bInUseCachedChannelData = false);
 
+	FSearchContext(float InDesiredPermutationTimeOffset, const FPoseIndicesHistory* InPoseIndicesHistory,
+		const FSearchResult& InCurrentResult, const FFloatInterval& InPoseJumpThresholdTime, const FPoseSearchEvent& InEventToSearch);
+
 	// deleting copies and moves since members could reference other members (e.g.: CurrentResultPoseVector could point to CurrentResultPoseVectorData, so it'll require proper copies and movers implementations)
 	FSearchContext(const FSearchContext& Other) = delete;
 	FSearchContext(FSearchContext&& Other) = delete;
@@ -209,6 +214,8 @@ struct POSESEARCH_API FSearchContext
 
 	void SetAssetsToConsider(TConstArrayView<const UObject*> InAssetsToConsider) { AssetsToConsider = InAssetsToConsider; }
 	TConstArrayView<const UObject*> GetAssetsToConsider() const { return AssetsToConsider; }
+	// Experimental, this feature might be removed without warning, not for production use
+	const FPoseSearchEvent& GetEventToSearch() const { return EventToSearch; }
 	
 	// returns the world space transform of the bone SchemaBoneIdx at time SampleTime
 	FTransform GetWorldBoneTransformAtTime(float SampleTime, const FRole& SampleRole, int8 SchemaBoneIdx);
@@ -239,6 +246,7 @@ private:
 	
 	// if AssetsToConsider is not empty, we'll search only for poses from UObject(s) that are in the AssetsToConsider
 	TConstArrayView<const UObject*> AssetsToConsider;
+	const FPoseSearchEvent& EventToSearch;
 
 	const float DesiredPermutationTimeOffset = 0.f;
 	const FPoseIndicesHistory* PoseIndicesHistory = nullptr;
diff --git a/Engine/Plugins/Animation/PoseSearch/Source/Runtime/Public/PoseSearch/PoseSearchDatabase.h b/Engine/Plugins/Animation/PoseSearch/Source/Runtime/Public/PoseSearch/PoseSearchDatabase.h
--- a/Engine/Plugins/Animation/PoseSearch/Source/Runtime/Public/PoseSearch/PoseSearchDatabase.h
+++ b/Engine/Plugins/Animation/PoseSearch/Source/Runtime/Public/PoseSearch/PoseSearchDatabase.h
@@ -36,9 +36,13 @@ enum class EPoseSearchMode : int32
 	// Optimized search mode: the database projects the poses into a PCA space using only the most significant "NumberOfPrincipalComponents" dimensions, and construct a kdtree to facilitate the search.
 	PCAKDTree,
 
+	// Experimental, this feature might be removed without warning, not for production use
 	// Optimized search mode using a vantage point tree
+	VPTree UMETA(DisplayName = "VPTree (Experimental)"),
+
 	// Experimental, this feature might be removed without warning, not for production use
-	VPTree UMETA(DisplayName = "VPTree (Experimental)")
+	// search will only be performed on events.
+	EventOnly UMETA(DisplayName = "EventOnly (Experimental)")
 };
 
 UENUM()
@@ -581,6 +585,7 @@ private:
 	UE::PoseSearch::FSearchResult SearchPCAKDTree(UE::PoseSearch::FSearchContext& SearchContext) const;
 	UE::PoseSearch::FSearchResult SearchVPTree(UE::PoseSearch::FSearchContext& SearchContext) const;
 	UE::PoseSearch::FSearchResult SearchBruteForce(UE::PoseSearch::FSearchContext& SearchContext) const;
+	UE::PoseSearch::FSearchResult SearchEvent(UE::PoseSearch::FSearchContext& SearchContext) const;
 
 	typedef TArray<int32, TInlineAllocator<256, TMemStackAllocator<>>> FSelectableAssetIdx;
 	void PopulateSelectableAssetIdx(FSelectableAssetIdx& SelectableAssetIdx, TConstArrayView<const UObject*> AssetsToConsider) const;
diff --git a/Engine/Plugins/Animation/PoseSearch/Source/Runtime/Public/PoseSearch/PoseSearchEvent.h b/Engine/Plugins/Animation/PoseSearch/Source/Runtime/Public/PoseSearch/PoseSearchEvent.h
new file mode 100644
--- /dev/null
+++ b/Engine/Plugins/Animation/PoseSearch/Source/Runtime/Public/PoseSearch/PoseSearchEvent.h
@@ -0,0 +1,32 @@
+// Copyright Epic Games, Inc. All Rights Reserved.
+
+#pragma once
+
+#include "GameplayTagContainer.h"
+#include "PoseSearchEvent.generated.h"
+
+USTRUCT(Experimental, BlueprintType, Category="Animation|Pose Search")
+struct FPoseSearchEvent
+{
+	GENERATED_BODY()
+
+	bool IsValid() const
+	{
+		// maybe we should check if TimeToEvent >= 0?
+		return EventTag.IsValid();
+	}
+
+	void Reset()
+	{
+		EventTag = FGameplayTag();
+		TimeToEvent = 0.f;
+	}
+
+	// Tag identifying this event
+	UPROPERTY(EditAnywhere, BlueprintReadWrite, Category=State)
+	FGameplayTag EventTag;
+
+	// Time in seconds until this event occurs
+	UPROPERTY(EditAnywhere, BlueprintReadWrite, Category=State)
+	float TimeToEvent = 0.f;
+};
diff --git a/Engine/Plugins/Animation/PoseSearch/Source/Runtime/Public/PoseSearch/PoseSearchIndex.h b/Engine/Plugins/Animation/PoseSearch/Source/Runtime/Public/PoseSearch/PoseSearchIndex.h
--- a/Engine/Plugins/Animation/PoseSearch/Source/Runtime/Public/PoseSearch/PoseSearchIndex.h
+++ b/Engine/Plugins/Animation/PoseSearch/Source/Runtime/Public/PoseSearch/PoseSearchIndex.h
@@ -2,6 +2,7 @@
 
 #pragma once
 
+#include "GameplayTagContainer.h"
 #include "PoseSearch/KDTree.h"
 #include "PoseSearch/PoseSearchCost.h"
 #include "PoseSearch/PoseSearchDefines.h"
@@ -347,6 +348,45 @@ struct FSparsePoseMultiMap
 	TArray<Type> DataValues;
 };
 
+#if WITH_EDITOR
+// Experimental, this feature might be removed without warning, not for production use
+struct FEventDataCollector
+{
+	typedef TSet<int32, DefaultKeyFuncs<int32>, TInlineSetAllocator<256>> FPoseIndexes;
+	typedef TPair<FGameplayTag, FPoseIndexes> FTagToPoseIndexes;
+	typedef TMap<FGameplayTag, FPoseIndexes, TInlineSetAllocator<16>> FData;
+
+	void Emplace(const FGameplayTag& EventTag, int32 PoseIdx);
+	void MergeWith(const FEventDataCollector& Other);
+	const FData& GetData() const { return Data; }
+
+private:
+	TMap<FGameplayTag, FPoseIndexes, TInlineSetAllocator<16>> Data;
+};
+#endif // WITH_EDITOR
+
+// Experimental, this feature might be removed without warning, not for production use
+struct FEventData
+{
+	typedef TArray<int32> FPoseIndexes;
+	typedef TPair<FGameplayTag, FPoseIndexes> FTagToPoseIndexes;
+	typedef TArray<FTagToPoseIndexes> FData;
+
+	friend FArchive& operator<<(FArchive& Ar, FEventData& EventData);
+	bool operator==(const FEventData& Other) const;
+	const TConstArrayView<int32> GetPosesWithEvent(const FGameplayTag& GameplayTag) const;
+	const FData& GetData() const { return Data; }
+	void Reset();
+
+#if WITH_EDITOR
+	void Initialize(const FEventDataCollector& EventDataCollector);
+	POSESEARCH_API SIZE_T GetAllocatedSize(void) const;
+#endif // WITH_EDITOR
+
+private:
+	FData Data;
+};
+
 /**
 * case class for FSearchIndex. building block used to gather data for data mining and calculate weights, pca, kdtree stuff
 */
@@ -357,6 +397,8 @@ struct FSearchIndexBase
 	TAlignedArray<FPoseMetadata> PoseMetadata;
 	bool bAnyBlockTransition = false;
 	TAlignedArray<FSearchIndexAsset> Assets;
+	// Experimental, this feature might be removed without warning, not for production use
+	FEventData EventData;
 
 	// minimum of the database metadata CostAddend: it represents the minimum cost of any search for the associated database (we'll skip the search in case the search result total cost is already less than MinCostAddend)
 	float MinCostAddend = -MAX_FLT;
@@ -403,6 +445,9 @@ struct FSearchIndexBase
 
 	bool operator==(const FSearchIndexBase& Other) const;
 
+	// Experimental, this feature might be removed without warning, not for production use
+	TArray<int32> GetPosesWithEvent(const FGameplayTag& GameplayTag) const;
+
 	friend FArchive& operator<<(FArchive& Ar, FSearchIndexBase& Index);
 };
 
diff --git a/Engine/Plugins/Animation/PoseSearch/Source/Runtime/Public/PoseSearch/PoseSearchLibrary.h b/Engine/Plugins/Animation/PoseSearch/Source/Runtime/Public/PoseSearch/PoseSearchLibrary.h
--- a/Engine/Plugins/Animation/PoseSearch/Source/Runtime/Public/PoseSearch/PoseSearchLibrary.h
+++ b/Engine/Plugins/Animation/PoseSearch/Source/Runtime/Public/PoseSearch/PoseSearchLibrary.h
@@ -6,6 +6,7 @@
 #include "Animation/AnimNodeReference.h"
 #include "Kismet/BlueprintFunctionLibrary.h"
 #include "PoseSearch/PoseSearchDefines.h"
+#include "PoseSearch/PoseSearchEvent.h"
 #include "PoseSearch/PoseSearchHistory.h"
 #include "PoseSearch/PoseSearchResult.h"
 #include "PoseSearch/PoseSearchRole.h"
@@ -75,6 +76,9 @@ PRAGMA_ENABLE_DEPRECATION_WARNINGS
 	bool bJumpedToPose = false;
 
 	UE::PoseSearch::FPoseIndicesHistory PoseIndicesHistory;
+
+	// Experimental, this feature might be removed without warning, not for production use
+	FPoseSearchEvent CurrentEventToSearch;
 };
 
 USTRUCT(Experimental, BlueprintType, Category="Animation|Pose Search")
@@ -82,7 +86,6 @@ struct POSESEARCH_API FPoseSearchFutureProperties
 {
 	GENERATED_BODY()
 
-public:
 	// Animation to play (it'll start at AnimationTime seconds)
 	UPROPERTY(EditAnywhere, BlueprintReadWrite, Category=State)
 	TObjectPtr<const UObject> Animation;
@@ -101,7 +104,6 @@ struct POSESEARCH_API FPoseSearchContinuingProperties
 {
 	GENERATED_BODY()
 
-public:
 	// Currently playing animation
 	UPROPERTY(EditAnywhere, BlueprintReadWrite, Category=State)
 	TObjectPtr<const UObject> PlayingAsset = nullptr;
@@ -186,7 +188,8 @@ public:
 		EPoseSearchInterruptMode InterruptMode,
 		bool bShouldUseCachedChannelData,
 		bool bDebugDrawQuery = false,
-		bool bDebugDrawCurResult = false);
+		bool bDebugDrawCurResult = false,
+		const FPoseSearchEvent& EventToSearch = FPoseSearchEvent());
 
 	UE_DEPRECATED(5.6, "Use other UpdateMotionMatchingState signatures instead")
 	static void UpdateMotionMatchingState(
