PoseSearch - fix for null NotifyEvent->NotifyStateClass in FAnimationAssetSampler::ExtractAnimNotifyStates

diff --git a/Engine/Plugins/Animation/PoseSearch/Source/Runtime/Private/PoseSearchAssetSampler.cpp b/Engine/Plugins/Animation/PoseSearch/Source/Runtime/Private/PoseSearchAssetSampler.cpp
--- a/Engine/Plugins/Animation/PoseSearch/Source/Runtime/Private/PoseSearchAssetSampler.cpp
+++ b/Engine/Plugins/Animation/PoseSearch/Source/Runtime/Private/PoseSearchAssetSampler.cpp
@@ -691,7 +691,7 @@ void FAnimationAssetSampler::ExtractAnimNotifyStates(float Time, FAnimNotifyCont
 	{
 		if (const FAnimNotifyEvent* NotifyEvent = EventReference.GetNotify())
 		{
-			if (!ProcessAnimNotifyState(NotifyEvent->NotifyStateClass))
+			if (NotifyEvent->NotifyStateClass && !ProcessAnimNotifyState(NotifyEvent->NotifyStateClass))
 			{
 				break;
 			}
