PoseSearch - fix for FSearchResult::CanAdvance

diff --git a/Engine/Plugins/Animation/PoseSearch/Source/Runtime/Private/PoseSearchResult.cpp b/Engine/Plugins/Animation/PoseSearch/Source/Runtime/Private/PoseSearchResult.cpp
--- a/Engine/Plugins/Animation/PoseSearch/Source/Runtime/Private/PoseSearchResult.cpp
+++ b/Engine/Plugins/Animation/PoseSearch/Source/Runtime/Private/PoseSearchResult.cpp
@@ -55,32 +55,24 @@ const FSearchIndexAsset* FSearchResult::GetSearchIndexAsset(bool bMandatory) con
 
 bool FSearchResult::CanAdvance(float DeltaTime) const
 {
-	bool bCanAdvance = false;
-	if (const FSearchIndexAsset* SearchIndexAsset = GetSearchIndexAsset())
+	const FSearchIndexAsset* SearchIndexAsset = GetSearchIndexAsset();
+	if (!SearchIndexAsset)
 	{
-		const FPoseSearchDatabaseAnimationAssetBase* DatabaseAnimationAssetBase = Database->GetDatabaseAnimationAsset<FPoseSearchDatabaseAnimationAssetBase>(*SearchIndexAsset);
-		check(DatabaseAnimationAssetBase);
-
-		float SteppedTime = AssetTime;
-		if (const UBlendSpace* BlendSpace = Cast<UBlendSpace>(DatabaseAnimationAssetBase->GetAnimationAsset()))
-		{
-			TArray<FBlendSampleData> BlendSamples;
-			int32 TriangulationIndex = 0;
-			BlendSpace->GetSamplesFromBlendInput(SearchIndexAsset->GetBlendParameters(), BlendSamples, TriangulationIndex, true);
-
-			const float PlayLength = BlendSpace->GetAnimationLengthFromSampleData(BlendSamples);
+		return false;
+	}
+	
+	if (SearchIndexAsset->IsLooping())
+	{
+		return true;
+	}
 
-			// Asset player time for blend spaces is normalized [0, 1] so we need to convert it back to real time before we advance it
-			SteppedTime = AssetTime * PlayLength;
-			bCanAdvance = ETAA_Finished != FAnimationRuntime::AdvanceTime(SearchIndexAsset->IsLooping(), DeltaTime, SteppedTime, PlayLength);
-		}
-		else
-		{
-			const float AssetLength = DatabaseAnimationAssetBase->GetPlayLength();
-			bCanAdvance = ETAA_Finished != FAnimationRuntime::AdvanceTime(SearchIndexAsset->IsLooping(), DeltaTime, SteppedTime, AssetLength);
-		}
+	// leaving this if statement for debuggability
+	if (INDEX_NONE != Database->GetPoseIndexFromTime(AssetTime + DeltaTime, *SearchIndexAsset))
+	{
+		return true;
 	}
-	return bCanAdvance;
+
+	return false;
 }
 
 } // namespace UE::PoseSearch
