PoseSearch - blend spaces support for interaction assets

diff --git a/Engine/Plugins/Animation/PoseSearch/Source/Editor/Private/PoseSearchDatabaseAssetTree.cpp b/Engine/Plugins/Animation/PoseSearch/Source/Editor/Private/PoseSearchDatabaseAssetTree.cpp
--- a/Engine/Plugins/Animation/PoseSearch/Source/Editor/Private/PoseSearchDatabaseAssetTree.cpp
+++ b/Engine/Plugins/Animation/PoseSearch/Source/Editor/Private/PoseSearchDatabaseAssetTree.cpp
@@ -1090,7 +1090,7 @@ namespace UE::PoseSearch
 								bModified = true;
 							}
 							
-							const FFloatInterval SamplingRange = DatabaseAnimationAssetBase->GetEffectiveSamplingRange();
+							const FFloatInterval SamplingRange = DatabaseAnimationAssetBase->GetEffectiveSamplingRange(FVector::ZeroVector);
 							const float StartTime = SamplingRange.Min;
 							const float Duration = SamplingRange.Max - SamplingRange.Min;
 							const FName TrackName = "PoseSearch";
diff --git a/Engine/Plugins/Animation/PoseSearch/Source/Editor/Private/PoseSearchDatabaseEditorReflection.cpp b/Engine/Plugins/Animation/PoseSearch/Source/Editor/Private/PoseSearchDatabaseEditorReflection.cpp
--- a/Engine/Plugins/Animation/PoseSearch/Source/Editor/Private/PoseSearchDatabaseEditorReflection.cpp
+++ b/Engine/Plugins/Animation/PoseSearch/Source/Editor/Private/PoseSearchDatabaseEditorReflection.cpp
@@ -33,12 +33,24 @@ void UPoseSearchDatabaseReflectionBase::PostEditChangeProperty(FPropertyChangedE
 	{
 		const bool bShouldRefreshView = (PropertyChangedEvent.Property == nullptr) || !PropertyChangedEvent.Property->IsA(FFloatProperty::StaticClass());
 		AssetTreeWidget->FinalizeTreeChanges(true, bShouldRefreshView);
+
+		if (const TSharedPtr<UE::PoseSearch::FDatabaseAssetTreeNode> AssetTreeNode = WeakAssetTreeNode.Pin())
+		{
+			if (const TSharedPtr<UE::PoseSearch::FDatabaseViewModel> ViewModel = AssetTreeNode->EditorViewModel.Pin())
+			{
+				UPoseSearchDatabase* Database = ViewModel->GetPoseSearchDatabase();
+				if (IsValid(Database))
+				{
+					Database->Modify();
+				}
+			}
+		}
 	}
 }
 
 bool UPoseSearchDatabaseSequenceReflection::ApplyChanges() 
 {
-	const float PlayLength = Sequence.GetPlayLength();
+	const float PlayLength = Sequence.GetPlayLength(FVector::ZeroVector);
 	const FFloatInterval ClampedSamplingRange(FMath::Clamp(Sequence.SamplingRange.Min, 0.0f, PlayLength), FMath::Clamp(Sequence.SamplingRange.Max, 0.0f, PlayLength));
 	if (ClampedSamplingRange != Sequence.SamplingRange)
 	{
@@ -57,7 +69,6 @@ bool UPoseSearchDatabaseSequenceReflection::ApplyChanges()
 					if (*DatabaseSequence != Sequence)
 					{
 						*DatabaseSequence = Sequence;
-						Database->MarkPackageDirty();
 						return true;
 					}
 				}
@@ -82,7 +93,6 @@ bool UPoseSearchDatabaseBlendSpaceReflection::ApplyChanges()
 					if (*DatabaseBlendSpace != BlendSpace)
 					{
 						*DatabaseBlendSpace = BlendSpace;
-						Database->MarkPackageDirty();
 						return true;
 					}
 				}
@@ -95,7 +105,7 @@ bool UPoseSearchDatabaseBlendSpaceReflection::ApplyChanges()
 
 bool UPoseSearchDatabaseAnimCompositeReflection::ApplyChanges()
 {
-	const float PlayLength = AnimComposite.GetPlayLength();
+	const float PlayLength = AnimComposite.GetPlayLength(FVector::ZeroVector);
 	const FFloatInterval ClampedSamplingRange(FMath::Clamp(AnimComposite.SamplingRange.Min, 0.0f, PlayLength), FMath::Clamp(AnimComposite.SamplingRange.Max, 0.0f, PlayLength));
 	if (ClampedSamplingRange != AnimComposite.SamplingRange)
 	{
@@ -114,7 +124,6 @@ bool UPoseSearchDatabaseAnimCompositeReflection::ApplyChanges()
 					if (*DatabaseAnimComposite != AnimComposite)
 					{
 						*DatabaseAnimComposite = AnimComposite;
-						Database->MarkPackageDirty();
 						return true;
 					}
 				}
@@ -127,7 +136,7 @@ bool UPoseSearchDatabaseAnimCompositeReflection::ApplyChanges()
 
 bool UPoseSearchDatabaseAnimMontageReflection::ApplyChanges()
 {
-	const float PlayLength = AnimMontage.GetPlayLength();
+	const float PlayLength = AnimMontage.GetPlayLength(FVector::ZeroVector);
 	const FFloatInterval ClampedSamplingRange(FMath::Clamp(AnimMontage.SamplingRange.Min, 0.0f, PlayLength), FMath::Clamp(AnimMontage.SamplingRange.Max, 0.0f, PlayLength));
 	if (ClampedSamplingRange != AnimMontage.SamplingRange)
 	{
@@ -146,7 +155,6 @@ bool UPoseSearchDatabaseAnimMontageReflection::ApplyChanges()
 					if (*DatabaseAnimMontage != AnimMontage)
 					{
 						*DatabaseAnimMontage = AnimMontage;
-						Database->MarkPackageDirty();
 						return true;
 					}
 				}
@@ -159,7 +167,9 @@ bool UPoseSearchDatabaseAnimMontageReflection::ApplyChanges()
 
 bool UPoseSearchDatabaseMultiAnimAssetReflection::ApplyChanges()
 {
-	const float PlayLength = MultiAnimAsset.GetPlayLength();
+	// @todo: add support for InteractionAsset contatining blend spaces
+	FVector BlendParameters = FVector::ZeroVector;
+	const float PlayLength = MultiAnimAsset.GetPlayLength(BlendParameters);
 	const FFloatInterval ClampedSamplingRange(FMath::Clamp(MultiAnimAsset.SamplingRange.Min, 0.0f, PlayLength), FMath::Clamp(MultiAnimAsset.SamplingRange.Max, 0.0f, PlayLength));
 	if (ClampedSamplingRange != MultiAnimAsset.SamplingRange)
 	{
@@ -178,7 +188,6 @@ bool UPoseSearchDatabaseMultiAnimAssetReflection::ApplyChanges()
 					if (*DatabaseMultiAnimAsset != MultiAnimAsset)
 					{
 						*DatabaseMultiAnimAsset = MultiAnimAsset;
-						Database->MarkPackageDirty();
 						return true;
 					}
 				}
diff --git a/Engine/Plugins/Animation/PoseSearch/Source/Editor/Private/PoseSearchDebuggerDatabaseView.cpp b/Engine/Plugins/Animation/PoseSearch/Source/Editor/Private/PoseSearchDebuggerDatabaseView.cpp
--- a/Engine/Plugins/Animation/PoseSearch/Source/Editor/Private/PoseSearchDebuggerDatabaseView.cpp
+++ b/Engine/Plugins/Animation/PoseSearch/Source/Editor/Private/PoseSearchDebuggerDatabaseView.cpp
@@ -220,13 +220,14 @@ static void AddUnfilteredDatabaseRow(const UPoseSearchDatabase* Database, TConst
 
 		if (const FPoseSearchDatabaseAnimationAssetBase* DatabaseAsset = Database->GetDatabaseAnimationAsset<FPoseSearchDatabaseAnimationAssetBase>(*SearchIndexAsset))
 		{
-			const float PlayLength = DatabaseAsset->GetPlayLength();
+			const FVector BlendParameters = SearchIndexAsset->GetBlendParameters();
+			const float PlayLength = DatabaseAsset->GetPlayLength(BlendParameters);
 			UObject* AnimationAsset = DatabaseAsset->GetAnimationAsset();
 
 			Row->AssetName = DatabaseAsset->GetName();
 			Row->AssetPath = AnimationAsset ? AnimationAsset->GetPathName() : "";
 			Row->bLooping = DatabaseAsset->IsLooping();
-			Row->BlendParameters = SearchIndexAsset->GetBlendParameters();
+			Row->BlendParameters = BlendParameters;
 			Row->AnimFrame = DatabaseAsset->GetFrameAtTime(Time);
 			Row->AnimPercentage = FMath::IsNearlyZero(PlayLength) ? 0.f : Time / PlayLength;
 		}
diff --git a/Engine/Plugins/Animation/PoseSearch/Source/Editor/Private/PoseSearchInteractionAssetEditor.cpp b/Engine/Plugins/Animation/PoseSearch/Source/Editor/Private/PoseSearchInteractionAssetEditor.cpp
--- a/Engine/Plugins/Animation/PoseSearch/Source/Editor/Private/PoseSearchInteractionAssetEditor.cpp
+++ b/Engine/Plugins/Animation/PoseSearch/Source/Editor/Private/PoseSearchInteractionAssetEditor.cpp
@@ -226,7 +226,9 @@ void FInteractionAssetViewModel::PreviewBackwardStep()
 {
 	if (const UPoseSearchInteractionAsset* InteractionAsset = GetInteractionAsset())
 	{
-		const float NewPlayTime = FMath::Clamp(PlayTime - StepDeltaTime, 0.f, InteractionAsset->GetPlayLength());
+		// @todo: add support for InteractionAsset contatining blend spaces
+		FVector BlendParameters = FVector::ZeroVector;
+		const float NewPlayTime = FMath::Clamp(PlayTime - StepDeltaTime, 0.f, InteractionAsset->GetPlayLength(BlendParameters));
 		SetPlayTime(NewPlayTime, false);
 	}
 }
@@ -250,7 +252,9 @@ void FInteractionAssetViewModel::PreviewForwardStep()
 {
 	if (const UPoseSearchInteractionAsset* InteractionAsset = GetInteractionAsset())
 	{
-		const float NewPlayTime = FMath::Clamp(PlayTime + StepDeltaTime, 0.f, InteractionAsset->GetPlayLength());
+		// @todo: add support for InteractionAsset contatining blend spaces
+		FVector BlendParameters = FVector::ZeroVector;
+		const float NewPlayTime = FMath::Clamp(PlayTime + StepDeltaTime, 0.f, InteractionAsset->GetPlayLength(BlendParameters));
 		SetPlayTime(NewPlayTime, false);
 	}
 }
@@ -259,7 +263,9 @@ void FInteractionAssetViewModel::PreviewForwardEnd()
 {
 	if (const UPoseSearchInteractionAsset* InteractionAsset = GetInteractionAsset())
 	{
-		SetPlayTime(InteractionAsset->GetPlayLength(), false);
+		// @todo: add support for InteractionAsset contatining blend spaces
+		FVector BlendParameters = FVector::ZeroVector;
+		SetPlayTime(InteractionAsset->GetPlayLength(BlendParameters), false);
 	}
 }
 
@@ -279,7 +285,9 @@ void FInteractionAssetViewModel::Tick(float DeltaSeconds)
 	}
 
 	PlayTime += DeltaSeconds * DeltaTimeMultiplier;
-	PlayTime = FMath::Clamp(PlayTime, 0.f, InteractionAsset->GetPlayLength());
+	// @todo: add support for InteractionAsset contatining blend spaces
+	FVector BlendParameters = FVector::ZeroVector;
+	PlayTime = FMath::Clamp(PlayTime, 0.f, InteractionAsset->GetPlayLength(BlendParameters));
 
 	FRoleToIndex InteractionAssetRoleToIndex;
 	for (int32 RoleIndex = 0; RoleIndex < InteractionAsset->GetNumRoles(); ++RoleIndex)
@@ -384,7 +392,9 @@ TRange<double> FInteractionAssetViewModel::GetPreviewPlayRange() const
 	constexpr double ViewRangeSlack = 0.2;
 	if (const UPoseSearchInteractionAsset* InteractionAsset = GetInteractionAsset())
 	{
-		return TRange<double>(-ViewRangeSlack, InteractionAsset->GetPlayLength() + ViewRangeSlack);
+		// @todo: add support for InteractionAsset contatining blend spaces
+		FVector BlendParameters = FVector::ZeroVector;
+		return TRange<double>(-ViewRangeSlack, InteractionAsset->GetPlayLength(BlendParameters) + ViewRangeSlack);
 	}
 	
 	return TRange<double>(-ViewRangeSlack, ViewRangeSlack);
@@ -394,7 +404,9 @@ void FInteractionAssetViewModel::SetPlayTime(float NewPlayTime, bool bInTickPlay
 {
 	if (const UPoseSearchInteractionAsset* InteractionAsset = GetInteractionAsset())
 	{
-		NewPlayTime = FMath::Clamp(NewPlayTime, 0.f, InteractionAsset->GetPlayLength());
+		// @todo: add support for InteractionAsset contatining blend spaces
+		FVector BlendParameters = FVector::ZeroVector;
+		NewPlayTime = FMath::Clamp(NewPlayTime, 0.f, InteractionAsset->GetPlayLength(BlendParameters));
 		DeltaTimeMultiplier = bInTickPlayTime ? DeltaTimeMultiplier : 0.f;
 
 		if (!FMath::IsNearlyEqual(PlayTime, NewPlayTime))
diff --git a/Engine/Plugins/Animation/PoseSearch/Source/Runtime/Private/PoseSearchDatabase.cpp b/Engine/Plugins/Animation/PoseSearch/Source/Runtime/Private/PoseSearchDatabase.cpp
--- a/Engine/Plugins/Animation/PoseSearch/Source/Runtime/Private/PoseSearchDatabase.cpp
+++ b/Engine/Plugins/Animation/PoseSearch/Source/Runtime/Private/PoseSearchDatabase.cpp
@@ -227,12 +227,80 @@ static inline void EvaluatePoseKernel(UE::PoseSearch::FSearchResult& Result, con
 	}
 }
 
+static void IterateOverBlendSpaceSamplingParameter(const UBlendSpace* BlendSpace, bool bUseSingleSample, const FVector& SingleSampleBlendParameters,
+	bool bUseGridForSampling, int32 NumberOfHorizontalSamples, int32 NumberOfVerticalSamples,
+	const TFunction<void(const FVector& BlendParameters)>& ProcessSamplingParameter)
+{
+	check(BlendSpace);
+	
+	if (bUseSingleSample)
+	{
+		ProcessSamplingParameter(SingleSampleBlendParameters);
+	}
+	else
+	{
+		int32 HorizontalBlendNum, VerticalBlendNum;
+
+		if (bUseGridForSampling)
+		{
+			HorizontalBlendNum = BlendSpace->GetBlendParameter(0).GridNum + 1;
+			VerticalBlendNum = BlendSpace->IsA<UBlendSpace1D>() ? 1 : BlendSpace->GetBlendParameter(1).GridNum + 1;
+		}
+		else
+		{
+			HorizontalBlendNum = FMath::Max(NumberOfHorizontalSamples, 1);
+			VerticalBlendNum = BlendSpace->IsA<UBlendSpace1D>() ? 1 : FMath::Max(NumberOfVerticalSamples, 1);
+		}
+
+		check(HorizontalBlendNum >= 1 && VerticalBlendNum >= 1);
+
+		const FBlendParameter& HorizontalBlendParameter = BlendSpace->GetBlendParameter(0);
+		const FBlendParameter& VerticalBlendParameter = BlendSpace->GetBlendParameter(1);
+
+		const int32 WrapInputHorizontalBlendNum = HorizontalBlendParameter.bWrapInput ? HorizontalBlendNum + 1 : HorizontalBlendNum;
+		const int32 WrapInputVerticalBlendNum = VerticalBlendParameter.bWrapInput ? VerticalBlendNum + 1 : VerticalBlendNum;
+
+		for (int32 HorizontalBlendIndex = 0; HorizontalBlendIndex < HorizontalBlendNum; HorizontalBlendIndex++)
+		{
+			for (int32 VerticalBlendIndex = 0; VerticalBlendIndex < VerticalBlendNum; VerticalBlendIndex++)
+			{
+				const FVector BlendParameters = FVector(
+						WrapInputHorizontalBlendNum > 1 ? 
+							HorizontalBlendParameter.Min + (HorizontalBlendParameter.Max - HorizontalBlendParameter.Min) * ((float)HorizontalBlendIndex) / (WrapInputHorizontalBlendNum - 1) :
+							HorizontalBlendParameter.Min,
+						WrapInputVerticalBlendNum > 1 ?
+							VerticalBlendParameter.Min + (VerticalBlendParameter.Max - VerticalBlendParameter.Min) * ((float)VerticalBlendIndex) / (WrapInputVerticalBlendNum - 1) :
+							VerticalBlendParameter.Min,
+						0.f);
+
+				ProcessSamplingParameter(BlendParameters);
+			}
+		}
+	}
+}
+
+static bool IsBlendSpaceRootMotionEnabled(const UBlendSpace* BlendSpace)
+{
+	check(BlendSpace);
+
+	bool bIsRootMotionUsedInBlendSpace = false;
+	BlendSpace->ForEachImmutableSample([&bIsRootMotionUsedInBlendSpace](const FBlendSample& Sample)
+		{
+			const UAnimSequence* Sequence = Sample.Animation.Get();
+			if (IsValid(Sequence) && Sequence->HasRootMotion())
+			{
+				bIsRootMotionUsedInBlendSpace = true;
+			}
+		});
+
+	return bIsRootMotionUsedInBlendSpace;
+}
+
 } // namespace UE::PoseSearch
 
 //////////////////////////////////////////////////////////////////////////
 // FPoseSearchDatabaseAnimationAssetBase
-
-float FPoseSearchDatabaseAnimationAssetBase::GetPlayLength() const
+float FPoseSearchDatabaseAnimationAssetBase::GetPlayLength(const FVector& BlendParameters) const
 {
 	if (const UAnimationAsset* AnimationAsset = Cast<UAnimationAsset>(GetAnimationAsset()))
 	{
@@ -288,6 +356,14 @@ USkeletalMesh* FPoseSearchDatabaseAnimationAssetBase::GetPreviewMeshForRole(cons
 	return nullptr;
 }
 
+void FPoseSearchDatabaseAnimationAssetBase::IterateOverSamplingParameter(const TFunction<void(const FVector& BlendParameters)>& ProcessSamplingParameter) const
+{
+	if (GetAnimationAsset())
+	{
+		ProcessSamplingParameter(FVector::ZeroVector);
+	}
+}
+
 #endif // WITH_EDITOR
 
 #if WITH_EDITORONLY_DATA
@@ -316,9 +392,14 @@ int64 FPoseSearchDatabaseAnimationAssetBase::GetEditorMemSize() const
 	return EditorMemCount.GetNum();
 }
 
+FFloatInterval FPoseSearchDatabaseAnimationAssetBase::GetEffectiveSamplingRange(const FVector& BlendParameters) const
+{
+	return GetEffectiveSamplingRange(GetPlayLength(BlendParameters), GetSamplingRange());
+}
+
 FFloatInterval FPoseSearchDatabaseAnimationAssetBase::GetEffectiveSamplingRange() const
 {
-	return GetEffectiveSamplingRange(GetPlayLength(), GetSamplingRange());
+	return GetEffectiveSamplingRange(GetPlayLength(FVector::ZeroVector), GetSamplingRange());
 }
 
 FFloatInterval FPoseSearchDatabaseAnimationAssetBase::GetEffectiveSamplingRange(float PlayLength, const FFloatInterval& SamplingRange)
@@ -394,21 +475,40 @@ bool FPoseSearchDatabaseBlendSpace::IsLooping() const
 
 bool FPoseSearchDatabaseBlendSpace::IsRootMotionEnabled() const
 {
-	bool bIsRootMotionUsedInBlendSpace = false;
+	if (BlendSpace)
+	{
+		return UE::PoseSearch::IsBlendSpaceRootMotionEnabled(BlendSpace);
+	}
+	return false;
+}
 
+FFloatInterval FPoseSearchDatabaseBlendSpace::GetEffectiveSamplingRange(const FVector& BlendParameters) const
+{
 	if (BlendSpace)
 	{
-		BlendSpace->ForEachImmutableSample([&bIsRootMotionUsedInBlendSpace](const FBlendSample& Sample)
-			{
-				const UAnimSequence* Sequence = Sample.Animation.Get();
-				if (IsValid(Sequence) && Sequence->HasRootMotion())
-				{
-					bIsRootMotionUsedInBlendSpace = true;
-				}
-			});
+		int32 TriangulationIndex = 0;
+		TArray<FBlendSampleData> BlendSamples;
+		BlendSpace->GetSamplesFromBlendInput(BlendParameters, BlendSamples, TriangulationIndex, true);
+		const float PlayLength = BlendSpace->GetAnimationLengthFromSampleData(BlendSamples);
+
+		// scaling blend space SamplingRange from the space [0, 1] to [0, PlayLength] with PlayLength calculated from the BlendSamples
+		FFloatInterval ScaledSamplingRange = SamplingRange;
+		ScaledSamplingRange.Min *= PlayLength;
+		ScaledSamplingRange.Max *= PlayLength;
+
+		return FPoseSearchDatabaseAnimationAssetBase::GetEffectiveSamplingRange(PlayLength, ScaledSamplingRange);
 	}
 
-	return bIsRootMotionUsedInBlendSpace;
+	return FFloatInterval(0.f, 0.f);
+}
+
+void FPoseSearchDatabaseBlendSpace::IterateOverSamplingParameter(const TFunction<void(const FVector& BlendParameters)>& ProcessSamplingParameter) const
+{
+	if (BlendSpace)
+	{
+		UE::PoseSearch::IterateOverBlendSpaceSamplingParameter(BlendSpace, bUseSingleSample, FVector(BlendParamX, BlendParamY, 0.f),
+			bUseGridForSampling, NumberOfHorizontalSamples, NumberOfVerticalSamples, ProcessSamplingParameter);
+	}
 }
 
 void FPoseSearchDatabaseBlendSpace::GetBlendSpaceParameterSampleRanges(int32& HorizontalBlendNum, int32& VerticalBlendNum) const
@@ -448,7 +548,9 @@ FVector FPoseSearchDatabaseBlendSpace::BlendParameterForSampleRanges(int32 Horiz
 	const bool bWrapInputOnVerticalAxis = BlendSpace->GetBlendParameter(1).bWrapInput;
 
 	int32 HorizontalBlendNum, VerticalBlendNum;
+	PRAGMA_DISABLE_DEPRECATION_WARNINGS
 	GetBlendSpaceParameterSampleRanges(HorizontalBlendNum, VerticalBlendNum);
+	PRAGMA_ENABLE_DEPRECATION_WARNINGS
 
 	if (bWrapInputOnHorizontalAxis)
 	{
@@ -548,9 +650,9 @@ UObject* FPoseSearchDatabaseMultiAnimAsset::GetAnimationAsset() const
 	return MultiAnimAsset.Get();
 }
 
-float FPoseSearchDatabaseMultiAnimAsset::GetPlayLength() const
+float FPoseSearchDatabaseMultiAnimAsset::GetPlayLength(const FVector& BlendParameters) const
 {
-	return MultiAnimAsset ? MultiAnimAsset->GetPlayLength() : 0.f;
+	return MultiAnimAsset ? MultiAnimAsset->GetPlayLength(BlendParameters) : 0.f;
 }
 
 #if WITH_EDITOR
@@ -598,6 +700,38 @@ bool FPoseSearchDatabaseMultiAnimAsset::IsRootMotionEnabled() const
 {
 	return MultiAnimAsset ? MultiAnimAsset->HasRootMotion() : false;
 }
+
+void FPoseSearchDatabaseMultiAnimAsset::IterateOverSamplingParameter(const TFunction<void(const FVector& BlendParameters)>& ProcessSamplingParameter) const
+{
+	if (MultiAnimAsset)
+	{
+		const UBlendSpace* BlendSpace = nullptr;
+		for (int32 RoleIndex = 0; RoleIndex < MultiAnimAsset->GetNumRoles(); ++RoleIndex)
+		{
+			if (const UAnimationAsset* AnimationAsset = MultiAnimAsset->GetAnimationAsset(MultiAnimAsset->GetRole(RoleIndex)))
+			{
+				if (AnimationAsset->GetClass()->IsChildOf(UBlendSpace::StaticClass()))
+				{
+					// @todo: right now we just sample using the first blend space, but we should probably make sure if there're multiple blendspaces they are consistent with each other
+					BlendSpace = Cast<UBlendSpace>(AnimationAsset);
+					check(BlendSpace);
+					break;
+				}
+			}
+		}
+
+		if (BlendSpace)
+		{
+			UE::PoseSearch::IterateOverBlendSpaceSamplingParameter(BlendSpace, false, FVector::ZeroVector,
+				false, NumberOfHorizontalSamples, NumberOfVerticalSamples, ProcessSamplingParameter);
+		}
+		else
+		{
+			ProcessSamplingParameter(FVector::ZeroVector);
+		}
+	}
+}
+
 #endif // WITH_EDITORONLY_DATA
 
 //////////////////////////////////////////////////////////////////////////
diff --git a/Engine/Plugins/Animation/PoseSearch/Source/Runtime/Private/PoseSearchDerivedData.cpp b/Engine/Plugins/Animation/PoseSearch/Source/Runtime/Private/PoseSearchDerivedData.cpp
--- a/Engine/Plugins/Animation/PoseSearch/Source/Runtime/Private/PoseSearchDerivedData.cpp
+++ b/Engine/Plugins/Animation/PoseSearch/Source/Runtime/Private/PoseSearchDerivedData.cpp
@@ -319,42 +319,14 @@ public:
 	}
 };
 
-static void FindValidBlendSpaceIntervals(const FPoseSearchDatabaseBlendSpace* DatabaseBlendSpace, const TArray<FBlendSampleData>& BlendSamples, const FFloatInterval& ExcludeFromDatabaseParameters, TArray<FFloatRange>& ValidRanges)
-{
-	check(DatabaseBlendSpace);
-
-	const float PlayLength = DatabaseBlendSpace->BlendSpace->GetAnimationLengthFromSampleData(BlendSamples);
-
-	const bool bIsLooping = DatabaseBlendSpace->IsLooping();
-
-	// scaling blend space SamplingRange from the space [0, 1] to [0, PlayLength] with PlayLength calculated from the BlendSamples
-	FFloatInterval ScaledSamplingRange = DatabaseBlendSpace->GetSamplingRange();
-	ScaledSamplingRange.Min *= PlayLength;
-	ScaledSamplingRange.Max *= PlayLength;
-
-	const FFloatInterval EffectiveSamplingInterval = FPoseSearchDatabaseAnimationAssetBase::GetEffectiveSamplingRange(PlayLength, ScaledSamplingRange);
-	FFloatRange EffectiveSamplingRange = FFloatRange::Inclusive(EffectiveSamplingInterval.Min, EffectiveSamplingInterval.Max);
-	if (!bIsLooping)
-	{
-		const FFloatRange ExcludeFromDatabaseRange(ExcludeFromDatabaseParameters.Min, PlayLength + ExcludeFromDatabaseParameters.Max);
-		EffectiveSamplingRange = FFloatRange::Intersection(EffectiveSamplingRange, ExcludeFromDatabaseRange);
-	}
-
-	// start from a single interval defined by the database sequence sampling range
-	ValidRanges.Reset();
-	ValidRanges.Add(EffectiveSamplingRange);
-
-	// @todo: evaluate eventual UAnimNotifyState_PoseSearchExcludeFromDatabase(s) like in FindValidSequenceIntervals if necessary
-}
-
-static void FindValidSequenceIntervals(const FPoseSearchDatabaseAnimationAssetBase* DatabaseAsset, const FFloatInterval& ExcludeFromDatabaseParameters, TArray<FFloatRange>& ValidRanges)
+static void FindValidRanges(const FPoseSearchDatabaseAnimationAssetBase* DatabaseAsset, const FVector& BlendParameters, const FFloatInterval& ExcludeFromDatabaseParameters, TArray<FFloatRange>& ValidRanges)
 {
 	check(DatabaseAsset);
 
 	const bool bIsLooping = DatabaseAsset->IsLooping();
-	const float PlayLength = DatabaseAsset->GetPlayLength();
+	const float PlayLength = DatabaseAsset->GetPlayLength(BlendParameters);
 
-	const FFloatInterval EffectiveSamplingInterval = DatabaseAsset->GetEffectiveSamplingRange();
+	const FFloatInterval EffectiveSamplingInterval = DatabaseAsset->GetEffectiveSamplingRange(BlendParameters);
 	FFloatRange EffectiveSamplingRange = FFloatRange::Inclusive(EffectiveSamplingInterval.Min, EffectiveSamplingInterval.Max);
 	if (!bIsLooping)
 	{
@@ -408,7 +380,6 @@ static bool InitSearchIndexAssets(FSearchIndexBase& SearchIndex, const UPoseSear
 
 	SearchIndex.Assets.Reset();
 	TArray<FFloatRange> ValidRanges;
-	TArray<FBlendSampleData> BlendSamples;
 
 	bool bAnyErrors = false;
 
@@ -458,91 +429,42 @@ static bool InitSearchIndexAssets(FSearchIndexBase& SearchIndex, const UPoseSear
 			const bool bIsLooping = DatabaseAsset->IsLooping();
 			const bool bDisableReselection = DatabaseAsset->IsDisableReselection();
 
-			if (const FPoseSearchDatabaseBlendSpace* DatabaseBlendSpace = DatabaseToLookForAssets->GetDatabaseAnimationAsset<FPoseSearchDatabaseBlendSpace>(AnimationAssetIndex))
-			{
-				int32 HorizontalBlendNum, VerticalBlendNum;
-				DatabaseBlendSpace->GetBlendSpaceParameterSampleRanges(HorizontalBlendNum, VerticalBlendNum);
-
-				const bool bWrapInputOnHorizontalAxis = DatabaseBlendSpace->BlendSpace->GetBlendParameter(0).bWrapInput;
-				const bool bWrapInputOnVerticalAxis = DatabaseBlendSpace->BlendSpace->GetBlendParameter(1).bWrapInput;
-				for (int32 HorizontalIndex = 0; HorizontalIndex < HorizontalBlendNum; HorizontalIndex++)
+			DatabaseAsset->IterateOverSamplingParameter(
+				[DatabaseAsset, ExcludeFromDatabaseParameters, &ValidRanges, Schema, bAddUnmirrored, bAddMirrored, bIsLooping, bDisableReselection, AnimationAssetIndex, &TotalPoses, &SearchIndex]
+				(const FVector& BlendParameters)
 				{
-					for (int32 VerticalIndex = 0; VerticalIndex < VerticalBlendNum; VerticalIndex++)
-					{
-						const FVector BlendParameters = DatabaseBlendSpace->BlendParameterForSampleRanges(HorizontalIndex, VerticalIndex);
-
-						int32 TriangulationIndex = 0;
-						DatabaseBlendSpace->BlendSpace->GetSamplesFromBlendInput(BlendParameters, BlendSamples, TriangulationIndex, true);
-
-						FindValidBlendSpaceIntervals(DatabaseBlendSpace, BlendSamples, ExcludeFromDatabaseParameters, ValidRanges);
+					FindValidRanges(DatabaseAsset, BlendParameters, ExcludeFromDatabaseParameters, ValidRanges);
 
-						for (const FFloatRange& Range : ValidRanges)
+					for (const FFloatRange& Range : ValidRanges)
+					{
+						for (int32 PermutationIdx = 0; PermutationIdx < Schema->NumberOfPermutations; ++PermutationIdx)
 						{
-							for (int32 PermutationIdx = 0; PermutationIdx < Schema->NumberOfPermutations; ++PermutationIdx)
-							{
-								const FFloatInterval RangeInterval(Range.GetLowerBoundValue(), Range.GetUpperBoundValue());
+							const FFloatInterval RangeInterval(Range.GetLowerBoundValue(), Range.GetUpperBoundValue());
 
-								if (bAddUnmirrored)
-								{
-									const FSearchIndexAsset PoseSearchIndexAsset(AnimationAssetIndex, TotalPoses, false, bIsLooping,
-										bDisableReselection, RangeInterval, Schema->SampleRate, PermutationIdx, BlendParameters);
-									if (PoseSearchIndexAsset.GetNumPoses() > 0)
-									{
-										SearchIndex.Assets.Add(PoseSearchIndexAsset);
-										TotalPoses += PoseSearchIndexAsset.GetNumPoses();
-									}
-								}
-
-								if (bAddMirrored)
+							if (bAddUnmirrored)
+							{
+								const FSearchIndexAsset PoseSearchIndexAsset(AnimationAssetIndex, TotalPoses, false, bIsLooping,
+									bDisableReselection, RangeInterval, Schema->SampleRate, PermutationIdx, BlendParameters);
+								if (PoseSearchIndexAsset.GetNumPoses() > 0)
 								{
-									const FSearchIndexAsset PoseSearchIndexAsset(AnimationAssetIndex, TotalPoses, true, bIsLooping,
-										bDisableReselection, RangeInterval, Schema->SampleRate, PermutationIdx, BlendParameters);
-									if (PoseSearchIndexAsset.GetNumPoses() > 0)
-									{
-										SearchIndex.Assets.Add(PoseSearchIndexAsset);
-										TotalPoses += PoseSearchIndexAsset.GetNumPoses();
-									}
+									SearchIndex.Assets.Add(PoseSearchIndexAsset);
+									TotalPoses += PoseSearchIndexAsset.GetNumPoses();
 								}
 							}
-						}
-					}
-				}
-			}
-			// support for FPoseSearchDatabaseSequence, FPoseSearchDatabaseAnimComposite, FPoseSearchDatabaseAnimMontage, FPoseSearchDatabaseMultiAnimAsset
-			else
-			{
-				FindValidSequenceIntervals(DatabaseAsset, ExcludeFromDatabaseParameters, ValidRanges);
 
-				for (const FFloatRange& Range : ValidRanges)
-				{
-					for (int32 PermutationIdx = 0; PermutationIdx < Schema->NumberOfPermutations; ++PermutationIdx)
-					{
-						const FFloatInterval RangeInterval(Range.GetLowerBoundValue(), Range.GetUpperBoundValue());
-
-						if (bAddUnmirrored)
-						{
-							const FSearchIndexAsset PoseSearchIndexAsset(AnimationAssetIndex, TotalPoses, false, bIsLooping,
-								bDisableReselection, RangeInterval, Schema->SampleRate, PermutationIdx);
-							if (PoseSearchIndexAsset.GetNumPoses() > 0)
+							if (bAddMirrored)
 							{
-								SearchIndex.Assets.Add(PoseSearchIndexAsset);
-								TotalPoses += PoseSearchIndexAsset.GetNumPoses();
-							}
-						}
-
-						if (bAddMirrored)
-						{
-							const FSearchIndexAsset PoseSearchIndexAsset(AnimationAssetIndex, TotalPoses, true, bIsLooping,
-								bDisableReselection, RangeInterval, Schema->SampleRate, PermutationIdx);
-							if (PoseSearchIndexAsset.GetNumPoses() > 0)
-							{
-								SearchIndex.Assets.Add(PoseSearchIndexAsset);
-								TotalPoses += PoseSearchIndexAsset.GetNumPoses();
+								const FSearchIndexAsset PoseSearchIndexAsset(AnimationAssetIndex, TotalPoses, true, bIsLooping,
+									bDisableReselection, RangeInterval, Schema->SampleRate, PermutationIdx, BlendParameters);
+								if (PoseSearchIndexAsset.GetNumPoses() > 0)
+								{
+									SearchIndex.Assets.Add(PoseSearchIndexAsset);
+									TotalPoses += PoseSearchIndexAsset.GetNumPoses();
+								}
 							}
 						}
 					}
-				}
-			}
+				});
 		}
 	}
 
@@ -986,49 +908,24 @@ static bool IndexDatabase(FSearchIndexBase& SearchIndexBase, const UPoseSearchDa
 
 	for (int32 AnimationAssetIndex = 0; AnimationAssetIndex < DatabaseToLookForAssets->GetNumAnimationAssets(); ++AnimationAssetIndex)
 	{
-		if (const FPoseSearchDatabaseBlendSpace* DatabaseBlendSpace = DatabaseToLookForAssets->GetDatabaseAnimationAsset<FPoseSearchDatabaseBlendSpace>(AnimationAssetIndex))
+		if (const FPoseSearchDatabaseAnimationAssetBase* DatabaseAnimationAssetBase = DatabaseToLookForAssets->GetDatabaseAnimationAsset<FPoseSearchDatabaseAnimationAssetBase>(AnimationAssetIndex))
 		{
-			if (DatabaseBlendSpace->BlendSpace)
-			{
-				int32 HorizontalBlendNum, VerticalBlendNum;
-				DatabaseBlendSpace->GetBlendSpaceParameterSampleRanges(HorizontalBlendNum, VerticalBlendNum);
-
-				for (int32 HorizontalIndex = 0; HorizontalIndex < HorizontalBlendNum; HorizontalIndex++)
+			DatabaseAnimationAssetBase->IterateOverSamplingParameter([DatabaseAnimationAssetBase, &SamplerMap, &Samplers](const FVector& BlendParameters)
 				{
-					for (int32 VerticalIndex = 0; VerticalIndex < VerticalBlendNum; VerticalIndex++)
+					const int32 NumRoles = DatabaseAnimationAssetBase->GetNumRoles();
+					for (int32 RoleIndex = 0; RoleIndex < NumRoles; ++RoleIndex)
 					{
-						const FVector BlendParameters = DatabaseBlendSpace->BlendParameterForSampleRanges(HorizontalIndex, VerticalIndex);
-
-						check(DatabaseBlendSpace->GetNumRoles() == 1);
-						const FTransform RootTransformOrigin = DatabaseBlendSpace->GetRootTransformOriginForRole(DatabaseBlendSpace->GetRole(0));
-						const FSamplerMapKey SamplerMapKey(DatabaseBlendSpace->BlendSpace, RootTransformOrigin, BlendParameters);
+						const FRole& Role = DatabaseAnimationAssetBase->GetRole(RoleIndex);
+						const UAnimationAsset* AnimationAsset = DatabaseAnimationAssetBase->GetAnimationAssetForRole(Role);
+						const FTransform RootTransformOrigin = DatabaseAnimationAssetBase->GetRootTransformOriginForRole(Role);
+						const FSamplerMapKey SamplerMapKey(AnimationAsset, RootTransformOrigin, BlendParameters);
 						if (!SamplerMap.Contains(SamplerMapKey))
 						{
 							SamplerMap.Add(SamplerMapKey, Samplers.Num());
-							Samplers.Emplace(DatabaseBlendSpace->BlendSpace, RootTransformOrigin, BlendParameters, FAnimationAssetSampler::DefaultRootTransformSamplingRate, false);
+							Samplers.Emplace(AnimationAsset, RootTransformOrigin, BlendParameters, FAnimationAssetSampler::DefaultRootTransformSamplingRate, false);
 						}
 					}
-				}
-			}
-		}
-		if (const FPoseSearchDatabaseAnimationAssetBase* DatabaseAnimationAssetBase = DatabaseToLookForAssets->GetDatabaseAnimationAsset<FPoseSearchDatabaseAnimationAssetBase>(AnimationAssetIndex))
-		{
-			if (const UObject* AnimationAssetObject = DatabaseAnimationAssetBase->GetAnimationAsset())
-			{
-				const int32 NumRoles = DatabaseAnimationAssetBase->GetNumRoles();
-				for (int32 RoleIndex = 0; RoleIndex < NumRoles; ++RoleIndex)
-				{
-					const FRole& Role = DatabaseAnimationAssetBase->GetRole(RoleIndex);
-					const UAnimationAsset* AnimationAsset = DatabaseAnimationAssetBase->GetAnimationAssetForRole(Role);
-					const FTransform RootTransformOrigin = DatabaseAnimationAssetBase->GetRootTransformOriginForRole(Role);
-					const FSamplerMapKey SamplerMapKey(AnimationAsset, RootTransformOrigin);
-					if (!SamplerMap.Contains(SamplerMapKey))
-					{
-						SamplerMap.Add(SamplerMapKey, Samplers.Num());
-						Samplers.Emplace(AnimationAsset, RootTransformOrigin, FVector::ZeroVector, FAnimationAssetSampler::DefaultRootTransformSamplingRate, false);
-					}
-				}
-			}
+				});
 		}
 	}
 
diff --git a/Engine/Plugins/Animation/PoseSearch/Source/Runtime/Private/PoseSearchInteractionAsset.cpp b/Engine/Plugins/Animation/PoseSearch/Source/Runtime/Private/PoseSearchInteractionAsset.cpp
--- a/Engine/Plugins/Animation/PoseSearch/Source/Runtime/Private/PoseSearchInteractionAsset.cpp
+++ b/Engine/Plugins/Animation/PoseSearch/Source/Runtime/Private/PoseSearchInteractionAsset.cpp
@@ -79,20 +79,32 @@ bool UPoseSearchInteractionAsset::HasRootMotion() const
 	return bHasAtLeastOneValidItem && bHasRootMotion;
 }
 
-float UPoseSearchInteractionAsset::GetPlayLength() const
+float UPoseSearchInteractionAsset::GetPlayLength(const FVector& BlendParameters) const
 {
 	float MaxPlayLength = 0.f;
 	for (const FPoseSearchInteractionAssetItem& Item : Items)
 	{
 		if (const UAnimationAsset* AnimationAsset = Item.Animation.Get())
 		{
-			MaxPlayLength = FMath::Max(MaxPlayLength, AnimationAsset->GetPlayLength());
+			if (const UBlendSpace* BlendSpace = Cast<UBlendSpace>(AnimationAsset))
+			{
+				int32 TriangulationIndex = 0;
+				TArray<FBlendSampleData> BlendSamples;
+				BlendSpace->GetSamplesFromBlendInput(BlendParameters, BlendSamples, TriangulationIndex, true);
+				const float PlayLength = BlendSpace->GetAnimationLengthFromSampleData(BlendSamples);
+				MaxPlayLength = FMath::Max(MaxPlayLength, PlayLength);
+			}
+			else
+			{
+				MaxPlayLength = FMath::Max(MaxPlayLength, AnimationAsset->GetPlayLength());
+			}
 		}
 	}
 	return MaxPlayLength;
 }
 
 #if WITH_EDITOR
+// @todo: pass const FVector& BlendParameters as input and use a similar logic to GetPlayLength to calculate BS PlayLength
 int32 UPoseSearchInteractionAsset::GetFrameAtTime(float Time) const
 {
 	const UAnimationAsset* MaxPlayLengthAnim = nullptr;
diff --git a/Engine/Plugins/Animation/PoseSearch/Source/Runtime/Public/PoseSearch/MultiAnimAsset.h b/Engine/Plugins/Animation/PoseSearch/Source/Runtime/Public/PoseSearch/MultiAnimAsset.h
--- a/Engine/Plugins/Animation/PoseSearch/Source/Runtime/Public/PoseSearch/MultiAnimAsset.h
+++ b/Engine/Plugins/Animation/PoseSearch/Source/Runtime/Public/PoseSearch/MultiAnimAsset.h
@@ -19,7 +19,7 @@ public:
 
 	[[nodiscard]] virtual bool IsLooping() const PURE_VIRTUAL(UMultiAnimAsset::IsLooping, return false;);
 	[[nodiscard]] virtual bool HasRootMotion() const PURE_VIRTUAL(UMultiAnimAsset::HasRootMotion, return false;);
-	[[nodiscard]] virtual float GetPlayLength() const PURE_VIRTUAL(UMultiAnimAsset::GetPlayLength, return 0.f;);
+	[[nodiscard]] virtual float GetPlayLength(const FVector& BlendParameters) const PURE_VIRTUAL(UMultiAnimAsset::GetPlayLength, return 0.f;);
 
 #if WITH_EDITOR
 	[[nodiscard]] virtual int32 GetFrameAtTime(float Time) const PURE_VIRTUAL(UMultiAnimAsset::GetFrameAtTime, return 0;);
diff --git a/Engine/Plugins/Animation/PoseSearch/Source/Runtime/Public/PoseSearch/PoseSearchDatabase.h b/Engine/Plugins/Animation/PoseSearch/Source/Runtime/Public/PoseSearch/PoseSearchDatabase.h
--- a/Engine/Plugins/Animation/PoseSearch/Source/Runtime/Public/PoseSearch/PoseSearchDatabase.h
+++ b/Engine/Plugins/Animation/PoseSearch/Source/Runtime/Public/PoseSearch/PoseSearchDatabase.h
@@ -59,7 +59,11 @@ struct POSESEARCH_API FPoseSearchDatabaseAnimationAssetBase
 	GENERATED_BODY()
 	virtual ~FPoseSearchDatabaseAnimationAssetBase() = default;
 	virtual UObject* GetAnimationAsset() const { return nullptr; }
-	virtual float GetPlayLength() const;
+	
+	UE_DEPRECATED(5.6, "Use other signatures")
+	virtual float GetPlayLength() const { return GetPlayLength(FVector::ZeroVector); }
+	virtual float GetPlayLength(const FVector& BlendParameters) const;
+
 	virtual int32 GetNumRoles() const { return 1; }
 	virtual UE::PoseSearch::FRole GetRole(int32 RoleIndex) const { return UE::PoseSearch::DefaultRole; }
 	virtual UAnimationAsset* GetAnimationAssetForRole(const UE::PoseSearch::FRole& Role) const;
@@ -70,6 +74,7 @@ struct POSESEARCH_API FPoseSearchDatabaseAnimationAssetBase
 	virtual bool IsSkeletonCompatible(TObjectPtr<const UPoseSearchSchema> InSchema) const;
 	// Experimental, this feature might be removed without warning, not for production use
 	virtual USkeletalMesh* GetPreviewMeshForRole(const UE::PoseSearch::FRole& Role) const;
+	virtual void IterateOverSamplingParameter(const TFunction<void(const FVector& BlendParameters)>& ProcessSamplingParameter) const;
 #endif // WITH_EDITOR
 
 #if WITH_EDITORONLY_DATA
@@ -85,6 +90,9 @@ struct POSESEARCH_API FPoseSearchDatabaseAnimationAssetBase
 
 	// [0, 0] represents the entire frame range of the original animation.
 	virtual FFloatInterval GetSamplingRange() const { return FFloatInterval(0.f, 0.f); }
+	virtual FFloatInterval GetEffectiveSamplingRange(const FVector& BlendParameters) const;
+	
+	UE_DEPRECATED(5.6, "Use other signatures")
 	FFloatInterval GetEffectiveSamplingRange() const;
 	static FFloatInterval GetEffectiveSamplingRange(float PlayLength, const FFloatInterval& SamplingRange);
 
@@ -219,8 +227,12 @@ struct POSESEARCH_API FPoseSearchDatabaseBlendSpace : public FPoseSearchDatabase
 	virtual bool IsLooping() const override;
 	virtual bool IsRootMotionEnabled() const override;
 	virtual FFloatInterval GetSamplingRange() const override { return SamplingRange; }
+	virtual FFloatInterval GetEffectiveSamplingRange(const FVector& BlendParameters) const override;
+	virtual void IterateOverSamplingParameter(const TFunction<void(const FVector& BlendParameters)>& ProcessSamplingParameter) const override;
 
+	UE_DEPRECATED(5.6, "Use IterateOverSamplingParameter to iterate over all the sampling paramters for the blend space")
 	void GetBlendSpaceParameterSampleRanges(int32& HorizontalBlendNum, int32& VerticalBlendNum) const;
+	UE_DEPRECATED(5.6, "Use IterateOverSamplingParameter to iterate over all the sampling paramters for the blend space")
 	FVector BlendParameterForSampleRanges(int32 HorizontalBlendIndex, int32 VerticalBlendIndex) const;
 #endif // WITH_EDITORONLY_DATA
 
@@ -348,6 +360,14 @@ struct POSESEARCH_API FPoseSearchDatabaseMultiAnimAsset : public FPoseSearchData
 	TObjectPtr<UMultiAnimAsset> MultiAnimAsset;
 
 #if WITH_EDITORONLY_DATA
+	// Sets the number of horizontal samples in the blend spaces referenced by the UMultiAnimAsset to pull the animation data coverage from. The larger the samples the more the data, but also the more memory and performance it takes.
+	UPROPERTY(EditAnywhere, Category = "Settings", meta = (ClampMin = "1", UIMin = "1", UIMax = "25", DisplayPriority = 6))
+	int32 NumberOfHorizontalSamples = 1;
+	
+	// Sets the number of vertical samples for the blend spaces referenced by the UMultiAnimAsset to pull the animation data coverage from. The larger the samples the more the data, but also the more memory and performance it takes.
+	UPROPERTY(EditAnywhere, Category = "Settings", meta = (ClampMin = "1", UIMin = "1", UIMax = "25", DisplayPriority = 7))
+	int32 NumberOfVerticalSamples = 1;
+
 	// It allows users to set a time range to an individual UMultiAnimAsset in the database. 
 	// This is effectively trimming the beginning and end of the animation in the database (not in the original UMultiAnimAsset).
 	// If set to [0, 0] it will be the entire frame range of the original UMultiAnimAsset.
@@ -358,6 +378,7 @@ struct POSESEARCH_API FPoseSearchDatabaseMultiAnimAsset : public FPoseSearchData
 	virtual bool IsLooping() const override;
 	virtual bool IsRootMotionEnabled() const override;
 	virtual FFloatInterval GetSamplingRange() const override { return SamplingRange; }
+	virtual void IterateOverSamplingParameter(const TFunction<void(const FVector& BlendParameters)>& ProcessSamplingParameter) const override;
 #endif // WITH_EDITORONLY_DATA
 
 #if WITH_EDITOR
@@ -365,8 +386,7 @@ struct POSESEARCH_API FPoseSearchDatabaseMultiAnimAsset : public FPoseSearchData
 #endif // WITH_EDITOR
 
 	virtual UObject* GetAnimationAsset() const override;
-	virtual float GetPlayLength() const override;
-
+	virtual float GetPlayLength(const FVector& BlendParameters) const override;
 	virtual int32 GetNumRoles() const override;
 	virtual UE::PoseSearch::FRole GetRole(int32 RoleIndex) const override;
 	virtual UAnimationAsset* GetAnimationAssetForRole(const UE::PoseSearch::FRole& Role) const override;
@@ -381,6 +401,8 @@ struct POSESEARCH_API FPoseSearchDatabaseMultiAnimAsset : public FPoseSearchData
 		return static_cast<const FPoseSearchDatabaseAnimationAssetBase&>(A) == static_cast<const FPoseSearchDatabaseAnimationAssetBase&>(B)
 			&& A.MultiAnimAsset == B.MultiAnimAsset
 #if WITH_EDITORONLY_DATA
+			&& A.NumberOfHorizontalSamples == B.NumberOfHorizontalSamples
+			&& A.NumberOfVerticalSamples == B.NumberOfVerticalSamples
 			&& A.SamplingRange == B.SamplingRange
 #endif // WITH_EDITORONLY_DATA
 			;
diff --git a/Engine/Plugins/Animation/PoseSearch/Source/Runtime/Public/PoseSearch/PoseSearchInteractionAsset.h b/Engine/Plugins/Animation/PoseSearch/Source/Runtime/Public/PoseSearch/PoseSearchInteractionAsset.h
--- a/Engine/Plugins/Animation/PoseSearch/Source/Runtime/Public/PoseSearch/PoseSearchInteractionAsset.h
+++ b/Engine/Plugins/Animation/PoseSearch/Source/Runtime/Public/PoseSearch/PoseSearchInteractionAsset.h
@@ -74,7 +74,7 @@ public:
 
 	virtual bool IsLooping() const override;
 	virtual bool HasRootMotion() const override;
-	virtual float GetPlayLength() const override;
+	virtual float GetPlayLength(const FVector& BlendParameters) const override;
 
 #if WITH_EDITOR
 	virtual int32 GetFrameAtTime(float Time) const override;
diff --git a/Engine/Plugins/Experimental/AnimNextPoseSearch/Source/AnimNextPoseSearch/Private/MotionMatchingTrait.cpp b/Engine/Plugins/Experimental/AnimNextPoseSearch/Source/AnimNextPoseSearch/Private/MotionMatchingTrait.cpp
--- a/Engine/Plugins/Experimental/AnimNextPoseSearch/Source/AnimNextPoseSearch/Private/MotionMatchingTrait.cpp
+++ b/Engine/Plugins/Experimental/AnimNextPoseSearch/Source/AnimNextPoseSearch/Private/MotionMatchingTrait.cpp
@@ -2,18 +2,19 @@
 
 #include "MotionMatchingTrait.h"
 
-#include "AnimNextAnimGraphSettings.h"
 #include "Animation/AnimSequence.h"
 #include "Animation/BlendSpace.h"
+#include "AnimNextAnimGraphSettings.h"
 #include "DataInterface/DataInterfaceStructAdapter.h"
-#include "IPoseHistory.h"
 #include "Graph/AnimNextGraphInstance.h"
 #include "GraphInterfaces/AnimNextNativeDataInterface_AnimSequencePlayer.h"
 #include "GraphInterfaces/AnimNextNativeDataInterface_BlendSpacePlayer.h"
+#include "IPoseHistory.h"
 #include "Module/AnimNextModuleInstance.h"
 #include "PoseSearch/MultiAnimAsset.h"
 #include "PoseSearch/PoseSearchDatabase.h"
 #include "PoseSearch/PoseSearchDerivedData.h"
+#include "PoseSearch/PoseSearchSchema.h"
 #include "TraitInterfaces/IBlendStack.h"
 #include "TraitInterfaces/IGraphFactory.h"
 #include "TraitInterfaces/ITimeline.h"
@@ -247,9 +248,10 @@ namespace UE::AnimNext
 			{
 				if (const UE::PoseSearch::FSearchIndexAsset* SearchIndexAsset = MotionMatchingState.CurrentSearchResult.GetSearchIndexAsset())
 				{
-					if (const FPoseSearchDatabaseAnimationAssetBase* DatabaseAnimationAssetBase = MotionMatchingState.CurrentSearchResult.Database->GetDatabaseAnimationAsset<FPoseSearchDatabaseAnimationAssetBase>(*SearchIndexAsset))
+					const UPoseSearchDatabase* Database = MotionMatchingState.CurrentSearchResult.Database.Get();
+					if (const FPoseSearchDatabaseAnimationAssetBase* DatabaseAnimationAssetBase = Database->GetDatabaseAnimationAsset<FPoseSearchDatabaseAnimationAssetBase>(*SearchIndexAsset))
 					{
-						AnimationAsset = DatabaseAnimationAssetBase->GetAnimationAsset();
+						AnimationAsset = DatabaseAnimationAssetBase->GetAnimationAssetForRole(Database->Schema->GetDefaultRole());
 						AssetTime = MotionMatchingState.CurrentSearchResult.AssetTime;
 						bLoop = SearchIndexAsset->IsLooping();
 						BlendSpaceParameters = SearchIndexAsset->GetBlendParameters();
@@ -279,7 +281,6 @@ namespace UE::AnimNext
 					NewPayload.InitializeAs<FAnimNextNativeDataInterface_BlendSpacePlayer>();
 					FAnimNextNativeDataInterface_BlendSpacePlayer& NewNativeInterface = NewPayload.GetMutable<FAnimNextNativeDataInterface_BlendSpacePlayer>();
 					NewNativeInterface.BlendSpace = Cast<UBlendSpace>(AnimationAsset);
-
 					
 					NewNativeInterface.XAxisSamplePoint = BlendSpaceParameters.X;
 					NewNativeInterface.YAxisSamplePoint = BlendSpaceParameters.Y;
