PoseSearch - SearchEvent improvements

diff --git a/Engine/Plugins/Animation/PoseSearch/Source/Editor/Private/PoseSearchDatabaseEditorReflection.cpp b/Engine/Plugins/Animation/PoseSearch/Source/Editor/Private/PoseSearchDatabaseEditorReflection.cpp
--- a/Engine/Plugins/Animation/PoseSearch/Source/Editor/Private/PoseSearchDatabaseEditorReflection.cpp
+++ b/Engine/Plugins/Animation/PoseSearch/Source/Editor/Private/PoseSearchDatabaseEditorReflection.cpp
@@ -249,11 +249,12 @@ void UPoseSearchDatabaseStatistics::Initialize(const UPoseSearchDatabase* PoseSe
 		// Kinematic Information
 	
 		// using FText instead of meta = (ForceUnits = "cm/s") to keep properties consistent
+		PRAGMA_DISABLE_DEPRECATION_WARNINGS
 		AverageSpeed = FText::Format(LOCTEXT("StatsAverageSpeed", "{0} cm/s"), SearchIndex.Stats.AverageSpeed);
 		MaxSpeed = FText::Format(LOCTEXT("StatsMaxSpeed", "{0} cm/s"), SearchIndex.Stats.MaxSpeed);
 		AverageAcceleration = FText::Format(LOCTEXT("StatsAverageAcceleration", "{0} cm/s²"), SearchIndex.Stats.AverageAcceleration);
 		MaxAcceleration = FText::Format(LOCTEXT("StatsMaxAcceleration", "{0} cm/s²"), SearchIndex.Stats.MaxAcceleration);
-
+		PRAGMA_ENABLE_DEPRECATION_WARNINGS
 		// Principal Component Analysis
 			
 		ExplainedVariance = SearchIndex.PCAExplainedVarianceEditorOnly * 100.f;
diff --git a/Engine/Plugins/Animation/PoseSearch/Source/Runtime/Private/KDTree.cpp b/Engine/Plugins/Animation/PoseSearch/Source/Runtime/Private/KDTree.cpp
--- a/Engine/Plugins/Animation/PoseSearch/Source/Runtime/Private/KDTree.cpp
+++ b/Engine/Plugins/Animation/PoseSearch/Source/Runtime/Private/KDTree.cpp
@@ -146,8 +146,8 @@ private:
 #endif
 
 FKDTree::FKDTree(AccessorType Count, AccessorType Dim, const float* Data, AccessorType MaxLeafSize)
-: DataSource(Count, Dim, Data)
-, KDTreeImplementation(nullptr)
+	: DataSource(Count, Dim, Data)
+	, KDTreeImplementation(nullptr)
 {
 #if UE_POSE_SEARCH_USE_NANOFLANN
 	if (Count > 0 && Dim > 0 && Data)
@@ -199,6 +199,8 @@ void CopySubTree(FKDTree& KDTree, FKDTreeImplementation::NodePtr& ThisNode, cons
 #endif // UE_POSE_SEARCH_USE_NANOFLANN
 
 FKDTree::FKDTree(const FKDTree& Other)
+	: DataSource()
+	, KDTreeImplementation(nullptr)
 {
 #if UE_POSE_SEARCH_USE_NANOFLANN
 	if (this != &Other && Other.KDTreeImplementation)
@@ -230,24 +232,85 @@ FKDTree::FKDTree(const FKDTree& Other)
 			check(Other.KDTreeImplementation->vAcc.size() <= AccessorTypeMax);
 			const AccessorType VAccSize = Other.KDTreeImplementation->vAcc.size();
 			KDTreeImplementation->vAcc.resize(VAccSize);
-			
+
 			for (AccessorType i = 0; i < VAccSize; ++i)
 			{
 				KDTreeImplementation->vAcc[i] = Other.KDTreeImplementation->vAcc[i];
 			}
-			
+
 			CopySubTree(*this, KDTreeImplementation->root_node, Other.KDTreeImplementation->root_node);
 		}
 	}
 #endif // UE_POSE_SEARCH_USE_NANOFLANN
 }
 
+FKDTree::FKDTree(FKDTree&& Other)
+	: DataSource()
+	, KDTreeImplementation(nullptr)
+{
+#if UE_POSE_SEARCH_USE_NANOFLANN
+	if (this != &Other)
+	{
+		if (Other.KDTreeImplementation)
+		{
+			KDTreeImplementation = new FKDTreeImplementation(0, DataSource, nanoflann::KDTreeSingleIndexAdaptorParams(0));
+
+			DataSource = MoveTemp(Other.DataSource);
+
+			check(Other.KDTreeImplementation->m_size <= AccessorTypeMax);
+			KDTreeImplementation->m_size = Other.KDTreeImplementation->m_size;
+
+			if (KDTreeImplementation->m_size > 0)
+			{
+				KDTreeImplementation->dim = Other.KDTreeImplementation->dim;
+
+				check(Other.KDTreeImplementation->root_bbox.size() <= AccessorTypeMax);
+				const AccessorType root_bbox_size = Other.KDTreeImplementation->root_bbox.size();
+				KDTreeImplementation->root_bbox.resize(root_bbox_size);
+
+				for (AccessorType i = 0; i < root_bbox_size; ++i)
+				{
+					KDTreeImplementation->root_bbox[i] = Other.KDTreeImplementation->root_bbox[i];
+				}
+
+				check(Other.KDTreeImplementation->m_leaf_max_size <= AccessorTypeMax);
+				const AccessorType KDTreeLeafMaxSize = Other.KDTreeImplementation->m_leaf_max_size;
+				KDTreeImplementation->m_leaf_max_size = KDTreeLeafMaxSize;
+
+				check(Other.KDTreeImplementation->vAcc.size() <= AccessorTypeMax);
+				const AccessorType VAccSize = Other.KDTreeImplementation->vAcc.size();
+				KDTreeImplementation->vAcc.resize(VAccSize);
+
+				for (AccessorType i = 0; i < VAccSize; ++i)
+				{
+					KDTreeImplementation->vAcc[i] = Other.KDTreeImplementation->vAcc[i];
+				}
+
+				CopySubTree(*this, KDTreeImplementation->root_node, Other.KDTreeImplementation->root_node);
+			}
+		}
+
+		Other.Reset();
+	}
+#endif // UE_POSE_SEARCH_USE_NANOFLANN
+}
+
 FKDTree& FKDTree::operator=(const FKDTree& Other)
 {
 	if (this != &Other)
 	{
-		Reset();
-		new(this)FKDTree(Other);
+		this->~FKDTree();
+		new(this) FKDTree(Other);
+	}
+	return *this;
+}
+
+FKDTree& FKDTree::operator=(FKDTree&& Other)
+{
+	if (this != &Other)
+	{
+		this->~FKDTree();
+		new(this) FKDTree(MoveTemp(Other));
 	}
 	return *this;
 }
@@ -290,8 +353,8 @@ void FKDTree::Reset()
 
 void FKDTree::Construct(AccessorType Count, AccessorType Dim, const float* Data, AccessorType MaxLeafSize)
 {
-	Reset();
-	new(this)FKDTree(Count, Dim, Data, MaxLeafSize);
+	this->~FKDTree();
+	new(this) FKDTree(Count, Dim, Data, MaxLeafSize);
 }
 
 template <typename RESULTSET>
diff --git a/Engine/Plugins/Animation/PoseSearch/Source/Runtime/Private/PoseSearchDatabase.cpp b/Engine/Plugins/Animation/PoseSearch/Source/Runtime/Private/PoseSearchDatabase.cpp
--- a/Engine/Plugins/Animation/PoseSearch/Source/Runtime/Private/PoseSearchDatabase.cpp
+++ b/Engine/Plugins/Animation/PoseSearch/Source/Runtime/Private/PoseSearchDatabase.cpp
@@ -192,7 +192,7 @@ private:
 
 template<bool bReconstructPoseValues, bool bAlignedAndPadded>
 static inline void EvaluatePoseKernel(UE::PoseSearch::FSearchResult& Result, const UE::PoseSearch::FSearchIndex& SearchIndex, TConstArrayView<float> QueryValues, TArrayView<float> ReconstructedPoseValuesBuffer,
-	int32 PoseIdx, const UE::PoseSearch::FSearchFilters& SearchFilters, UE::PoseSearch::FSearchContext& SearchContext, const UPoseSearchDatabase* Database, TConstArrayView<float> DynamicWeightsSqrt, bool bUpdateBestCandidates, int32 ResultIndex = -1)
+	int32 PoseIdx, int32 EventPoseIdx, const UE::PoseSearch::FSearchFilters& SearchFilters, UE::PoseSearch::FSearchContext& SearchContext, const UPoseSearchDatabase* Database, TConstArrayView<float> DynamicWeightsSqrt, bool bUpdateBestCandidates, int32 ResultIndex = -1)
 {
 	using namespace UE::PoseSearch;
 
@@ -209,6 +209,7 @@ static inline void EvaluatePoseKernel(UE::PoseSearch::FSearchResult& Result, con
 		{
 			Result.PoseCost = PoseCost;
 			Result.PoseIdx = PoseIdx;
+			Result.EventPoseIdx = EventPoseIdx;
 
 #if WITH_EDITOR && ENABLE_ANIM_DEBUG && UE_POSE_SEARCH_TRACE_ENABLED
 			if (bUpdateBestCandidates)
@@ -221,6 +222,7 @@ static inline void EvaluatePoseKernel(UE::PoseSearch::FSearchResult& Result, con
 #if UE_POSE_SEARCH_TRACE_ENABLED
 		if (bUpdateBestCandidates)
 		{
+			// @todo: add tracking for EventPoseIdx
 			SearchContext.Track(Database, PoseIdx, EPoseCandidateFlags::Valid_Pose, PoseCost);
 		}
 #endif // UE_POSE_SEARCH_TRACE_ENABLED
@@ -1782,7 +1784,7 @@ UE::PoseSearch::FSearchResult UPoseSearchDatabase::SearchPCAKDTree(UE::PoseSearc
 					const TConstArrayView<int32> PoseIndexes = SearchIndex.PCAValuesVectorToPoseIndexes[Results[ResultIndex].Index];
 					for (int32 Index = 0; Index < PoseIndexes.Num() && NumEvaluatePoseKernelCalls < MaxNumEvaluatePoseKernelCalls; ++Index, ++NumEvaluatePoseKernelCalls)
 					{
-						EvaluatePoseKernel<false, true>(Result, SearchIndex, QueryValues, TArrayView<float>(), PoseIndexes[Index], SearchFilters, SearchContext, this, DynamicWeightsSqrt, true
+						EvaluatePoseKernel<false, true>(Result, SearchIndex, QueryValues, TArrayView<float>(), PoseIndexes[Index], INDEX_NONE, SearchFilters, SearchContext, this, DynamicWeightsSqrt, true
 #if WITH_EDITOR && ENABLE_ANIM_DEBUG
 							, SortedResultsIndexes[ResultIndex]
 #endif // WITH_EDITOR && ENABLE_ANIM_DEBUG
@@ -1798,7 +1800,7 @@ UE::PoseSearch::FSearchResult UPoseSearchDatabase::SearchPCAKDTree(UE::PoseSearc
 					const TConstArrayView<int32> PoseIndexes = SearchIndex.PCAValuesVectorToPoseIndexes[Results[ResultIndex].Index];
 					for (int32 Index = 0; Index < PoseIndexes.Num() && NumEvaluatePoseKernelCalls < MaxNumEvaluatePoseKernelCalls; ++Index, ++NumEvaluatePoseKernelCalls)
 					{
-						EvaluatePoseKernel<false, false>(Result, SearchIndex, QueryValues, TArrayView<float>(), PoseIndexes[Index], SearchFilters, SearchContext, this, DynamicWeightsSqrt, true
+						EvaluatePoseKernel<false, false>(Result, SearchIndex, QueryValues, TArrayView<float>(), PoseIndexes[Index], INDEX_NONE, SearchFilters, SearchContext, this, DynamicWeightsSqrt, true
 #if WITH_EDITOR && ENABLE_ANIM_DEBUG
 							, SortedResultsIndexes[ResultIndex]
 #endif // WITH_EDITOR && ENABLE_ANIM_DEBUG
@@ -1815,7 +1817,7 @@ UE::PoseSearch::FSearchResult UPoseSearchDatabase::SearchPCAKDTree(UE::PoseSearc
 			check(IsAligned(ReconstructedPoseValuesBuffer.GetData(), alignof(VectorRegister4Float)));
 			for (int32 ResultIndex = 0; ResultIndex < NumResults; ++ResultIndex)
 			{
-				EvaluatePoseKernel<true, false>(Result, SearchIndex, QueryValues, ReconstructedPoseValuesBuffer, Results[ResultIndex].Index, SearchFilters, SearchContext, this, DynamicWeightsSqrt, true
+				EvaluatePoseKernel<true, false>(Result, SearchIndex, QueryValues, ReconstructedPoseValuesBuffer, Results[ResultIndex].Index, INDEX_NONE, SearchFilters, SearchContext, this, DynamicWeightsSqrt, true
 #if WITH_EDITOR && ENABLE_ANIM_DEBUG
 					, SortedResultsIndexes[ResultIndex]
 #endif // WITH_EDITOR && ENABLE_ANIM_DEBUG
@@ -1827,7 +1829,7 @@ UE::PoseSearch::FSearchResult UPoseSearchDatabase::SearchPCAKDTree(UE::PoseSearc
 		{
 			for (int32 ResultIndex = 0; ResultIndex < NumResults; ++ResultIndex)
 			{
-				EvaluatePoseKernel<false, true>(Result, SearchIndex, QueryValues, TArrayView<float>(), Results[ResultIndex].Index, SearchFilters, SearchContext, this, DynamicWeightsSqrt, true
+				EvaluatePoseKernel<false, true>(Result, SearchIndex, QueryValues, TArrayView<float>(), Results[ResultIndex].Index, INDEX_NONE, SearchFilters, SearchContext, this, DynamicWeightsSqrt, true
 #if WITH_EDITOR && ENABLE_ANIM_DEBUG
 					, SortedResultsIndexes[ResultIndex]
 #endif // WITH_EDITOR && ENABLE_ANIM_DEBUG
@@ -1839,7 +1841,7 @@ UE::PoseSearch::FSearchResult UPoseSearchDatabase::SearchPCAKDTree(UE::PoseSearc
 		{
 			for (int32 ResultIndex = 0; ResultIndex < NumResults; ++ResultIndex)
 			{
-				EvaluatePoseKernel<false, false>(Result, SearchIndex, QueryValues, TArrayView<float>(), Results[ResultIndex].Index, SearchFilters, SearchContext, this, DynamicWeightsSqrt, true
+				EvaluatePoseKernel<false, false>(Result, SearchIndex, QueryValues, TArrayView<float>(), Results[ResultIndex].Index, INDEX_NONE, SearchFilters, SearchContext, this, DynamicWeightsSqrt, true
 #if WITH_EDITOR && ENABLE_ANIM_DEBUG
 					, SortedResultsIndexes[ResultIndex]
 #endif // WITH_EDITOR && ENABLE_ANIM_DEBUG
@@ -1926,7 +1928,7 @@ UE::PoseSearch::FSearchResult UPoseSearchDatabase::SearchVPTree(UE::PoseSearch::
 				const TConstArrayView<int32> PoseIndexes = SearchIndex.ValuesVectorToPoseIndexes[IndexDistance.Index];
 				for (int32 Index = 0; Index < PoseIndexes.Num() && NumEvaluatePoseKernelCalls < MaxNumEvaluatePoseKernelCalls; ++Index, ++NumEvaluatePoseKernelCalls)
 				{
-					EvaluatePoseKernel<false, false>(Result, SearchIndex, QueryValues, TArrayView<float>(), PoseIndexes[Index], SearchFilters, SearchContext, this, DynamicWeightsSqrt, true
+					EvaluatePoseKernel<false, false>(Result, SearchIndex, QueryValues, TArrayView<float>(), PoseIndexes[Index], INDEX_NONE, SearchFilters, SearchContext, this, DynamicWeightsSqrt, true
 #if WITH_EDITOR && ENABLE_ANIM_DEBUG
 						, ResultIndex
 #endif // WITH_EDITOR && ENABLE_ANIM_DEBUG
@@ -1941,7 +1943,7 @@ UE::PoseSearch::FSearchResult UPoseSearchDatabase::SearchVPTree(UE::PoseSearch::
 				const FIndexDistance& IndexDistance = UnsortedResults[ResultIndex];
 
 				// @todo: IndexDistance.Distance is the Sqrt(DissimilarityCost), so there's no need to calculate it again in SearchIndex.ComparePoses
-				EvaluatePoseKernel<false, false>(Result, SearchIndex, QueryValues, TArrayView<float>(), IndexDistance.Index, SearchFilters, SearchContext, this, DynamicWeightsSqrt, true
+				EvaluatePoseKernel<false, false>(Result, SearchIndex, QueryValues, TArrayView<float>(), IndexDistance.Index, INDEX_NONE, SearchFilters, SearchContext, this, DynamicWeightsSqrt, true
 #if WITH_EDITOR && ENABLE_ANIM_DEBUG
 					, ResultIndex
 #endif // WITH_EDITOR && ENABLE_ANIM_DEBUG
@@ -2019,7 +2021,7 @@ UE::PoseSearch::FSearchResult UPoseSearchDatabase::SearchBruteForce(UE::PoseSear
 				check(IsAligned(ReconstructedPoseValuesBuffer.GetData(), alignof(VectorRegister4Float)));
 				for (int32 PoseIdx = 0; PoseIdx < SearchIndex.GetNumPoses(); ++PoseIdx)
 				{
-					EvaluatePoseKernel<true, false>(Result, SearchIndex, QueryValues, ReconstructedPoseValuesBuffer, PoseIdx, SearchFilters, SearchContext, this, DynamicWeightsSqrt, bUpdateBestCandidates
+					EvaluatePoseKernel<true, false>(Result, SearchIndex, QueryValues, ReconstructedPoseValuesBuffer, PoseIdx, INDEX_NONE, SearchFilters, SearchContext, this, DynamicWeightsSqrt, bUpdateBestCandidates
 #if WITH_EDITOR && ENABLE_ANIM_DEBUG
 						, PoseIdx
 #endif // WITH_EDITOR && ENABLE_ANIM_DEBUG
@@ -2031,7 +2033,7 @@ UE::PoseSearch::FSearchResult UPoseSearchDatabase::SearchBruteForce(UE::PoseSear
 			{
 				for (int32 PoseIdx = 0; PoseIdx < SearchIndex.GetNumPoses(); ++PoseIdx)
 				{
-					EvaluatePoseKernel<false, true>(Result, SearchIndex, QueryValues, TArrayView<float>(), PoseIdx, SearchFilters, SearchContext, this, DynamicWeightsSqrt, bUpdateBestCandidates
+					EvaluatePoseKernel<false, true>(Result, SearchIndex, QueryValues, TArrayView<float>(), PoseIdx, INDEX_NONE, SearchFilters, SearchContext, this, DynamicWeightsSqrt, bUpdateBestCandidates
 #if WITH_EDITOR && ENABLE_ANIM_DEBUG
 						, PoseIdx
 #endif // WITH_EDITOR && ENABLE_ANIM_DEBUG
@@ -2043,7 +2045,7 @@ UE::PoseSearch::FSearchResult UPoseSearchDatabase::SearchBruteForce(UE::PoseSear
 			{
 				for (int32 PoseIdx = 0; PoseIdx < SearchIndex.GetNumPoses(); ++PoseIdx)
 				{
-					EvaluatePoseKernel<false, false>(Result, SearchIndex, QueryValues, TArrayView<float>(), PoseIdx, SearchFilters, SearchContext, this, DynamicWeightsSqrt, bUpdateBestCandidates
+					EvaluatePoseKernel<false, false>(Result, SearchIndex, QueryValues, TArrayView<float>(), PoseIdx, INDEX_NONE, SearchFilters, SearchContext, this, DynamicWeightsSqrt, bUpdateBestCandidates
 #if WITH_EDITOR && ENABLE_ANIM_DEBUG
 						, PoseIdx
 #endif // WITH_EDITOR && ENABLE_ANIM_DEBUG
@@ -2069,7 +2071,7 @@ UE::PoseSearch::FSearchResult UPoseSearchDatabase::SearchBruteForce(UE::PoseSear
 					const int32 LastPoseIdx = FirstPoseIdx + SearchIndexAsset.GetNumPoses();
 					for (int32 PoseIdx = FirstPoseIdx; PoseIdx < LastPoseIdx; ++PoseIdx)
 					{
-						EvaluatePoseKernel<true, false>(Result, SearchIndex, QueryValues, ReconstructedPoseValuesBuffer, PoseIdx, SearchFilters, SearchContext, this, DynamicWeightsSqrt, bUpdateBestCandidates
+						EvaluatePoseKernel<true, false>(Result, SearchIndex, QueryValues, ReconstructedPoseValuesBuffer, PoseIdx, INDEX_NONE, SearchFilters, SearchContext, this, DynamicWeightsSqrt, bUpdateBestCandidates
 #if WITH_EDITOR && ENABLE_ANIM_DEBUG
 							, ++ResultIndex
 #endif // WITH_EDITOR && ENABLE_ANIM_DEBUG
@@ -2087,7 +2089,7 @@ UE::PoseSearch::FSearchResult UPoseSearchDatabase::SearchBruteForce(UE::PoseSear
 					const int32 LastPoseIdx = FirstPoseIdx + SearchIndexAsset.GetNumPoses();
 					for (int32 PoseIdx = FirstPoseIdx; PoseIdx < LastPoseIdx; ++PoseIdx)
 					{
-						EvaluatePoseKernel<false, true>(Result, SearchIndex, QueryValues, TArrayView<float>(), PoseIdx, SearchFilters, SearchContext, this, DynamicWeightsSqrt, bUpdateBestCandidates
+						EvaluatePoseKernel<false, true>(Result, SearchIndex, QueryValues, TArrayView<float>(), PoseIdx, INDEX_NONE, SearchFilters, SearchContext, this, DynamicWeightsSqrt, bUpdateBestCandidates
 #if WITH_EDITOR && ENABLE_ANIM_DEBUG
 							, ++ResultIndex
 #endif // WITH_EDITOR && ENABLE_ANIM_DEBUG
@@ -2105,7 +2107,7 @@ UE::PoseSearch::FSearchResult UPoseSearchDatabase::SearchBruteForce(UE::PoseSear
 					const int32 LastPoseIdx = FirstPoseIdx + SearchIndexAsset.GetNumPoses();
 					for (int32 PoseIdx = FirstPoseIdx; PoseIdx < LastPoseIdx; ++PoseIdx)
 					{
-						EvaluatePoseKernel<false, false>(Result, SearchIndex, QueryValues, TArrayView<float>(), PoseIdx, SearchFilters, SearchContext, this, DynamicWeightsSqrt, bUpdateBestCandidates
+						EvaluatePoseKernel<false, false>(Result, SearchIndex, QueryValues, TArrayView<float>(), PoseIdx, INDEX_NONE, SearchFilters, SearchContext, this, DynamicWeightsSqrt, bUpdateBestCandidates
 #if WITH_EDITOR && ENABLE_ANIM_DEBUG
 							, ++ResultIndex
 #endif // WITH_EDITOR && ENABLE_ANIM_DEBUG
@@ -2169,41 +2171,46 @@ UE::PoseSearch::FSearchResult UPoseSearchDatabase::SearchEvent(UE::PoseSearch::F
 			PopulateSelectableAssetIdx(SelectableAssetIdx, SearchContext.GetAssetsToConsider());
 
 			FNonSelectableIdx NonSelectableIdx;
-			PopulateNonSelectableIdx(NonSelectableIdx, SearchContext
+			if (EventToSearch.bEnablePoseFilters)
+			{
+				PopulateNonSelectableIdx(NonSelectableIdx, SearchContext
 #if UE_POSE_SEARCH_TRACE_ENABLED
-				, QueryValues, DynamicWeightsSqrt
+					, QueryValues, DynamicWeightsSqrt
 #endif // UE_POSE_SEARCH_TRACE_ENABLED
-			);
+				);
+			}
 
 			FSearchFilters SearchFilters(Schema, NonSelectableIdx, SelectableAssetIdx, SearchIndex.bAnyBlockTransition);
 
 			if (FMath::IsNearlyZero(EventToSearch.TimeToEvent))
 			{
-				for (int32 PoseIdx : PosesWithEvent)
+				for (int32 EventPoseIdx : PosesWithEvent)
 				{
-					EvaluatePoseKernel<false, false>(Result, SearchIndex, QueryValues, TArrayView<float>(), PoseIdx, SearchFilters, SearchContext, this, DynamicWeightsSqrt, true
+					EvaluatePoseKernel<false, false>(Result, SearchIndex, QueryValues, TArrayView<float>(), EventPoseIdx, EventPoseIdx, SearchFilters, SearchContext, this, DynamicWeightsSqrt, true
 #if WITH_EDITOR && ENABLE_ANIM_DEBUG
-						, PoseIdx
+						, EventPoseIdx
 #endif // WITH_EDITOR && ENABLE_ANIM_DEBUG
 					);
 				}
 			}
 			else
 			{
-				for (int32 PoseIdx : PosesWithEvent)
+				const float TimeToEventStart = EventToSearch.TimeToEvent * EventToSearch.PlayRateRangeOverride.Min;
+				const float TimeToEventEnd = EventToSearch.TimeToEvent * EventToSearch.PlayRateRangeOverride.Max;
+				for (int32 EventPoseIdx : PosesWithEvent)
 				{
 					// Calculating the pose which is TimeToEvent seconds before the event
-					const FSearchIndexAsset& SearchIndexAsset = SearchIndex.GetAssetForPose(PoseIdx);
-					const float EventTime = SearchIndexAsset.GetTimeFromPoseIndex(PoseIdx, Schema->SampleRate);
-					const int32 SearchPoseIndex = SearchIndexAsset.GetPoseIndexFromTime(EventTime - EventToSearch.TimeToEvent, Schema->SampleRate);
-					if (SearchPoseIndex != INDEX_NONE)
-					{
-						EvaluatePoseKernel<false, false>(Result, SearchIndex, QueryValues, TArrayView<float>(), SearchPoseIndex, SearchFilters, SearchContext, this, DynamicWeightsSqrt, true
+					const FSearchIndexAsset& SearchIndexAsset = SearchIndex.GetAssetForPose(EventPoseIdx);
+					const float EventTime = SearchIndexAsset.GetTimeFromPoseIndex(EventPoseIdx, Schema->SampleRate);
+					SearchIndexAsset.IteratePoseIndexesOverTime(EventTime - TimeToEventEnd, EventTime - TimeToEventStart, Schema->SampleRate,
+						[this, &Result, &SearchIndex, &QueryValues, EventPoseIdx, &SearchFilters, &SearchContext, &DynamicWeightsSqrt](int32 PoseIdx)
+						{
+							EvaluatePoseKernel<false, false>(Result, SearchIndex, QueryValues, TArrayView<float>(), PoseIdx, EventPoseIdx, SearchFilters, SearchContext, this, DynamicWeightsSqrt, true
 #if WITH_EDITOR && ENABLE_ANIM_DEBUG
-							, SearchPoseIndex
+								, PoseIdx
 #endif // WITH_EDITOR && ENABLE_ANIM_DEBUG
-						);
-					}
+							);
+						});
 				}
 			}
 		}
diff --git a/Engine/Plugins/Animation/PoseSearch/Source/Runtime/Private/PoseSearchDerivedData.cpp b/Engine/Plugins/Animation/PoseSearch/Source/Runtime/Private/PoseSearchDerivedData.cpp
--- a/Engine/Plugins/Animation/PoseSearch/Source/Runtime/Private/PoseSearchDerivedData.cpp
+++ b/Engine/Plugins/Animation/PoseSearch/Source/Runtime/Private/PoseSearchDerivedData.cpp
@@ -1042,6 +1042,7 @@ static bool IndexDatabase(FSearchIndexBase& SearchIndexBase, const UPoseSearchDa
 	}
 
 	// Joining Stats
+	PRAGMA_DISABLE_DEPRECATION_WARNINGS
 	int32 NumAccumulatedSamples = 0;
 	SearchIndexBase.Stats = FSearchStats();
 	for (int32 AssetIdx = 0; AssetIdx != SearchIndexBase.Assets.Num(); ++AssetIdx)
@@ -1061,6 +1062,7 @@ static bool IndexDatabase(FSearchIndexBase& SearchIndexBase, const UPoseSearchDa
 		SearchIndexBase.Stats.AverageSpeed *= Denom;
 		SearchIndexBase.Stats.AverageAcceleration *= Denom;
 	}
+	PRAGMA_ENABLE_DEPRECATION_WARNINGS
 
 	// Calculate Min Cost Addend
 	SearchIndexBase.MinCostAddend = 0.f;
@@ -1203,10 +1205,12 @@ static void CompareSearchIndexBase(const FSearchIndexBase& A, const FSearchIndex
 		StringBuilder.Append(TEXT("MinCostAddend mismatch\n"));
 	}
 
+	PRAGMA_DISABLE_DEPRECATION_WARNINGS
 	if (A.Stats != B.Stats)
 	{
 		StringBuilder.Append(TEXT("Stats mismatch\n"));
 	}
+	PRAGMA_ENABLE_DEPRECATION_WARNINGS
 }
 
 static void CompareSearchIndex(const FSearchIndex& A, const FSearchIndex& B, const UPoseSearchSchema* Schema, FStringBuilderBase& StringBuilder)
diff --git a/Engine/Plugins/Animation/PoseSearch/Source/Runtime/Private/PoseSearchEvent.cpp b/Engine/Plugins/Animation/PoseSearch/Source/Runtime/Private/PoseSearchEvent.cpp
new file mode 100644
--- /dev/null
+++ b/Engine/Plugins/Animation/PoseSearch/Source/Runtime/Private/PoseSearchEvent.cpp
@@ -0,0 +1,44 @@
+// Copyright Epic Games, Inc. All Rights Reserved.
+
+#include "PoseSearch/PoseSearchEvent.h"
+
+bool FPoseSearchEvent::IsValid() const
+{
+	// maybe we should check if TimeToEvent >= 0?
+	return EventTag.IsValid();
+}
+
+void FPoseSearchEvent::Reset()
+{
+	EventTag = FGameplayTag();
+	TimeToEvent = 0.f;
+}
+
+FPoseSearchEvent FPoseSearchEvent::GetPlayRateOverriddenEvent(const FFloatInterval& PlayRateRangeBase) const
+{
+	FPoseSearchEvent PlayRateOverriddenEventToSearch = *this;
+
+	if (!bUsePlayRateRangeOverride)
+	{
+		PlayRateOverriddenEventToSearch.PlayRateRangeOverride = PlayRateRangeBase;
+	}
+
+	return PlayRateOverriddenEventToSearch;
+}
+
+void UPoseSearchEventLibrary::UpdatePoseSearchEvent(const FPoseSearchEvent& InNewEvent, bool bIsNewEventValid, float DeltaSeconds, FPoseSearchEvent& InOutCurrentEvent)
+{
+	if (bIsNewEventValid && InNewEvent.IsValid())
+	{
+		InOutCurrentEvent = InNewEvent;
+	}
+	else
+	{
+		InOutCurrentEvent.TimeToEvent -= DeltaSeconds;
+		if (InOutCurrentEvent.TimeToEvent < 0.f)
+		{
+			InOutCurrentEvent.TimeToEvent = 0.f;
+			InOutCurrentEvent.EventTag = FGameplayTag();
+		}
+	}
+}
diff --git a/Engine/Plugins/Animation/PoseSearch/Source/Runtime/Private/PoseSearchHistory.cpp b/Engine/Plugins/Animation/PoseSearch/Source/Runtime/Private/PoseSearchHistory.cpp
--- a/Engine/Plugins/Animation/PoseSearch/Source/Runtime/Private/PoseSearchHistory.cpp
+++ b/Engine/Plugins/Animation/PoseSearch/Source/Runtime/Private/PoseSearchHistory.cpp
@@ -658,37 +658,43 @@ FPoseHistory::FPoseHistory(FPoseHistory&& Other)
 
 FPoseHistory& FPoseHistory::operator=(const FPoseHistory& Other)
 {
+	if (this != &Other)
+	{
 #if ENABLE_ANIM_DEBUG
-	UE_MT_SCOPED_WRITE_ACCESS(Other.PoseDataThreadSafeCounter);
-	UE_MT_SCOPED_WRITE_ACCESS(PoseDataThreadSafeCounter);
+		UE_MT_SCOPED_WRITE_ACCESS(Other.PoseDataThreadSafeCounter);
+		UE_MT_SCOPED_WRITE_ACCESS(PoseDataThreadSafeCounter);
 #endif // ENABLE_ANIM_DEBUG
 
-	MaxNumPoses = Other.MaxNumPoses;
-	SamplingInterval = Other.SamplingInterval;
+		MaxNumPoses = Other.MaxNumPoses;
+		SamplingInterval = Other.SamplingInterval;
 
-	Trajectory = Other.Trajectory;
-	TrajectoryDataState = Other.TrajectoryDataState;
-	TrajectorySpeedMultiplier = Other.TrajectorySpeedMultiplier;
+		Trajectory = Other.Trajectory;
+		TrajectoryDataState = Other.TrajectoryDataState;
+		TrajectorySpeedMultiplier = Other.TrajectorySpeedMultiplier;
 
-	PoseData = Other.PoseData;
+		PoseData = Other.PoseData;
+	}
 	return *this;
 }
 
 FPoseHistory& FPoseHistory::operator=(FPoseHistory&& Other)
 {
+	if (this != &Other)
+	{
 #if ENABLE_ANIM_DEBUG
-	UE_MT_SCOPED_WRITE_ACCESS(Other.PoseDataThreadSafeCounter);
-	UE_MT_SCOPED_WRITE_ACCESS(PoseDataThreadSafeCounter);
+		UE_MT_SCOPED_WRITE_ACCESS(Other.PoseDataThreadSafeCounter);
+		UE_MT_SCOPED_WRITE_ACCESS(PoseDataThreadSafeCounter);
 #endif // ENABLE_ANIM_DEBUG
 
-	MaxNumPoses = MoveTemp(Other.MaxNumPoses);
-	SamplingInterval = MoveTemp(Other.SamplingInterval);
+		MaxNumPoses = MoveTemp(Other.MaxNumPoses);
+		SamplingInterval = MoveTemp(Other.SamplingInterval);
 
-	Trajectory = MoveTemp(Other.Trajectory);
-	TrajectoryDataState = MoveTemp(Other.TrajectoryDataState);
-	TrajectorySpeedMultiplier = MoveTemp(Other.TrajectorySpeedMultiplier);
+		Trajectory = MoveTemp(Other.Trajectory);
+		TrajectoryDataState = MoveTemp(Other.TrajectoryDataState);
+		TrajectorySpeedMultiplier = MoveTemp(Other.TrajectorySpeedMultiplier);
 
-	PoseData = MoveTemp(Other.PoseData);
+		PoseData = MoveTemp(Other.PoseData);
+	}
 	return *this;
 }
 
diff --git a/Engine/Plugins/Animation/PoseSearch/Source/Runtime/Private/PoseSearchIndex.cpp b/Engine/Plugins/Animation/PoseSearch/Source/Runtime/Private/PoseSearchIndex.cpp
--- a/Engine/Plugins/Animation/PoseSearch/Source/Runtime/Private/PoseSearchIndex.cpp
+++ b/Engine/Plugins/Animation/PoseSearch/Source/Runtime/Private/PoseSearchIndex.cpp
@@ -282,6 +282,14 @@ FArchive& operator<<(FArchive& Ar, FSearchIndexAsset& IndexAsset)
 
 //////////////////////////////////////////////////////////////////////////
 // FSearchStats
+void FSearchStats::Reset()
+{
+	AverageSpeed = 0.f;
+	MaxSpeed = 0.f;
+	AverageAcceleration = 0.f;
+	MaxAcceleration = 0.f;
+}
+
 bool FSearchStats::operator==(const FSearchStats& Other) const
 {
 	return
@@ -327,7 +335,22 @@ void FEventDataCollector::MergeWith(const FEventDataCollector& Other)
 // FEventData
 FArchive& operator<<(FArchive& Ar, FEventData& EventData)
 {
+#if WITH_EDITOR
+	if (Ar.IsSaving())
+	{
+		EventData.ValidateEventData();
+	}
+#endif // WITH_EDITOR
+
 	Ar << EventData.Data;
+
+#if WITH_EDITOR
+	if (Ar.IsLoading())
+	{
+		EventData.ValidateEventData();
+	}
+#endif // WITH_EDITOR
+
 	return Ar;
 }
 
@@ -350,6 +373,13 @@ const TConstArrayView<int32> FEventData::GetPosesWithEvent(const FGameplayTag& G
 	return TConstArrayView<int32>();
 }
 
+bool FEventData::IsPoseFromEventTag(int32 EventPoseIdx, const FGameplayTag& GameplayTag) const
+{
+	const TConstArrayView<int32> PoseIndexes = GetPosesWithEvent(GameplayTag);
+	// since PoseIndexes is sorted we can perform a binary search instead of PoseIndexes.Contains(EventPoseIdx);
+	return Algo::BinarySearch(PoseIndexes, EventPoseIdx) != INDEX_NONE;
+}
+
 void FEventData::Reset()
 {
 	Data.Reset();
@@ -384,10 +414,96 @@ SIZE_T FEventData::GetAllocatedSize(void) const
 	return Data.GetAllocatedSize();
 }
 
+void FEventData::ValidateEventData() const
+{
+	if (!Algo::IsSorted(Data, [](const FTagToPoseIndexes& A, const FTagToPoseIndexes& B)
+		{
+			// converting FGameplayTag::TagName to string to be deterministic across multiple editor restarts
+			return A.Key.ToString() < B.Key.ToString();
+		}))
+	{
+		UE_LOG(LogPoseSearch, Error, TEXT("FEventData::ValidateEventData FGameplayTag are not properly sorted!"));
+	}
+
+	for (const FEventData::FTagToPoseIndexes& TagToPoseIndexes : Data)
+	{
+		if (!Algo::IsSorted(TagToPoseIndexes.Value))
+		{
+			UE_LOG(LogPoseSearch, Error, TEXT("FEventData::ValidateEventData FPoseIndexes are not properly sorted!"));
+			break;
+		}
+	}
+}
 #endif // WITH_EDITOR
 
 //////////////////////////////////////////////////////////////////////////
 // FPoseSearchBaseIndex
+FSearchIndexBase::FSearchIndexBase()
+	: Values()
+	, ValuesVectorToPoseIndexes()
+	, PoseMetadata()
+	, bAnyBlockTransition(false)
+	, Assets()
+	, EventData()
+	, MinCostAddend(-MAX_FLT)
+	PRAGMA_DISABLE_DEPRECATION_WARNINGS
+	, Stats()
+	PRAGMA_ENABLE_DEPRECATION_WARNINGS
+{
+}
+
+FSearchIndexBase::~FSearchIndexBase()
+{
+}
+
+FSearchIndexBase::FSearchIndexBase(const FSearchIndexBase& Other)
+	: Values(Other.Values)
+	, ValuesVectorToPoseIndexes(Other.ValuesVectorToPoseIndexes)
+	, PoseMetadata(Other.PoseMetadata)
+	, bAnyBlockTransition(Other.bAnyBlockTransition)
+	, Assets(Other.Assets)
+	, EventData(Other.EventData)
+	, MinCostAddend(Other.MinCostAddend)
+	PRAGMA_DISABLE_DEPRECATION_WARNINGS
+	, Stats(Other.Stats)
+	PRAGMA_ENABLE_DEPRECATION_WARNINGS
+{
+}
+
+FSearchIndexBase::FSearchIndexBase(FSearchIndexBase&& Other)
+	: Values(MoveTemp(Other.Values))
+	, ValuesVectorToPoseIndexes(MoveTemp(Other.ValuesVectorToPoseIndexes))
+	, PoseMetadata(MoveTemp(Other.PoseMetadata))
+	, bAnyBlockTransition(MoveTemp(Other.bAnyBlockTransition))
+	, Assets(MoveTemp(Other.Assets))
+	, EventData(MoveTemp(Other.EventData))
+	, MinCostAddend(MoveTemp(Other.MinCostAddend))
+	PRAGMA_DISABLE_DEPRECATION_WARNINGS
+	, Stats(MoveTemp(Other.Stats))
+	PRAGMA_ENABLE_DEPRECATION_WARNINGS
+{
+}
+
+FSearchIndexBase& FSearchIndexBase::operator=(const FSearchIndexBase& Other)
+{
+	if (this != &Other)
+	{
+		this->~FSearchIndexBase();
+		new(this) FSearchIndexBase(Other);
+	}
+	return *this;
+}
+
+FSearchIndexBase& FSearchIndexBase::operator=(FSearchIndexBase&& Other)
+{
+	if (this != &Other)
+	{
+		this->~FSearchIndexBase();
+		new(this) FSearchIndexBase(MoveTemp(Other));
+	}
+	return *this;
+}
+
 const FSearchIndexAsset& FSearchIndexBase::GetAssetForPose(int32 PoseIdx) const
 {
 	const uint32 AssetIndex = PoseMetadata[PoseIdx].GetAssetIndex();
@@ -415,12 +531,21 @@ bool FSearchIndexBase::IsEmpty() const
 
 void FSearchIndexBase::Reset()
 {
-	*this = FSearchIndexBase();
+	Values.Reset();
+	ValuesVectorToPoseIndexes = FSparsePoseMultiMap<int32>();
+	PoseMetadata.Reset();
+	bAnyBlockTransition = false;
+	Assets.Reset();
+	EventData.Reset();
+	MinCostAddend = -MAX_FLT;
+	PRAGMA_DISABLE_DEPRECATION_WARNINGS
+	Stats.Reset();
+	PRAGMA_ENABLE_DEPRECATION_WARNINGS
 }
 
 void FSearchIndexBase::PruneDuplicateValues(float SimilarityThreshold, int32 DataCardinality, bool bDoNotGenerateValuesVectorToPoseIndexes)
 {
-	ValuesVectorToPoseIndexes = FSparsePoseMultiMap<int32>();
+	ValuesVectorToPoseIndexes.Reset();
 
 	const int32 NumPoses = GetNumPoses();
 	if (SimilarityThreshold > 0.f && NumPoses >= 2)
@@ -499,7 +624,9 @@ bool FSearchIndexBase::operator==(const FSearchIndexBase& Other) const
 		Assets == Other.Assets &&
 		EventData == Other.EventData &&
 		MinCostAddend == Other.MinCostAddend &&
+		PRAGMA_DISABLE_DEPRECATION_WARNINGS
 		Stats == Other.Stats;
+		PRAGMA_ENABLE_DEPRECATION_WARNINGS
 }
 
 FArchive& operator<<(FArchive& Ar, FSearchIndexBase& Index)
@@ -511,7 +638,9 @@ FArchive& operator<<(FArchive& Ar, FSearchIndexBase& Index)
 	Ar << Index.Assets;
 	Ar << Index.EventData;
 	Ar << Index.MinCostAddend;
+	PRAGMA_DISABLE_DEPRECATION_WARNINGS
 	Ar << Index.Stats;
+	PRAGMA_ENABLE_DEPRECATION_WARNINGS
 	return Ar;
 }
 
@@ -538,12 +667,43 @@ FSearchIndex::FSearchIndex(const FSearchIndex& Other)
 	KDTree.DataSource.Data = PCAValues.IsEmpty() ? nullptr : PCAValues.GetData();
 }
 
+FSearchIndex::FSearchIndex(FSearchIndex&& Other)
+	: FSearchIndexBase(MoveTemp(Other))
+	, WeightsSqrt(MoveTemp(Other.WeightsSqrt))
+	, PCAValues(MoveTemp(Other.PCAValues))
+	, PCAValuesVectorToPoseIndexes(MoveTemp(Other.PCAValuesVectorToPoseIndexes))
+	, PCAProjectionMatrix(MoveTemp(Other.PCAProjectionMatrix))
+	, Mean(MoveTemp(Other.Mean))
+	, KDTree(MoveTemp(Other.KDTree))
+	, VPTree(MoveTemp(Other.VPTree))
+	#if WITH_EDITORONLY_DATA
+	, DeviationEditorOnly(MoveTemp(Other.DeviationEditorOnly))
+	, PCAExplainedVarianceEditorOnly(MoveTemp(Other.PCAExplainedVarianceEditorOnly))
+	#endif
+	PRAGMA_DISABLE_DEPRECATION_WARNINGS
+	, PCAExplainedVariance(MoveTemp(Other.PCAExplainedVariance))
+	PRAGMA_ENABLE_DEPRECATION_WARNINGS
+{
+	check(!PCAValues.IsEmpty() || KDTree.DataSource.PointCount == 0);
+	KDTree.DataSource.Data = PCAValues.IsEmpty() ? nullptr : PCAValues.GetData();
+}
+
 FSearchIndex& FSearchIndex::operator=(const FSearchIndex& Other)
 {
 	if (this != &Other)
 	{
 		this->~FSearchIndex();
-		new(this)FSearchIndex(Other);
+		new(this) FSearchIndex(Other);
+	}
+	return *this;
+}
+
+FSearchIndex& FSearchIndex::operator=(FSearchIndex&& Other)
+{
+	if (this != &Other)
+	{
+		this->~FSearchIndex();
+		new(this) FSearchIndex(MoveTemp(Other));
 	}
 	return *this;
 }
@@ -693,7 +853,7 @@ void FSearchIndex::PrunePCAValuesFromBlockTransitionPoses(int32 NumberOfPrincipa
 
 void FSearchIndex::PruneDuplicatePCAValues(float SimilarityThreshold, int32 NumberOfPrincipalComponents)
 {
-	PCAValuesVectorToPoseIndexes = FSparsePoseMultiMap<int32>();
+	PCAValuesVectorToPoseIndexes.Reset();
 
 	const uint32 NumPoses = GetNumPoses();
 	if (SimilarityThreshold > 0.f && NumPoses >= 2 && NumberOfPrincipalComponents > 0 && !PCAValues.IsEmpty())
diff --git a/Engine/Plugins/Animation/PoseSearch/Source/Runtime/Private/PoseSearchInteractionIsland.cpp b/Engine/Plugins/Animation/PoseSearch/Source/Runtime/Private/PoseSearchInteractionIsland.cpp
--- a/Engine/Plugins/Animation/PoseSearch/Source/Runtime/Private/PoseSearchInteractionIsland.cpp
+++ b/Engine/Plugins/Animation/PoseSearch/Source/Runtime/Private/PoseSearchInteractionIsland.cpp
@@ -900,7 +900,7 @@ bool FInteractionIsland::DoSearch_AnyThread(const UObject* AnimContext, const FP
 
 			// @todo: we could perform multiple UPoseSearchLibrary::MotionMatch in parallel!
 			PoseSearchResults[SearchIndex] = UPoseSearchLibrary::MotionMatch(AnimContexts, SearchContext.Roles,
-				SearchContext.PoseHistories, AssetsToSearch, ContinuingPropertiesToUse, PoseSearchFutureProperties);
+				SearchContext.PoseHistories, AssetsToSearch, ContinuingPropertiesToUse, PoseSearchFutureProperties, FPoseSearchEvent());
 		}
 		
 #if UE_POSE_SEARCH_ENABLE_ISLAND_CRITICAL_SECTION
diff --git a/Engine/Plugins/Animation/PoseSearch/Source/Runtime/Private/PoseSearchLibrary.cpp b/Engine/Plugins/Animation/PoseSearch/Source/Runtime/Private/PoseSearchLibrary.cpp
--- a/Engine/Plugins/Animation/PoseSearch/Source/Runtime/Private/PoseSearchLibrary.cpp
+++ b/Engine/Plugins/Animation/PoseSearch/Source/Runtime/Private/PoseSearchLibrary.cpp
@@ -314,7 +314,6 @@ PRAGMA_ENABLE_DEPRECATION_WARNINGS
 	ElapsedPoseSearchTime = std::numeric_limits<float>::infinity();
 	WantedPlayRate = 1.f;
 	PoseIndicesHistory.Reset();
-	CurrentEventToSearch.Reset();
 }
 
 void FMotionMatchingState::AdjustAssetTime(float AssetTime)
@@ -341,11 +340,28 @@ FVector FMotionMatchingState::GetEstimatedFutureRootMotionVelocity() const
 	return FVector::ZeroVector;
 }
 
-void FMotionMatchingState::UpdateWantedPlayRate(const UE::PoseSearch::FSearchContext& SearchContext, const FFloatInterval& PlayRate, float TrajectorySpeedMultiplier)
+void FMotionMatchingState::UpdateWantedPlayRate(const UE::PoseSearch::FSearchContext& SearchContext, const FFloatInterval& PlayRate, float TrajectorySpeedMultiplier, const FPoseSearchEvent& EventToSearch)
 {
 	if (CurrentSearchResult.IsValid())
 	{
-		if (!ensure(PlayRate.Min <= PlayRate.Max && PlayRate.Min > UE_KINDA_SMALL_NUMBER))
+		if (CurrentSearchResult.IsEventSearchResult())
+		{
+			// checking if CurrentSearchResult.EventPoseIdx is part of the EventToSearch.EventTag.
+			// If not, it's an event from a continuing pose search that hasn't been interrupted,
+			// so we keep the previously calculated WantedPlayRate
+			const bool bIsEventSearchFromTag = CurrentSearchResult.IsEventSearchFromTag(EventToSearch.EventTag);
+			if (bIsEventSearchFromTag)
+			{
+				const float TimeToEvent = CurrentSearchResult.CalculateTimeToEvent();
+				if (TimeToEvent > UE_KINDA_SMALL_NUMBER && EventToSearch.TimeToEvent > UE_KINDA_SMALL_NUMBER)
+				{
+					// EventToSearch.TimeToEvent is the desired time to event, and TimeToEvent is the actually current time to event. we calculate WantedPlayRate as ratio between the two
+					WantedPlayRate = TimeToEvent / EventToSearch.TimeToEvent;
+				}
+				// if we passed the event (TimeToEvent <= 0) we leave the WantedPlayRate as previously calculated
+			}
+		}
+		else if (!ensure(PlayRate.Min <= PlayRate.Max && PlayRate.Min > UE_KINDA_SMALL_NUMBER))
 		{
 			UE_LOG(LogPoseSearch, Error, TEXT("Couldn't update the WantedPlayRate in FMotionMatchingState::UpdateWantedPlayRate, because of invalid PlayRate interval (%f, %f)"), PlayRate.Min, PlayRate.Max);
 			WantedPlayRate = 1.f;
@@ -662,20 +678,7 @@ void UPoseSearchLibrary::UpdateMotionMatchingState(
 		InOutMotionMatchingState.CurrentSearchResult.Reset();
 	}
 
-	if (EventToSearch.IsValid())
-	{
-		InOutMotionMatchingState.CurrentEventToSearch = EventToSearch;
-	}
-	else
-	{
-		InOutMotionMatchingState.CurrentEventToSearch.TimeToEvent -= DeltaTime;
-	}
-	if (InOutMotionMatchingState.CurrentEventToSearch.TimeToEvent < 0.f)
-	{
-		InOutMotionMatchingState.CurrentEventToSearch.Reset();
-	}
-
-	FSearchContext SearchContext(0.f, &InOutMotionMatchingState.PoseIndicesHistory, InOutMotionMatchingState.CurrentSearchResult, PoseJumpThresholdTime, InOutMotionMatchingState.CurrentEventToSearch);
+	FSearchContext SearchContext(0.f, &InOutMotionMatchingState.PoseIndicesHistory, InOutMotionMatchingState.CurrentSearchResult, PoseJumpThresholdTime, EventToSearch.GetPlayRateOverriddenEvent(PlayRate));
 	SearchContext.AddRole(GetCommonDefaultRole(Databases), AnimContext, PoseHistory);
 
 	const bool bCanAdvance = InOutMotionMatchingState.CurrentSearchResult.CanAdvance(DeltaTime);
@@ -781,7 +784,7 @@ void UPoseSearchLibrary::UpdateMotionMatchingState(
 	}
 
 	// @todo: consider moving this into if (bSearch) to avoid calling SearchContext.GetCachedQuery if no search is required
-	InOutMotionMatchingState.UpdateWantedPlayRate(SearchContext, PlayRate, PoseHistory ? PoseHistory->GetTrajectorySpeedMultiplier() : 1.f);
+	InOutMotionMatchingState.UpdateWantedPlayRate(SearchContext, PlayRate, PoseHistory ? PoseHistory->GetTrajectorySpeedMultiplier() : 1.f, EventToSearch);
 
 	InOutMotionMatchingState.PoseIndicesHistory.Update(InOutMotionMatchingState.CurrentSearchResult, DeltaTime, PoseReselectHistory);
 
@@ -932,7 +935,7 @@ void UPoseSearchLibrary::MotionMatch(
 		return;
 	}
 
-	const FSearchResult SearchResult = MotionMatch(AnimContexts, Roles, PoseHistories, AssetsToSearch, ContinuingProperties, Future);
+	const FSearchResult SearchResult = MotionMatch(AnimContexts, Roles, PoseHistories, AssetsToSearch, ContinuingProperties, Future, FPoseSearchEvent());
 	if (SearchResult.IsValid())
 	{
 		const UPoseSearchDatabase* Database = SearchResult.Database.Get();
@@ -973,7 +976,8 @@ UE::PoseSearch::FSearchResult UPoseSearchLibrary::MotionMatch(
 	const TArrayView<const UE::PoseSearch::IPoseHistory*> PoseHistories, 
 	const TArrayView<const UObject*> AssetsToSearch,
 	const FPoseSearchContinuingProperties& ContinuingProperties,
-	const FPoseSearchFutureProperties& Future)
+	const FPoseSearchFutureProperties& Future,
+	const FPoseSearchEvent& EventToSearch)
 {
 	check(!AnimContexts.IsEmpty() && AnimContexts.Num() == Roles.Num() && AnimContexts.Num() == PoseHistories.Num());
 
@@ -1080,7 +1084,7 @@ UE::PoseSearch::FSearchResult UPoseSearchLibrary::MotionMatch(
 
 	FReconstructedPreviousSearchBestResultMap ReconstructedPreviousSearchBestResultMap;
 	FSearchResult ReconstructedPreviousSearchResult;
-	FSearchContext SearchContext(FutureIntervalTime, nullptr, ReconstructedPreviousSearchResult, FFloatInterval(0.f, 0.f), FPoseSearchEvent());
+	FSearchContext SearchContext(FutureIntervalTime, nullptr, ReconstructedPreviousSearchResult, FFloatInterval(0.f, 0.f), EventToSearch);
 
 	for (int32 RoleIndex = 0; RoleIndex < Roles.Num(); ++RoleIndex)
 	{
@@ -1342,7 +1346,7 @@ UE::PoseSearch::FSearchResult UPoseSearchLibrary::MotionMatch(
 		AnimContexts.Add(AnimInstance);
 	}
 
-	return MotionMatch(AnimContexts, Roles, PoseHistories, AssetsToSearch, ContinuingProperties, Future);
+	return MotionMatch(AnimContexts, Roles, PoseHistories, AssetsToSearch, ContinuingProperties, Future, FPoseSearchEvent());
 }
 
 // End deprecated signatures
diff --git a/Engine/Plugins/Animation/PoseSearch/Source/Runtime/Private/PoseSearchModule.cpp b/Engine/Plugins/Animation/PoseSearch/Source/Runtime/Private/PoseSearchModule.cpp
--- a/Engine/Plugins/Animation/PoseSearch/Source/Runtime/Private/PoseSearchModule.cpp
+++ b/Engine/Plugins/Animation/PoseSearch/Source/Runtime/Private/PoseSearchModule.cpp
@@ -63,9 +63,10 @@ public:
 			PoseHistory = &PoseHistoryProvider->GetPoseHistory();
 		}
 
+		// @todo: Add event support
 		const UObject* AnimContext = GraphContext.AnimInstanceProxy->GetAnimInstanceObject();
 		const UE::PoseSearch::FSearchResult SearchResult = UPoseSearchLibrary::MotionMatch(MakeArrayView(&AnimContext, 1),
-			MakeArrayView(&DefaultRole, 1), MakeArrayView(&PoseHistory, 1), AssetsToSearch, ContinuingProperties, Future);
+			MakeArrayView(&DefaultRole, 1), MakeArrayView(&PoseHistory, 1), AssetsToSearch, ContinuingProperties, Future, FPoseSearchEvent());
 
 		UE::Anim::IPoseSearchProvider::FSearchResult ProviderResult;
 		if (const UE::PoseSearch::FSearchIndexAsset* SearchIndexAsset = SearchResult.GetSearchIndexAsset())
diff --git a/Engine/Plugins/Animation/PoseSearch/Source/Runtime/Private/PoseSearchResult.cpp b/Engine/Plugins/Animation/PoseSearch/Source/Runtime/Private/PoseSearchResult.cpp
--- a/Engine/Plugins/Animation/PoseSearch/Source/Runtime/Private/PoseSearchResult.cpp
+++ b/Engine/Plugins/Animation/PoseSearch/Source/Runtime/Private/PoseSearchResult.cpp
@@ -39,6 +39,34 @@ void FSearchResult::Update(float NewAssetTime)
 	}
 }
 
+bool FSearchResult::IsEventSearchFromTag(const FGameplayTag& EventTag) const
+{
+	check(IsValid());
+	return Database->GetSearchIndex().EventData.IsPoseFromEventTag(EventPoseIdx, EventTag);
+}
+
+float FSearchResult::CalculateTimeToEvent() const
+{
+	check(IsValid() && IsEventSearchResult());
+	const FSearchIndex& SearchIndex = Database->GetSearchIndex();
+
+	check(SearchIndex.PoseMetadata[PoseIdx].GetAssetIndex() == SearchIndex.PoseMetadata[EventPoseIdx].GetAssetIndex());
+	const FSearchIndexAsset& SearchIndexAsset = SearchIndex.GetAssetForPose(PoseIdx);
+
+	// @todo: the mathc here can be simplified between GetDeltaTimeBetweenPoseIndexes and GetTimeFromPoseIndex methods
+	// DeltaTimeBetweenPoseAndEvent is the time in seconds between the event pose (EventPoseIdx) and the current pose (PoseIdx) taking into consideration looping.
+	const float DeltaTimeBetweenPoseAndEvent = SearchIndexAsset.GetDeltaTimeBetweenPoseIndexes(PoseIdx, EventPoseIdx, Database->Schema->SampleRate);
+	// PoseQuantizedTime is the quantized time associated to PoseIdx
+	const float PoseQuantizedTime = SearchIndexAsset.GetTimeFromPoseIndex(PoseIdx, Database->Schema->SampleRate);
+	// AssetTime is the current search result time, that differs from PoseQuantizedTime in case this search result is from a continuing pose search
+	// we calculate the QuantizationError as differnece between AssetTime and PoseQuantizedTime, time in seconds that this search result drifted away from the quantized time of the associated PoseIdx
+	const float QuantizationError = AssetTime - PoseQuantizedTime;
+
+	// time to event can be negative if PoseIdx already passed EventPoseIdx and the asset is not looping
+	const float TimeToEvent = DeltaTimeBetweenPoseAndEvent - QuantizationError;
+	return TimeToEvent;
+}
+
 const FSearchIndexAsset* FSearchResult::GetSearchIndexAsset(bool bMandatory) const
 {
 	if (bMandatory)
diff --git a/Engine/Plugins/Animation/PoseSearch/Source/Runtime/Public/PoseSearch/KDTree.h b/Engine/Plugins/Animation/PoseSearch/Source/Runtime/Public/PoseSearch/KDTree.h
--- a/Engine/Plugins/Animation/PoseSearch/Source/Runtime/Public/PoseSearch/KDTree.h
+++ b/Engine/Plugins/Animation/PoseSearch/Source/Runtime/Public/PoseSearch/KDTree.h
@@ -299,12 +299,12 @@ struct FKDTree
 	FKDTree(AccessorType Count, AccessorType Dim, const float* Data, AccessorType MaxLeafSize = 16);
 	FKDTree();
 	FKDTree(const FKDTree& Other);
-	FKDTree(FKDTree&& Other) = delete;
+	FKDTree(FKDTree&& Other);
 	
 	~FKDTree();
 
 	FKDTree& operator=(const FKDTree& Other);
-	FKDTree& operator=(FKDTree&& Other) = delete;
+	FKDTree& operator=(FKDTree&& Other);
 	bool operator==(const FKDTree& Other) const;
 	
 	void Reset();
diff --git a/Engine/Plugins/Animation/PoseSearch/Source/Runtime/Public/PoseSearch/PoseSearchEvent.h b/Engine/Plugins/Animation/PoseSearch/Source/Runtime/Public/PoseSearch/PoseSearchEvent.h
--- a/Engine/Plugins/Animation/PoseSearch/Source/Runtime/Public/PoseSearch/PoseSearchEvent.h
+++ b/Engine/Plugins/Animation/PoseSearch/Source/Runtime/Public/PoseSearch/PoseSearchEvent.h
@@ -3,24 +3,17 @@
 #pragma once
 
 #include "GameplayTagContainer.h"
+#include "Kismet/BlueprintFunctionLibrary.h"
 #include "PoseSearchEvent.generated.h"
 
-USTRUCT(Experimental, BlueprintType, Category="Animation|Pose Search")
+USTRUCT(Experimental, BlueprintType, Category="Animation|Pose Search|Experimental")
 struct FPoseSearchEvent
 {
 	GENERATED_BODY()
 
-	bool IsValid() const
-	{
-		// maybe we should check if TimeToEvent >= 0?
-		return EventTag.IsValid();
-	}
-
-	void Reset()
-	{
-		EventTag = FGameplayTag();
-		TimeToEvent = 0.f;
-	}
+	bool IsValid() const;
+	void Reset();
+	FPoseSearchEvent GetPlayRateOverriddenEvent(const FFloatInterval& PlayRateRangeBase) const;
 
 	// Tag identifying this event
 	UPROPERTY(EditAnywhere, BlueprintReadWrite, Category=State)
@@ -29,4 +22,30 @@ struct FPoseSearchEvent
 	// Time in seconds until this event occurs
 	UPROPERTY(EditAnywhere, BlueprintReadWrite, Category=State)
 	float TimeToEvent = 0.f;
+
+	// if true pose candidates will be filtered by MotionMatching node 'PoseJumpThresholdTime' (DiscardedBy_PoseJumpThresholdTime) and 'PoseReselectHistory' (DiscardedBy_PoseReselectHistory),
+	// as well as database assets FPoseSearchDatabaseAnimationAssetBase::bDisableReselection (DiscardedBy_AssetReselection)
+	UPROPERTY(EditAnywhere, BlueprintReadWrite, Category=State)
+	bool bEnablePoseFilters = false;
+	
+	// @todo: add PlayRate / PlayRateRangeOverride concept to UPoseSearchLibrary::MotionMatch
+	// if true PlayRateRangeOverride will be used instead of FAnimNode_MotionMatching::PlayRate
+	UPROPERTY(EditAnywhere, BlueprintReadWrite, Category=State)
+	bool bUsePlayRateRangeOverride = true;
+
+	// @todo: support the concept of time to event channel weight, as wey of scoring better poses that are closer to the TimeToEvent
+	// Effective range of play rate that can be applied to the selected animation to play, to account for better pose and trajectory matches.
+	// that means the event search will evaluate poses in the range [TimeToEvent * PlayRate.Min, TimeToEvent * PlayRate.Max] seconds from any poses tagged with EventTag
+	// PlayRateRangeOverride will be used if bUsePlayRateRangeOverride is true, otherwise FAnimNode_MotionMatching::PlayRate will be used instead
+	UPROPERTY(EditAnywhere, BlueprintReadWrite, Category=State, meta = (ClampMin = "0.2", ClampMax = "3.0", UIMin = "0.2", UIMax = "3.0"))
+	FFloatInterval PlayRateRangeOverride = FFloatInterval(1.f, 1.f);
+};
+
+UCLASS()
+class POSESEARCH_API UPoseSearchEventLibrary : public UBlueprintFunctionLibrary
+{
+	GENERATED_BODY()
+
+	UFUNCTION(BlueprintCallable, meta = (BlueprintThreadSafe), Category="Animation|Pose Search|Experimental")
+	static void UpdatePoseSearchEvent(const FPoseSearchEvent& InNewEvent, bool bIsNewEventValid, float DeltaSeconds, UPARAM(ref) FPoseSearchEvent& InOutCurrentEvent);
 };
diff --git a/Engine/Plugins/Animation/PoseSearch/Source/Runtime/Public/PoseSearch/PoseSearchIndex.h b/Engine/Plugins/Animation/PoseSearch/Source/Runtime/Public/PoseSearch/PoseSearchIndex.h
--- a/Engine/Plugins/Animation/PoseSearch/Source/Runtime/Public/PoseSearch/PoseSearchIndex.h
+++ b/Engine/Plugins/Animation/PoseSearch/Source/Runtime/Public/PoseSearch/PoseSearchIndex.h
@@ -171,19 +171,11 @@ struct FSearchIndexAsset
 #if DO_CHECK
 		check(IsInitialized());
 #endif
-
 		const int32 NumPoses = GetNumPoses();
 		int32 PoseOffset = FMath::RoundToInt(SchemaSampleRate * Time) - FirstSampleIdx;
 		if (bLooping)
 		{
-			if (PoseOffset < 0)
-			{
-				PoseOffset = (PoseOffset % NumPoses) + NumPoses;
-			}
-			else if (PoseOffset >= NumPoses)
-			{
-				PoseOffset = PoseOffset % NumPoses;
-			}
+			PoseOffset = FMath::WrapExclusive(PoseOffset, 0, NumPoses);
 			return FirstPoseIdx + PoseOffset;
 		}
 
@@ -195,6 +187,52 @@ struct FSearchIndexAsset
 		return INDEX_NONE;
 	}
 
+	// iteractes over all the valid pose indexes contained in this FSearchIndexAsset from StartTime to EndTime, calling ProcessPoseIndex function on every pose accounting for looping
+	void IteratePoseIndexesOverTime(float StartTime, float EndTime, int32 SchemaSampleRate, const TFunction<void(int32 PoseIdx)>& ProcessPoseIndex) const
+	{
+#if DO_CHECK
+		check(IsInitialized());
+		check(StartTime <= EndTime);
+#endif
+		const int32 NumPoses = GetNumPoses();
+		const int32 StartPoseOffset = FMath::RoundToInt(SchemaSampleRate * StartTime) - FirstSampleIdx;
+		const int32 EndPoseOffset = FMath::RoundToInt(SchemaSampleRate * EndTime) - FirstSampleIdx;
+		if (bLooping)
+		{
+			const int32 EndWrappedPoseOffset = FMath::WrapExclusive(EndPoseOffset, 0, NumPoses);
+			const int32 StartWrappedPoseOffset = FMath::WrapExclusive(StartPoseOffset, 0, NumPoses);
+
+			for (int32 CurrentWrappedPoseOffset = StartWrappedPoseOffset;;)
+			{
+				ProcessPoseIndex(FirstPoseIdx + CurrentWrappedPoseOffset);
+
+				if (CurrentWrappedPoseOffset == EndWrappedPoseOffset)
+				{
+					break;
+				}
+
+				CurrentWrappedPoseOffset = (CurrentWrappedPoseOffset + 1) % NumPoses;
+			}
+		}
+		else
+		{
+			check(StartPoseOffset <= EndPoseOffset);
+			const bool bOutOfStartRange = StartPoseOffset < 0 && EndPoseOffset < 0;
+			const bool bOutOfEndRange = StartPoseOffset >= NumPoses && EndPoseOffset >= NumPoses;
+
+			if (!bOutOfStartRange && !bOutOfEndRange)
+			{
+				const int32 EndClampedPoseOffset = FMath::Clamp(EndPoseOffset, 0, NumPoses - 1);
+				const int32 StartClampedPoseOffset = FMath::Clamp(StartPoseOffset, 0, NumPoses - 1);
+
+				for (int32 CurrentClampedPoseOffset = StartClampedPoseOffset; CurrentClampedPoseOffset <= EndClampedPoseOffset; ++CurrentClampedPoseOffset)
+				{
+					ProcessPoseIndex(FirstPoseIdx + CurrentClampedPoseOffset);
+				}
+			}
+		}
+	}
+
 	float GetTimeFromPoseIndex(int32 PoseIdx, int32 SchemaSampleRate) const
 	{
 		check(SchemaSampleRate > 0);
@@ -206,6 +244,23 @@ struct FSearchIndexAsset
 		return Time;
 	}
 
+	// calculates the delta time in seconds ToPoseIdx - FromPoseIdx
+	float GetDeltaTimeBetweenPoseIndexes(int32 FromPoseIdx, int32 ToPoseIdx, int32 SchemaSampleRate) const
+	{
+		check(SchemaSampleRate > 0);
+
+		check(IsPoseInRange(FromPoseIdx) && IsPoseInRange(ToPoseIdx));
+
+		int32 PoseOffset = ToPoseIdx - FromPoseIdx;
+		if (PoseOffset < 0 && bLooping)
+		{
+			PoseOffset += GetNumPoses();
+		}
+		
+		const float DeltaTime = PoseOffset / float(SchemaSampleRate);
+		return DeltaTime;
+	}
+
 	friend FArchive& operator<<(FArchive& Ar, FSearchIndexAsset& IndexAsset);
 
 private:
@@ -229,6 +284,7 @@ struct FSearchStats
 	float AverageAcceleration = 0.f;
 	float MaxAcceleration = 0.f;
 
+	void Reset();
 	bool operator==(const FSearchStats& Other) const;
 	friend FArchive& operator<<(FArchive& Ar, FSearchStats& Stats);
 };
@@ -237,7 +293,14 @@ struct FSearchStats
 template <typename Type = int32>
 struct FSparsePoseMultiMap
 {
-	FSparsePoseMultiMap(Type InMaxKey = Type(0), Type InMaxValue = Type(0))
+	FSparsePoseMultiMap() = default;
+	~FSparsePoseMultiMap() = default;
+	FSparsePoseMultiMap(const FSparsePoseMultiMap& Other) = default;
+	FSparsePoseMultiMap(FSparsePoseMultiMap&& Other) = default;
+	FSparsePoseMultiMap& operator=(const FSparsePoseMultiMap& Other) = default;
+	FSparsePoseMultiMap& operator=(FSparsePoseMultiMap&& Other) = default;
+
+	FSparsePoseMultiMap(Type InMaxKey, Type InMaxValue)
 	: MaxKey(InMaxKey)
 	, MaxValue(InMaxValue)
 	, DeltaKeyValue(InMaxValue >= InMaxKey ? InMaxValue - InMaxKey + 1 : 0)
@@ -342,6 +405,14 @@ struct FSparsePoseMultiMap
 		return Ar;
 	}
 
+	void Reset()
+	{
+		MaxKey = Type(0);
+		MaxValue = Type(0);
+		DeltaKeyValue = Type(0);
+		DataValues.Reset();
+	}
+
 	Type MaxKey = Type(0);
 	Type MaxValue = Type(0);
 	Type DeltaKeyValue = Type(0);
@@ -375,12 +446,14 @@ struct FEventData
 	friend FArchive& operator<<(FArchive& Ar, FEventData& EventData);
 	bool operator==(const FEventData& Other) const;
 	const TConstArrayView<int32> GetPosesWithEvent(const FGameplayTag& GameplayTag) const;
+	bool IsPoseFromEventTag(int32 EventPoseIdx, const FGameplayTag& GameplayTag) const;
 	const FData& GetData() const { return Data; }
 	void Reset();
 
 #if WITH_EDITOR
 	void Initialize(const FEventDataCollector& EventDataCollector);
 	POSESEARCH_API SIZE_T GetAllocatedSize(void) const;
+	void ValidateEventData() const;
 #endif // WITH_EDITOR
 
 private:
@@ -404,8 +477,17 @@ struct FSearchIndexBase
 	float MinCostAddend = -MAX_FLT;
 
 	// @todo: this property should be editor only
+	UE_DEPRECATED(5.6, "this property will become editor only. wrap its usage with #if WITH_EDITORONLY_DATA")
 	FSearchStats Stats;
 
+	// overriding ctor, copy, move to allow deprecations.
+	FSearchIndexBase();
+	~FSearchIndexBase();
+	FSearchIndexBase(const FSearchIndexBase& Other);
+	FSearchIndexBase(FSearchIndexBase&& Other);
+	FSearchIndexBase& operator=(const FSearchIndexBase& Other);
+	FSearchIndexBase& operator=(FSearchIndexBase&& Other);
+
 	int32 GetNumPoses() const { return PoseMetadata.Num(); }
 	int32 GetNumValuesVectors(int32 DataCardinality) const
 	{
@@ -481,9 +563,9 @@ struct FSearchIndex : public FSearchIndexBase
 	FSearchIndex() = default;
 	~FSearchIndex() = default;
 	FSearchIndex(const FSearchIndex& Other); // custom copy constructor to deal with the KDTree DataSrc
-	FSearchIndex(FSearchIndex&& Other) = delete;
+	FSearchIndex(FSearchIndex&& Other);
 	FSearchIndex& operator=(const FSearchIndex& Other); // custom equal operator to deal with the KDTree DataSrc
-	FSearchIndex& operator=(FSearchIndex&& Other) = delete;
+	FSearchIndex& operator=(FSearchIndex&& Other);
 
 	void Reset();
 	TConstArrayView<float> GetPoseValues(int32 PoseIdx) const;
diff --git a/Engine/Plugins/Animation/PoseSearch/Source/Runtime/Public/PoseSearch/PoseSearchLibrary.h b/Engine/Plugins/Animation/PoseSearch/Source/Runtime/Public/PoseSearch/PoseSearchLibrary.h
--- a/Engine/Plugins/Animation/PoseSearch/Source/Runtime/Public/PoseSearch/PoseSearchLibrary.h
+++ b/Engine/Plugins/Animation/PoseSearch/Source/Runtime/Public/PoseSearch/PoseSearchLibrary.h
@@ -59,7 +59,7 @@ PRAGMA_ENABLE_DEPRECATION_WARNINGS
 	// state back to the default state.
 	POSESEARCH_API void AdjustAssetTime(float AssetTime);
 
-	POSESEARCH_API void UpdateWantedPlayRate(const UE::PoseSearch::FSearchContext& SearchContext, const FFloatInterval& PlayRate, float TrajectorySpeedMultiplier);
+	POSESEARCH_API void UpdateWantedPlayRate(const UE::PoseSearch::FSearchContext& SearchContext, const FFloatInterval& PlayRate, float TrajectorySpeedMultiplier, const FPoseSearchEvent& EventToSearch /*= FPoseSearchEvent()*/);
 
 	FVector GetEstimatedFutureRootMotionVelocity() const;
 
@@ -76,9 +76,6 @@ PRAGMA_ENABLE_DEPRECATION_WARNINGS
 	bool bJumpedToPose = false;
 
 	UE::PoseSearch::FPoseIndicesHistory PoseIndicesHistory;
-
-	// Experimental, this feature might be removed without warning, not for production use
-	FPoseSearchEvent CurrentEventToSearch;
 };
 
 USTRUCT(Experimental, BlueprintType, Category="Animation|Pose Search")
@@ -243,7 +240,8 @@ public:
 		const TArrayView<const UE::PoseSearch::IPoseHistory*> PoseHistories, 
 		const TArrayView<const UObject*> AssetsToSearch,
 		const FPoseSearchContinuingProperties& ContinuingProperties,
-		const FPoseSearchFutureProperties& Future);
+		const FPoseSearchFutureProperties& Future,
+		const FPoseSearchEvent& EventToSearch);
 		
 	UE_DEPRECATED(5.5, "Use other MotionMatch signatures instead")
 	static UE::PoseSearch::FSearchResult MotionMatch(
diff --git a/Engine/Plugins/Animation/PoseSearch/Source/Runtime/Public/PoseSearch/PoseSearchResult.h b/Engine/Plugins/Animation/PoseSearch/Source/Runtime/Public/PoseSearch/PoseSearchResult.h
--- a/Engine/Plugins/Animation/PoseSearch/Source/Runtime/Public/PoseSearch/PoseSearchResult.h
+++ b/Engine/Plugins/Animation/PoseSearch/Source/Runtime/Public/PoseSearch/PoseSearchResult.h
@@ -8,6 +8,7 @@
 
 class UPoseSearchDatabase;
 class UPoseSearchSchema;
+struct FGameplayTag;
 
 namespace UE::PoseSearch
 {
@@ -17,12 +18,27 @@ struct FSearchResult
 {
 	// best cost of the currently selected PoseIdx (it could be equal to ContinuingPoseCost)
 	FPoseSearchCost PoseCost;
+	
+	// pose index from Database associated to this search result
 	int32 PoseIdx = INDEX_NONE;
+
+	// Experimental, this feature might be removed without warning, not for production use
+	// if this search result is from an event search, EventPoseIdx is the pose index at event time from Database associated to this search result
+	// if it's not an event search result EventPoseIdx is INDEX_NONE
+	int32 EventPoseIdx = INDEX_NONE;
+
+	// database associated to this search result. Poses are derived data stored in Database->GetSearchIndex()
 	TWeakObjectPtr<const UPoseSearchDatabase> Database;
+
+	// time of the associated GetSearchIndexAsset from this search. During continuing pose searches,
+	// AssetTime is integrated by DeltaSeconds, so it's not precisely the quantized time associated to PoseIdx
 	float AssetTime = 0.f;
+
+	// true if the search result is from a continuing pose search
 	bool bIsContinuingPoseSearch = false;
 
 #if WITH_EDITOR && ENABLE_ANIM_DEBUG && UE_POSE_SEARCH_TRACE_ENABLED
+	// debug information related to the brute force search, to be able to compare it versus a kdtree/vptree search
 	FPoseSearchCost BruteForcePoseCost;
 	int32 BestPosePos = 0;
 #endif // WITH_EDITOR && ENABLE_ANIM_DEBUG && UE_POSE_SEARCH_TRACE_ENABLED
@@ -32,9 +48,34 @@ struct FSearchResult
 	// state back to the default state.
 	void Update(float NewAssetTime);
 
-	bool IsValid() const { return PoseIdx != INDEX_NONE && Database != nullptr; }
+	bool IsValid() const
+	{
+		return
+			PoseIdx != INDEX_NONE &&
+			Database != nullptr;
+	}
 
-	void Reset() { PoseIdx = INDEX_NONE; Database = nullptr; AssetTime = 0.0f; bIsContinuingPoseSearch = false; }
+	// Experimental, this feature might be removed without warning, not for production use
+	bool IsEventSearchResult() const
+	{
+		return EventPoseIdx != INDEX_NONE;
+	}
+
+	// Experimental, this feature might be removed without warning, not for production use
+	// returns true if EventPoseIdx is from the input EventTag
+	bool IsEventSearchFromTag(const FGameplayTag& EventTag) const;
+	
+	// Experimental, this feature might be removed without warning, not for production use
+	float CalculateTimeToEvent() const;
+
+	void Reset()
+	{
+		PoseIdx = INDEX_NONE;
+		EventPoseIdx = INDEX_NONE;
+		Database = nullptr;
+		AssetTime = 0.0f;
+		bIsContinuingPoseSearch = false;
+	}
 
 	POSESEARCH_API const FSearchIndexAsset* GetSearchIndexAsset(bool bMandatory = false) const;
 	
@@ -43,13 +84,14 @@ struct FSearchResult
 	bool operator==(const FSearchResult& Other) const
 	{
 		// best cost of the currently selected PoseIdx (it could be equal to ContinuingPoseCost)
-		return	PoseCost == Other.PoseCost &&
+		return 
+			PoseCost == Other.PoseCost &&
 			PoseIdx == Other.PoseIdx &&
+			EventPoseIdx == Other.EventPoseIdx &&
 			Database == Other.Database &&
 			AssetTime == Other.AssetTime &&
 			bIsContinuingPoseSearch == Other.bIsContinuingPoseSearch;
 	}
-
 };
 
 } // namespace UE::PoseSearch
