PoseSearch - improved validation of SearchInfo from DDC against database properties

diff --git a/Engine/Plugins/Animation/PoseSearch/Source/Runtime/Private/PoseSearchDerivedData.cpp b/Engine/Plugins/Animation/PoseSearch/Source/Runtime/Private/PoseSearchDerivedData.cpp
--- a/Engine/Plugins/Animation/PoseSearch/Source/Runtime/Private/PoseSearchDerivedData.cpp
+++ b/Engine/Plugins/Animation/PoseSearch/Source/Runtime/Private/PoseSearchDerivedData.cpp
@@ -1086,6 +1086,67 @@ static bool IndexDatabase(FSearchIndexBase& SearchIndexBase, const UPoseSearchDa
 	return true;
 }
 
+// validating the SearchIndex against the Database to catch any data corruption, that can be caused by not updating or having conflicts over DDC key
+static bool ValidateSearchInfoAgainstDatabase(const FSearchIndex& SearchIndex, const UPoseSearchDatabase* Database, const UE::DerivedData::FCacheKey& FullIndexKey)
+{
+	check(Database && Database->Schema);
+
+	if (SearchIndex.GetNumDimensions() != Database->Schema->SchemaCardinality)
+	{
+		UE_LOG(LogPoseSearch, Warning, TEXT("%s - %s BuildIndex From Cache Corrupted! SchemaCardinality mismatch %d vs %d"), *LexToString(FullIndexKey.Hash), *Database->GetName(), SearchIndex.GetNumDimensions(), Database->Schema->SchemaCardinality);
+		return false;
+	}
+
+	TArray<bool> SourceAssetIdxs;
+	SourceAssetIdxs.SetNum(Database->GetNumAnimationAssets());
+
+	for (const FSearchIndexAsset& SearchIndexAsset : SearchIndex.Assets)
+	{
+		if (!SourceAssetIdxs.IsValidIndex(SearchIndexAsset.GetSourceAssetIdx()))
+		{
+			UE_LOG(LogPoseSearch, Warning, TEXT("%s - %s BuildIndex From Cache Corrupted! SearchIndex.Assets referencing missing asset with index %d"), *LexToString(FullIndexKey.Hash), *Database->GetName(), SearchIndexAsset.GetSourceAssetIdx());
+			return false;
+		}
+
+		SourceAssetIdxs[SearchIndexAsset.GetSourceAssetIdx()] = true;
+	}
+
+	for (int32 AnimationAssetIndex = 0; AnimationAssetIndex < Database->GetNumAnimationAssets(); ++AnimationAssetIndex)
+	{
+		if (const FPoseSearchDatabaseAnimationAssetBase* DatabaseAsset = Database->GetDatabaseAnimationAsset<FPoseSearchDatabaseAnimationAssetBase>(AnimationAssetIndex))
+		{
+			if (const UObject* AnimationAsset = DatabaseAsset->GetAnimationAsset())
+			{
+				if (DatabaseAsset->IsEnabled() && !SourceAssetIdxs[AnimationAssetIndex])
+				{
+					UE_LOG(LogPoseSearch, Warning, TEXT("%s - %s BuildIndex From Cache Corrupted! Couldn't find references to enabled asset %s in the SearchIndex"), *LexToString(FullIndexKey.Hash), *Database->GetName(), *AnimationAsset->GetName());
+					return false;
+				}
+				else if (!DatabaseAsset->IsEnabled() && SourceAssetIdxs[AnimationAssetIndex])
+				{
+					UE_LOG(LogPoseSearch, Warning, TEXT("%s - %s BuildIndex From Cache Corrupted! Found references to disabled asset %s in the SearchIndex"), *LexToString(FullIndexKey.Hash), *Database->GetName(), *AnimationAsset->GetName());
+					return false;
+				}
+			}
+			else
+			{
+				if (SourceAssetIdxs[AnimationAssetIndex])
+				{
+					UE_LOG(LogPoseSearch, Warning, TEXT("%s - %s BuildIndex From Cache Corrupted! Found references to null asset at index %d in the SearchIndex"), *LexToString(FullIndexKey.Hash), *Database->GetName(), AnimationAssetIndex);
+					return false;
+				}
+			}
+		}
+		else
+		{
+			UE_LOG(LogPoseSearch, Warning, TEXT("%s - %s BuildIndex From Cache Corrupted! null FPoseSearchDatabaseAnimationAssetBase asset at index %d!?"), *LexToString(FullIndexKey.Hash), *Database->GetName(), AnimationAssetIndex);
+			return false;
+		}
+	}
+
+	return true;
+}
+
 static void SynchronizeDatabaseChooser(UObject* Object)
 {
 	if (UPoseSearchDatabase* Database = Cast<UPoseSearchDatabase>(Object))
@@ -1619,13 +1680,12 @@ void FPoseSearchDatabaseAsyncCacheTask::OnGetComplete(UE::DerivedData::FCacheGet
 		// cache can be corrupted in case the version of the derived data cache has not being updated while 
 		// developing channels that changes their cardinality without impacting any asset properties
 		// so to account for this, we just reindex the database and update the associated DDC 
-		if (SearchIndex.GetNumDimensions() == Database->Schema->SchemaCardinality)
+		if (ValidateSearchInfoAgainstDatabase(SearchIndex, Database.Get(), FullIndexKey))
 		{
 			UE_LOG(LogPoseSearch, Log, TEXT("%s - %s BuildIndex From Cache"), *LexToString(FullIndexKey.Hash), *Database->GetName());
 		}
 		else
 		{
-			UE_LOG(LogPoseSearch, Warning, TEXT("%s - %s BuildIndex From Cache Corrupted!"), *LexToString(FullIndexKey.Hash), *Database->GetName());
 			bCacheCorrupted = true;
 		}
 
diff --git a/Engine/Plugins/Animation/PoseSearch/Source/Runtime/Private/PoseSearchDerivedDataKey.cpp b/Engine/Plugins/Animation/PoseSearch/Source/Runtime/Private/PoseSearchDerivedDataKey.cpp
--- a/Engine/Plugins/Animation/PoseSearch/Source/Runtime/Private/PoseSearchDerivedDataKey.cpp
+++ b/Engine/Plugins/Animation/PoseSearch/Source/Runtime/Private/PoseSearchDerivedDataKey.cpp
@@ -73,7 +73,7 @@ FKeyBuilder::FKeyBuilder(const UObject* Object, bool bUseDataVer, bool bPerformC
 		Hasher.Reset();
 
 		// used to invalidate the key without having to change POSESEARCHDB_DERIVEDDATA_VER all the times
-		int32 POSESEARCHDB_DERIVEDDATA_VER_SMALL = 2667;
+		int32 POSESEARCHDB_DERIVEDDATA_VER_SMALL = 2677;
 		FGuid VersionGuid = FDevSystemGuids::GetSystemGuid(FDevSystemGuids::Get().POSESEARCHDB_DERIVEDDATA_VER);
 		FString AnimationCompressionVersionString = UE::Anim::Compression::AnimationCompressionVersionString;
 
