PoseSearch - MM interactions improvements to support hybrid interactions between AnimInstance(s) and AnimNextComponent(s)

diff --git a/Engine/Plugins/Animation/PoseSearch/Source/Runtime/Private/PoseSearchContext.cpp b/Engine/Plugins/Animation/PoseSearch/Source/Runtime/Private/PoseSearchContext.cpp
--- a/Engine/Plugins/Animation/PoseSearch/Source/Runtime/Private/PoseSearchContext.cpp
+++ b/Engine/Plugins/Animation/PoseSearch/Source/Runtime/Private/PoseSearchContext.cpp
@@ -77,7 +77,6 @@ const AActor* GetContextOwningActor(const UObject* AnimContext)
 		return AnimNextComponent->GetOwner();
 	}
 	
-	unimplemented();
 	return nullptr;
 }
 
diff --git a/Engine/Plugins/Animation/PoseSearch/Source/Runtime/Private/PoseSearchInteractionIsland.cpp b/Engine/Plugins/Animation/PoseSearch/Source/Runtime/Private/PoseSearchInteractionIsland.cpp
--- a/Engine/Plugins/Animation/PoseSearch/Source/Runtime/Private/PoseSearchInteractionIsland.cpp
+++ b/Engine/Plugins/Animation/PoseSearch/Source/Runtime/Private/PoseSearchInteractionIsland.cpp
@@ -150,6 +150,54 @@ static void InitSearchResults(TArray<FInteractionSearchResult>& SearchResults, T
 	}
 }
 
+static bool FindComponentsForTickDependencies(const UObject* AnimContext, UActorComponent*& PreTickComponent, UActorComponent*& PostTickComponent)
+{
+	check(AnimContext);
+
+	const AActor* Actor = GetContextOwningActor(AnimContext);
+
+	FName MotionMatchingInteractionPreTag = "MotionMatchingInteractionPre";
+	if (UActorComponent* MotionMatchingInteractionPre = Actor->FindComponentByTag<UActorComponent>(MotionMatchingInteractionPreTag))
+	{
+		PreTickComponent = MotionMatchingInteractionPre;
+	}
+	else if (UActorComponent* CharacterMovementComponent = Actor->GetComponentByClass<UCharacterMovementComponent>())
+	{
+		PreTickComponent = CharacterMovementComponent;
+	}
+	else
+	{
+		UE_LOG(LogPoseSearch, Error, TEXT("FindComponentsForTickDependencies requires AActor %s to have a UCharacterMovementComponent, or a component tagged as 'MotionMatchingInteractionPre' to work!"), *Actor->GetName());
+		PreTickComponent = nullptr;
+	}
+
+	FName MotionMatchingInteractionPostTag = "MotionMatchingInteractionPost";
+	if (UActorComponent* MotionMatchingInteractionPost = Actor->FindComponentByTag<UActorComponent>(MotionMatchingInteractionPostTag))
+	{
+		PostTickComponent = MotionMatchingInteractionPost;
+	}
+	else if (const UAnimInstance* AnimInstance = Cast<UAnimInstance>(AnimContext))
+	{
+		// this is the UAnimInstance / AnimationBluePrint case
+		if (UActorComponent* SkeletalMeshComponent = Actor->GetComponentByClass<USkeletalMeshComponent>())
+		{
+			PostTickComponent = SkeletalMeshComponent;
+		}
+		else
+		{
+			UE_LOG(LogPoseSearch, Error, TEXT("FindComponentsForTickDependencies requires AActor %s to have a USkeletalMeshComponent, or a component tagged as 'MotionMatchingInteractionPost' to work!"), *Actor->GetName());
+			PostTickComponent = nullptr;
+		}
+	}
+	else
+	{
+		// this is the AnimNext case
+		PostTickComponent = const_cast<UActorComponent*>(Cast<UActorComponent>(AnimContext));
+	}
+
+	return PreTickComponent && PostTickComponent;
+}
+
 // FInteractionSearchContextBase
 ///////////////////////////////////////////////////////////
 bool FInteractionSearchContextBase::IsEquivalent(const FInteractionSearchContextBase& Other) const
@@ -428,57 +476,20 @@ FInteractionIsland::~FInteractionIsland()
 	InteractionSubsystem = nullptr;
 }
 
-UActorComponent* FInteractionIsland::FindPreTickComponent(const AActor* Actor)
-{
-	FName Tag = "MotionMatchingInteractionPre";
-	if (UActorComponent* PreTickComponent = Actor->FindComponentByTag<UActorComponent>(Tag))
-	{
-		return PreTickComponent;
-	}
-
-	if (UActorComponent* CharacterMovementComponent = Actor->GetComponentByClass<UCharacterMovementComponent>())
-	{
-		return CharacterMovementComponent;
-	}
-
-	UE_LOG(LogPoseSearch, Error, TEXT("FInteractionIsland::FindPreTickComponent requires AActor %s to have a UCharacterMovementComponent, or a component tagged as 'MotionMatchingInteractionPre' to work!"), *Actor->GetName());
-	return nullptr;
-}
-
-UActorComponent* FInteractionIsland::FindPostTickComponent(const AActor* Actor)
-{
-	FName Tag = "MotionMatchingInteractionPost";
-	if (UActorComponent* PostTickComponent = Actor->FindComponentByTag<UActorComponent>(Tag))
-	{
-		return PostTickComponent;
-	}
-
-	if (UActorComponent* SkeletalMeshComponent = Actor->GetComponentByClass<USkeletalMeshComponent>())
-	{
-		return SkeletalMeshComponent;
-	}
-
-	UE_LOG(LogPoseSearch, Error, TEXT("FInteractionIsland::FindPostTickComponent requires AActor %s to have a USkeletalMeshComponent, or a component tagged as 'MotionMatchingInteractionPost' to work!"), *Actor->GetName());
-	return nullptr;
-}
-
 IInteractionIslandDependency* FInteractionIsland::FindCustomDependency(UActorComponent* InTickComponent) const
 {
 	IModularFeatures& ModularFeatures = IModularFeatures::Get();
-	int32 NumFeatures = ModularFeatures.GetModularFeatureImplementationCount(IInteractionIslandDependency::FeatureName);
+	const int32 NumFeatures = ModularFeatures.GetModularFeatureImplementationCount(IInteractionIslandDependency::FeatureName);
 
 	// Add pre-tick function dependencies
-	for(int32 FeatureIndex = 0; FeatureIndex < NumFeatures; ++FeatureIndex)
+	for (int32 FeatureIndex = 0; FeatureIndex < NumFeatures; ++FeatureIndex)
 	{
-		IInteractionIslandDependency* IslandDependency = static_cast<IInteractionIslandDependency*>(ModularFeatures.GetModularFeatureImplementation(IInteractionIslandDependency::FeatureName, FeatureIndex));
-		if(IslandDependency == nullptr)
+		if (IInteractionIslandDependency* IslandDependency = static_cast<IInteractionIslandDependency*>(ModularFeatures.GetModularFeatureImplementation(IInteractionIslandDependency::FeatureName, FeatureIndex)))
 		{
-			continue;
-		}
-
-		if(IslandDependency->CanMakeDependency(InteractionSubsystem, InTickComponent))
-		{
-			return IslandDependency;
+			if (IslandDependency->CanMakeDependency(InteractionSubsystem, InTickComponent))
+			{
+				return IslandDependency;
+			}
 		}
 	}
 	return nullptr;
@@ -556,34 +567,33 @@ void FInteractionIsland::InjectToActor(const UObject* AnimContext, bool bAddTick
 		}
 #endif // ENABLE_ANIM_DEBUG
 
-		const AActor* Actor = GetContextOwningActor(AnimContext);
-		if (UActorComponent* PreTickComponent = FindPreTickComponent(Actor))
+		UActorComponent* PreTickComponent = nullptr;
+		UActorComponent* PostTickComponent = nullptr;
+		if (FindComponentsForTickDependencies(AnimContext, PreTickComponent, PostTickComponent))
 		{
-			if (UActorComponent* PostTickComponent = FindPostTickComponent(Actor))
+			const bool bIsMainActor = IslandAnimContexts.IsEmpty();
+
+			//	tick order: 
+			//		PreTickComponents (UCharacterMovementComponent, or UCharacterMoverComponent) ->
+			//			Island.PreTickFunction ->
+			//				first injected actor PostTickComponent (USkeletalMeshComponent, or WIP UAnimNextComponent) ->
+			//					Island.PostTickFunction ->
+			//						other PostTickComponent(s)
+			PreTickComponents.AddUnique(PreTickComponent);
+			PostTickComponents.AddUnique(PostTickComponent);
+			IslandAnimContexts.AddUnique(AnimContext);
+
+			// making sure that if we add a unique PreTickComponent, we add as well a unique PostTickComponent
+			// (so we can remove them later on in a consistent fashion)
+			check(PreTickComponents.Num() == PostTickComponents.Num() && PreTickComponents.Num() == IslandAnimContexts.Num());
+
+			if (bAddTickDependencies)
 			{
-				const bool bIsMainActor = PostTickComponents.IsEmpty();
-
-				//	tick order: 
-				//		PreTickComponents (UCharacterMovementComponent, or UCharacterMoverComponent) ->
-				//			Island.PreTickFunction ->
-				//				first injected actor PostTickComponent (USkeletalMeshComponent, or WIP UAnimNextComponent) ->
-				//					Island.PostTickFunction ->
-				//						other PostTickComponent(s)
-				PreTickComponents.AddUnique(PreTickComponent);
-				PostTickComponents.AddUnique(PostTickComponent);
-
-				// making sure that if we add a unique PreTickComponent, we add as well a unique PostTickComponent
-				// (so we can remove them later on in a consistent fashion)
-				check(PreTickComponents.Num() == PostTickComponents.Num());
-
-				if (bAddTickDependencies)
-				{
-					AddTickDependencies(PreTickComponent, PostTickComponent, bIsMainActor);
-				}
-				else
-				{
-					check(!bHasTickDependencies);
-				}
+				AddTickDependencies(PreTickComponent, PostTickComponent, bIsMainActor);
+			}
+			else
+			{
+				check(!bHasTickDependencies);
 			}
 		}
 	}
@@ -603,7 +613,7 @@ void FInteractionIsland::RemoveTickDependencies(bool bValidateTickDependencies)
 	// Called by UPoseSearchInteractionSubsystem::Tick when there aren't animation jobs flying. No need to FScopeLock Lock(&Mutex);
 	check(IsInGameThread());
 
-	check(PreTickComponents.Num() == PostTickComponents.Num());
+	check(PreTickComponents.Num() == PostTickComponents.Num() && PreTickComponents.Num() == IslandAnimContexts.Num());
 
 	if (!bHasTickDependencies)
 	{
@@ -709,6 +719,7 @@ void FInteractionIsland::Uninitialize(bool bValidateTickDependencies)
 	{
 		PreTickComponents.Reset();
 		PostTickComponents.Reset();
+		IslandAnimContexts.Reset();
 
 		SearchContexts.Reset();
 		SearchResults.Reset();
@@ -716,7 +727,7 @@ void FInteractionIsland::Uninitialize(bool bValidateTickDependencies)
 	}
 	else
 	{
-		check(PreTickComponents.IsEmpty() && PostTickComponents.IsEmpty() && SearchContexts.IsEmpty() && SearchResults.IsEmpty() && !bSearchPerfomed);
+		check(PreTickComponents.IsEmpty() && PostTickComponents.IsEmpty() && IslandAnimContexts.IsEmpty() && SearchContexts.IsEmpty() && SearchResults.IsEmpty() && !bSearchPerfomed);
 	}
 }
 
@@ -730,51 +741,25 @@ bool FInteractionIsland::IsInitialized() const
 	return !SearchContexts.IsEmpty();
 }
 
+const UObject* FInteractionIsland::GetMainAnimContext() const
+{
+	return !IslandAnimContexts.IsEmpty() ? IslandAnimContexts[0].Get() : nullptr;
+}
+
 const AActor* FInteractionIsland::GetMainActor() const
 {
 	if (IsInitialized())
 	{
-		check(!PostTickComponents.IsEmpty());
-		if (const UActorComponent* ActorComponent = PostTickComponents[0].Get())
-		{
-			const AActor* MainActor = ActorComponent->GetOwner();
-			check(MainActor);
-			return MainActor;
-		}
+		return GetContextOwningActor(GetMainAnimContext());
 	}
 	return nullptr;
 }
 
 bool FInteractionIsland::CanPerformSearch(const UObject* AnimContext) const
 {
+	// searches are performed only on the MainAnimContext / MainActor
 	check(AnimContext);
-
-	const AActor* MainActor = GetMainActor();
-	if (!MainActor)
-	{
-		UE_LOG(LogPoseSearch, Error, TEXT("FInteractionIsland::CanPerformSearch invalid MainActor during the search!"));
-
-#if ENABLE_ANIM_DEBUG
-		LogTickDependencies();
-#endif // ENABLE_ANIM_DEBUG
-
-		return false;
-	}
-
-	const AActor* Actor = GetContextOwningActor(AnimContext);
-	check(Actor);
-
-	if (Actor != MainActor)
-	{
-		UE_LOG(LogPoseSearch, Error, TEXT("FInteractionIsland::CanPerformSearch search NOT performed by the MainActor!"));
-
-#if ENABLE_ANIM_DEBUG
-		LogTickDependencies();
-#endif // ENABLE_ANIM_DEBUG
-		return false;
-	}
-	
-	return true;
+	return AnimContext == GetMainAnimContext();
 }
 
 #if ENABLE_ANIM_DEBUG
diff --git a/Engine/Plugins/Animation/PoseSearch/Source/Runtime/Private/PoseSearchInteractionSubsystem.cpp b/Engine/Plugins/Animation/PoseSearch/Source/Runtime/Private/PoseSearchInteractionSubsystem.cpp
--- a/Engine/Plugins/Animation/PoseSearch/Source/Runtime/Private/PoseSearchInteractionSubsystem.cpp
+++ b/Engine/Plugins/Animation/PoseSearch/Source/Runtime/Private/PoseSearchInteractionSubsystem.cpp
@@ -303,6 +303,7 @@ bool UPoseSearchInteractionSubsystem::ValidateAllIslands() const
 	TSet<TWeakObjectPtr<UActorComponent>> PreTickComponents;
 	TSet<TWeakObjectPtr<UActorComponent>> PostTickComponents;
 
+	// @todo: validate IslandAnimContexts reference the same actors as PreTickComponents and PostTickComponents
 	typedef TSet<const UObject*> FIslandAnimContexts;
 	TArray<FIslandAnimContexts> IslandsAnimContexts;
 
@@ -405,20 +406,11 @@ UE::PoseSearch::FInteractionIsland* UPoseSearchInteractionSubsystem::FindIsland(
 
 	if (AnimContext)
 	{
-		if (const AActor* Actor = GetContextOwningActor(AnimContext))
+		for (FInteractionIsland* Island : Islands)
 		{
-			if (UActorComponent* PostTickComponent = FInteractionIsland::FindPostTickComponent(Actor))
+			if (Island->GetIslandAnimContexts().Contains(AnimContext))
 			{
-				for (FInteractionIsland* Island : Islands)
-				{
-					for (const TWeakObjectPtr<UActorComponent>& IslandPostTickComponent : Island->GetPostTickComponents())
-					{
-						if (IslandPostTickComponent.Get() == PostTickComponent)
-						{
-							return Island;
-						}
-					}
-				}
+				return Island;
 			}
 		}
 	}
@@ -917,17 +909,10 @@ void UPoseSearchInteractionSubsystem::OnInteractionStart(UE::PoseSearch::FValidI
 		{
 			if (const UObject* AnimContext = AnimContextPtr.Get())
 			{
-				if (const UAnimInstance* AnimInstance = Cast<UAnimInstance>(AnimContext))
-				{
-					AActor* Actor = AnimInstance->GetOwningActor();
-					check(Actor);
-					Actors.Add(Actor);
-					PrimitiveComponents.Add(Cast<UPrimitiveComponent>(Actor->GetRootComponent()));
-				}
-				else
-				{
-					unimplemented();
-				}
+				AActor* Actor = const_cast<AActor*>(GetContextOwningActor(AnimContext));
+				check(Actor);
+				Actors.Add(Actor);
+				PrimitiveComponents.Add(Cast<UPrimitiveComponent>(Actor->GetRootComponent()));
 			}
 		}
 
diff --git a/Engine/Plugins/Animation/PoseSearch/Source/Runtime/Public/PoseSearch/PoseSearchInteractionIsland.h b/Engine/Plugins/Animation/PoseSearch/Source/Runtime/Public/PoseSearch/PoseSearchInteractionIsland.h
--- a/Engine/Plugins/Animation/PoseSearch/Source/Runtime/Public/PoseSearch/PoseSearchInteractionIsland.h
+++ b/Engine/Plugins/Animation/PoseSearch/Source/Runtime/Public/PoseSearch/PoseSearchInteractionIsland.h
@@ -92,12 +92,10 @@ struct FInteractionIsland
 
 	const TArray<TWeakObjectPtr<UActorComponent>>& GetPreTickComponents() const { return PreTickComponents; }
 	const TArray<TWeakObjectPtr<UActorComponent>>& GetPostTickComponents() const { return PostTickComponents; }
+	const TArray<TWeakObjectPtr<const UObject>>& GetIslandAnimContexts() const { return IslandAnimContexts; }
 	const TArray<FInteractionSearchContext>& GetSearchContexts() const { return SearchContexts; }
 	const TArray<FInteractionSearchResult>& GetSearchResults() const { return SearchResults; }
 
-	static UActorComponent* FindPreTickComponent(const AActor* Actor);
-	static UActorComponent* FindPostTickComponent(const AActor* Actor);
-
 	const FInteractionSearchResult* FindSearchResult(const FInteractionSearchContext& SearchContext) const;
 
 	bool IsInitialized() const;
@@ -113,6 +111,7 @@ private:
 	IInteractionIslandDependency* FindCustomDependency(UActorComponent* InTickComponent) const;
 
 	bool GetResult_AnyThread(const UObject* AnimContext, FPoseSearchInteractionBlueprintResult& Result);
+	const UObject* GetMainAnimContext() const;
 	const AActor* GetMainActor() const;
 	bool CanPerformSearch(const UObject* AnimContext) const;
 
@@ -137,11 +136,11 @@ private:
 	FPostTickFunction PostTickFunction;
 	bool bHasTickDependencies = false;
 
-	// islands don't get deallocated, so this array once warmed up will not hit the allocator and waste cycles
 	TArray<TWeakObjectPtr<UActorComponent>> PreTickComponents;
-	// islands don't get deallocated, so this array once warmed up will not hit the allocator and waste cycles
 	TArray<TWeakObjectPtr<UActorComponent>> PostTickComponents;
-	
+	// all the AnimContext in this island. Each SearchContexts will contain a subset of IslandAnimContexts 
+	TArray<TWeakObjectPtr<const UObject>> IslandAnimContexts;
+
 	// there's one FSearchContext for each search we need to perform (including all the possible roles permutations). Added by UPoseSearchInteractionSubsystem::Tick
 	// islands don't get deallocated, so this array once warmed up will not hit the allocator and waste cycles
 	TArray<FInteractionSearchContext> SearchContexts;
@@ -167,6 +166,7 @@ private:
 #endif // ENABLE_ANIM_DEBUG
 };
 
+// Experimental, this feature might be removed without warning, not for production use
 // Allows systems other than regular actor components to hook into interaction island dependencies 
 struct IInteractionIslandDependency : public IModularFeature
 {
