PoseSearch - reducing memory reallocations while indexing databases

diff --git a/Engine/Plugins/Animation/PoseSearch/Source/Runtime/Private/PoseSearchDerivedData.cpp b/Engine/Plugins/Animation/PoseSearch/Source/Runtime/Private/PoseSearchDerivedData.cpp
--- a/Engine/Plugins/Animation/PoseSearch/Source/Runtime/Private/PoseSearchDerivedData.cpp
+++ b/Engine/Plugins/Animation/PoseSearch/Source/Runtime/Private/PoseSearchDerivedData.cpp
@@ -149,11 +149,11 @@ private:
 	struct FEntry
 	{
 		int32 SchemaIndex = INDEX_NONE; // index of the Channel associated Schemas / SearchIndexBases index
-		const UPoseSearchFeatureChannel* Channel;
+		const UPoseSearchFeatureChannel* Channel = nullptr;
 	};
 
-	typedef TArray<FEntry> FEntries;		// array of FEntry with channels that can be normalized together (for example it contains all the phases of the left foot from different schemas)
-	typedef TArray<FEntries> FEntriesGroup;	// array of FEntries incompatible between each other
+	typedef TArray<FEntry, TInlineAllocator<16>> FEntries;		// array of FEntry with channels that can be normalized together (for example it contains all the phases of the left foot from different schemas)
+	typedef TArray<FEntries, TInlineAllocator<16>> FEntriesGroup;	// array of FEntries incompatible between each other
 
 	static void Add(const UPoseSearchFeatureChannel* Channel, int32 SchemaIndex, FEntriesGroup& EntriesGroup)
 	{
@@ -319,7 +319,9 @@ public:
 	}
 };
 
-static void FindValidRanges(const FPoseSearchDatabaseAnimationAssetBase* DatabaseAsset, const FVector& BlendParameters, const FFloatInterval& ExcludeFromDatabaseParameters, TArray<FFloatRange>& ValidRanges)
+typedef TArray<FFloatRange, TInlineAllocator<32>> FValidRanges;
+
+static void FindValidRanges(const FPoseSearchDatabaseAnimationAssetBase* DatabaseAsset, const FVector& BlendParameters, const FFloatInterval& ExcludeFromDatabaseParameters, FValidRanges& ValidRanges)
 {
 	check(DatabaseAsset);
 
@@ -358,10 +360,10 @@ static void FindValidRanges(const FPoseSearchDatabaseAnimationAssetBase* Databas
 						// number of ranges in ValidRanges, the algorithm iterates from end to start.
 						for (int32 RangeIdx = ValidRanges.Num() - 1; RangeIdx >= 0; --RangeIdx)
 						{
-							FFloatRange EvaluatedRange = ValidRanges[RangeIdx];
+							const FFloatRange EvaluatedRange = ValidRanges[RangeIdx];
 							ValidRanges.RemoveAt(RangeIdx);
 
-							TArray<FFloatRange> Diff = FFloatRange::Difference(EvaluatedRange, ExclusionRange);
+							const TArray<FFloatRange> Diff = FFloatRange::Difference(EvaluatedRange, ExclusionRange);
 							ValidRanges.Append(Diff);
 						}
 					}
@@ -379,7 +381,7 @@ static bool InitSearchIndexAssets(FSearchIndexBase& SearchIndex, const UPoseSear
 	check(Schema);
 
 	SearchIndex.Assets.Reset();
-	TArray<FFloatRange> ValidRanges;
+	FValidRanges ValidRanges;
 
 	bool bAnyErrors = false;
 
@@ -696,10 +698,11 @@ static void PreprocessSearchIndexKDTree(FSearchIndex& SearchIndex, const UPoseSe
 		if (AnyTestFlags(EMotionMatchTestFlags::TestKDTreeConstructDeterminism))
 		{
 			const int32 NumIterations = GVarMotionMatchTestNumIterations;
+			TAlignedArray<float> PCAValuesTest;
 			for (int32 Iteration = 0; Iteration < NumIterations; ++Iteration)
 			{
 				// copy PCAValues in a different container to ensure input data has different memory addresses
-				TAlignedArray<float> PCAValuesTest = SearchIndex.PCAValues;
+				PCAValuesTest = SearchIndex.PCAValues;
 
 				FKDTree KDTreeTest;
 				KDTreeTest.Construct(NumPCAValuesVectors, NumberOfPrincipalComponents, PCAValuesTest.GetData(), KDTreeMaxLeafSize);
@@ -1161,9 +1164,11 @@ static void SynchronizeDatabaseChooser(UObject* Object)
 		TArray<FAssetIdentifier> Referencers;
 		AssetRegistry.GetReferencers(RootChooser->GetPackage()->GetFName(), Referencers);
 
+		TArray<FAssetData> Assets;
+		Assets.Reserve(256);
 		for (const FAssetIdentifier& Referencer : Referencers)
 		{
-			TArray<FAssetData> Assets;
+			Assets.Reset();
 			AssetRegistry.GetAssetsByPackageName(Referencer.PackageName, Assets);
 
 			for (const FAssetData& Asset : Assets)
@@ -1395,7 +1400,7 @@ private:
 	bool bResidencyCleared = true;
 };
 
-class FPoseSearchDatabaseAsyncCacheTasks : public TArray<TUniquePtr<FPoseSearchDatabaseAsyncCacheTask>> {};
+class FPoseSearchDatabaseAsyncCacheTasks : public TArray<TUniquePtr<FPoseSearchDatabaseAsyncCacheTask>, TInlineAllocator<64>> {};
 
 FPoseSearchDatabaseAsyncCacheTask::FPoseSearchDatabaseAsyncCacheTask(UPoseSearchDatabase* InDatabase, bool bPerformConditionalPostLoadIfRequired, FPartialKeyHashes& PartialKeyHashes)
 	: Database(InDatabase)
@@ -1523,15 +1528,11 @@ void FPoseSearchDatabaseAsyncCacheTask::StartNewRequestIfNeeded(bool bPerformCon
 				
 				UE_LOG(LogPoseSearch, Log, TEXT("%s - %s BeginCache"), *LexToString(DerivedDataKey), *Database->GetName());
 
-				int32 GSkipDDC = 0;
-				const ECachePolicy Policy = GSkipDDC ? ECachePolicy::None : ECachePolicy::Default;
-			
-				TArray<FCacheGetRequest> CacheRequests;
 				const FCacheKey CacheKey{ Bucket, DerivedDataKey };
-				CacheRequests.Add({ { Database->GetPathName() }, CacheKey, Policy });
+				const FCacheGetRequest CacheRequest = { { Database->GetPathName() }, CacheKey, ECachePolicy::Default };
 
 				Owner = FRequestOwner(EPriority::Normal);
-				GetCache().Get(CacheRequests, Owner, [this](FCacheGetResponse&& Response)
+				GetCache().Get(MakeArrayView(&CacheRequest, 1), Owner, [this](FCacheGetResponse&& Response)
 					{
 						OnGetComplete(MoveTemp(Response));
 					});
@@ -1724,6 +1725,7 @@ void FPoseSearchDatabaseAsyncCacheTask::OnGetComplete(UE::DerivedData::FCacheGet
 				// collecting all the databases that need to be built to gather their FSearchIndexBase
 				// the first one is always the main database (the one we're calculating the index on)
 				TArray<const UPoseSearchDatabase*> IndexBaseDatabases;
+				IndexBaseDatabases.Reserve(64);
 				IndexBaseDatabases.Add(MainDatabase);
 				if (!MainDatabase)
 				{
@@ -1901,7 +1903,7 @@ void FPoseSearchDatabaseAsyncCacheTask::OnGetComplete(UE::DerivedData::FCacheGet
 				const bool bDoNotGenerateValuesVectorToPoseIndexes = MainDatabase->PoseSearchMode != EPoseSearchMode::VPTree;
 				SearchIndex.PruneDuplicateValues(MainDatabase->PosePruningSimilarityThreshold, MainDatabaseSchema->SchemaCardinality, bDoNotGenerateValuesVectorToPoseIndexes);
 
-				TArray<float> Deviation = FMeanDeviationCalculator::Calculate(SearchIndexBases, Schemas);
+				const TArray<float> Deviation = FMeanDeviationCalculator::Calculate(SearchIndexBases, Schemas);
 
 				// Building FSearchIndex
 				PreprocessSearchIndexWeights(SearchIndex, MainDatabaseSchema, Deviation);
@@ -2008,6 +2010,8 @@ void FPoseSearchDatabaseAsyncCacheTask::OnGetComplete(UE::DerivedData::FCacheGet
 
 				// putting SearchIndex to DDC
 				TArray<uint8> RawBytes;
+				// reserving 20k as initial buffer to serialize the SearchIndex to avoid multiple reallocations
+				RawBytes.Reserve(20 * 1024);
 				FMemoryWriter Writer(RawBytes);
 				Writer << SearchIndex;
 				FSharedBuffer RawData = MakeSharedBufferFromArray(MoveTemp(RawBytes));
@@ -2032,10 +2036,9 @@ void FPoseSearchDatabaseAsyncCacheTask::OnGetComplete(UE::DerivedData::FCacheGet
 #if ENABLE_ANIM_DEBUG
 							if (AnyTestFlags(EMotionMatchTestFlags::ValidateDDC))
 							{
-								TArray<FCacheGetRequest> CacheRequests;
 								const FCacheKey CacheKey{ Bucket, DerivedDataKey };
-								CacheRequests.Add({ { Database->GetPathName() }, CacheKey, ECachePolicy::Default });
-								GetCache().Get(CacheRequests, Owner, [this, MainDatabaseSchema, MainDatabaseName](FCacheGetResponse&& Response)
+								const FCacheGetRequest CacheRequest = { { Database->GetPathName() }, CacheKey, ECachePolicy::Default };
+								GetCache().Get(MakeArrayView(&CacheRequest, 1), Owner, [this, MainDatabaseSchema, MainDatabaseName](FCacheGetResponse&& Response)
 									{
 										const FCacheKey& FullIndexKey = Response.Record.GetKey();
 										check(FullIndexKey.Hash == DerivedDataKey);
@@ -2163,11 +2166,12 @@ void FAsyncPoseSearchDatabasesManagement::SynchronizeDatabases()
 		const TDatabasesToSynchronize DatabasesToSynchronizeCopy = DatabasesToSynchronize;
 		DatabasesToSynchronize.Reset();
 
+		TArray<UAnimSequenceBase*, TInlineAllocator<256>> SequencesBase;
 		for (const TDatabasesToSynchronizePair& Pair : DatabasesToSynchronizeCopy)
 		{
 			if (Pair.Key.IsValid())
 			{
-				TArray<UAnimSequenceBase*> SequencesBase;
+				SequencesBase.Reset();
 				for (const TWeakObjectPtr<UAnimSequenceBase>& SequenceBase : Pair.Value)
 				{
 					if (SequenceBase.IsValid())
@@ -2326,6 +2330,7 @@ void FAsyncPoseSearchDatabasesManagement::Tick(float DeltaTime)
 		};
 
 		TArray<FTestSample> TestSamples;
+		TestSamples.Reserve(256);
 		for (const TUniquePtr<FPoseSearchDatabaseAsyncCacheTask>& TaskPtr : Tasks)
 		{
 			if (const UPoseSearchDatabase* Database = TaskPtr->GetDatabase())
