PoseSearch - adding UPoseSearchDatabase::AppendToClassSchema to support incremental cook

diff --git a/Engine/Plugins/Animation/PoseSearch/Source/Runtime/Private/PoseSearchDatabase.cpp b/Engine/Plugins/Animation/PoseSearch/Source/Runtime/Private/PoseSearchDatabase.cpp
--- a/Engine/Plugins/Animation/PoseSearch/Source/Runtime/Private/PoseSearchDatabase.cpp
+++ b/Engine/Plugins/Animation/PoseSearch/Source/Runtime/Private/PoseSearchDatabase.cpp
@@ -834,6 +834,18 @@ TConstArrayView<float> UPoseSearchDatabase::CalculateDynamicWeightsSqrt(TArrayVi
 	return SearchIndex.WeightsSqrt;
 }
 
+#if WITH_EDITORONLY_DATA
+void UPoseSearchDatabase::AppendToClassSchema(FAppendToClassSchemaContext& Context)
+{
+	using namespace UE::PoseSearch;
+
+	Super::AppendToClassSchema(Context);
+	
+	Context.Update(&DatabaseIndexDerivedDataCacheKeyVersion, sizeof(DatabaseIndexDerivedDataCacheKeyVersion));
+	Context.Update(&FDevSystemGuids::Get().POSESEARCHDB_DERIVEDDATA_VER, sizeof(FDevSystemGuids::Get().POSESEARCHDB_DERIVEDDATA_VER));
+}
+#endif // WITH_EDITORONLY_DATA
+
 const UE::PoseSearch::FSearchIndex& UPoseSearchDatabase::GetSearchIndex() const
 {
 	// making sure the search index is consistent. if it fails the calling code hasn't been protected by FAsyncPoseSearchDatabasesManagement::RequestAsyncBuildIndex
diff --git a/Engine/Plugins/Animation/PoseSearch/Source/Runtime/Private/PoseSearchDerivedDataKey.cpp b/Engine/Plugins/Animation/PoseSearch/Source/Runtime/Private/PoseSearchDerivedDataKey.cpp
--- a/Engine/Plugins/Animation/PoseSearch/Source/Runtime/Private/PoseSearchDerivedDataKey.cpp
+++ b/Engine/Plugins/Animation/PoseSearch/Source/Runtime/Private/PoseSearchDerivedDataKey.cpp
@@ -73,13 +73,13 @@ FKeyBuilder::FKeyBuilder(const UObject* Object, bool bUseDataVer, bool bPerformC
 		Hasher.Reset();
 
 		// used to invalidate the key without having to change POSESEARCHDB_DERIVEDDATA_VER all the times
-		int32 POSESEARCHDB_DERIVEDDATA_VER_SMALL = 2677;
+		int32 NonConstDatabaseIndexDerivedDataCacheKeyVersion = DatabaseIndexDerivedDataCacheKeyVersion;
 		FGuid VersionGuid = FDevSystemGuids::GetSystemGuid(FDevSystemGuids::Get().POSESEARCHDB_DERIVEDDATA_VER);
 		FString AnimationCompressionVersionString = UE::Anim::Compression::AnimationCompressionVersionString;
 
 		*this << VersionGuid;
 		*this << AnimationCompressionVersionString;
-		*this << POSESEARCHDB_DERIVEDDATA_VER_SMALL;
+		*this << NonConstDatabaseIndexDerivedDataCacheKeyVersion;
 
 		LocalCachedHash.Hash = Hasher.Finalize();
 	}
diff --git a/Engine/Plugins/Animation/PoseSearch/Source/Runtime/Public/PoseSearch/PoseSearchDatabase.h b/Engine/Plugins/Animation/PoseSearch/Source/Runtime/Public/PoseSearch/PoseSearchDatabase.h
--- a/Engine/Plugins/Animation/PoseSearch/Source/Runtime/Public/PoseSearch/PoseSearchDatabase.h
+++ b/Engine/Plugins/Animation/PoseSearch/Source/Runtime/Public/PoseSearch/PoseSearchDatabase.h
@@ -603,6 +603,10 @@ public:
 	// Experimental, this feature might be removed without warning, not for production use
 	TConstArrayView<float> CalculateDynamicWeightsSqrt(TArrayView<float> DynamicWeightsSqrtBuffer) const;
 
+#if WITH_EDITORONLY_DATA
+	static void AppendToClassSchema(FAppendToClassSchemaContext& Context);
+#endif // WITH_EDITORONLY_DATA
+
 private:
 	UE::PoseSearch::FSearchResult SearchPCAKDTree(UE::PoseSearch::FSearchContext& SearchContext) const;
 	UE::PoseSearch::FSearchResult SearchVPTree(UE::PoseSearch::FSearchContext& SearchContext) const;
diff --git a/Engine/Plugins/Animation/PoseSearch/Source/Runtime/Public/PoseSearch/PoseSearchDefines.h b/Engine/Plugins/Animation/PoseSearch/Source/Runtime/Public/PoseSearch/PoseSearchDefines.h
--- a/Engine/Plugins/Animation/PoseSearch/Source/Runtime/Public/PoseSearch/PoseSearchDefines.h
+++ b/Engine/Plugins/Animation/PoseSearch/Source/Runtime/Public/PoseSearch/PoseSearchDefines.h
@@ -52,6 +52,7 @@ static constexpr int32 PreallocatedSearchesNum = 4;
 static constexpr float FiniteDelta = 1 / 60.0f; // Time delta used for computing pose derivatives
 static constexpr int32 MaxNumberOfCollectedPoseCandidatesPerDatabase = 256;
 static constexpr int32 InvalidSearchId = INDEX_NONE; // @todo: it should be 0, since the search id is an hash!
+static constexpr int32 DatabaseIndexDerivedDataCacheKeyVersion = 2677; // bump this value to force reindex ALL the UPoseSearchDatabases, without having to change the global FDevSystemGuids::POSESEARCHDB_DERIVEDDATA_VER
 
 template<uint32 Alignment = DEFAULT_ALIGNMENT> class TMemStackSetAllocator : public TSetAllocator<TSparseArrayAllocator<TMemStackAllocator<Alignment>, TMemStackAllocator<Alignment>>, TMemStackAllocator<Alignment>> {};
 
