PoseSearch - fix for out of scope FSearchContext in UPoseSearchLibrary::UpdateMotionMatchingState

diff --git a/Engine/Plugins/Animation/PoseSearch/Source/Runtime/Private/PoseSearchLibrary.cpp b/Engine/Plugins/Animation/PoseSearch/Source/Runtime/Private/PoseSearchLibrary.cpp
--- a/Engine/Plugins/Animation/PoseSearch/Source/Runtime/Private/PoseSearchLibrary.cpp
+++ b/Engine/Plugins/Animation/PoseSearch/Source/Runtime/Private/PoseSearchLibrary.cpp
@@ -678,7 +678,8 @@ void UPoseSearchLibrary::UpdateMotionMatchingState(
 		InOutMotionMatchingState.CurrentSearchResult.Reset();
 	}
 
-	FSearchContext SearchContext(0.f, &InOutMotionMatchingState.PoseIndicesHistory, InOutMotionMatchingState.CurrentSearchResult, PoseJumpThresholdTime, EventToSearch.GetPlayRateOverriddenEvent(PlayRate));
+	const FPoseSearchEvent PlayRateOverriddenEvent = EventToSearch.GetPlayRateOverriddenEvent(PlayRate);
+	FSearchContext SearchContext(0.f, &InOutMotionMatchingState.PoseIndicesHistory, InOutMotionMatchingState.CurrentSearchResult, PoseJumpThresholdTime, PlayRateOverriddenEvent);
 	SearchContext.AddRole(GetCommonDefaultRole(Databases), AnimContext, PoseHistory);
 
 	const bool bCanAdvance = InOutMotionMatchingState.CurrentSearchResult.CanAdvance(DeltaTime);
