PoseSearch - fix to avoid using deallocated memory in FSearchContext

diff --git a/Engine/Plugins/Animation/PoseSearch/Source/Runtime/Private/PoseSearchResult.cpp b/Engine/Plugins/Animation/PoseSearch/Source/Runtime/Private/PoseSearchResult.cpp
--- a/Engine/Plugins/Animation/PoseSearch/Source/Runtime/Private/PoseSearchResult.cpp
+++ b/Engine/Plugins/Animation/PoseSearch/Source/Runtime/Private/PoseSearchResult.cpp
@@ -9,6 +9,9 @@
 namespace UE::PoseSearch
 {
 
+// Experimental, this feature might be removed without warning, not for production use
+const FSearchResult FSearchResult::DefaultSearchResult;
+
 void FSearchResult::Update(float NewAssetTime)
 {
 	if (IsValid())
diff --git a/Engine/Plugins/Animation/PoseSearch/Source/Runtime/Public/PoseSearch/PoseSearchContext.h b/Engine/Plugins/Animation/PoseSearch/Source/Runtime/Public/PoseSearch/PoseSearchContext.h
--- a/Engine/Plugins/Animation/PoseSearch/Source/Runtime/Public/PoseSearch/PoseSearchContext.h
+++ b/Engine/Plugins/Animation/PoseSearch/Source/Runtime/Public/PoseSearch/PoseSearchContext.h
@@ -99,7 +99,7 @@ private:
 
 	// NoTe: mapping Role to the index of the associated asset that this FDebugDrawParams is drawing.
 	// NOT the index of the UPoseSearchSchema::Skeletons! Use UPoseSearchSchema::GetRoledSkeleton API to resolve that Role to FPoseSearchRoledSkeleton 
-	const FRoleToIndex& RoleToIndex;
+	const FRoleToIndex RoleToIndex;
 
 	const UPoseSearchDatabase* Database = nullptr;
 	TAlignedArray<float> DynamicWeightsSqrtBuffer;
@@ -138,7 +138,7 @@ struct POSESEARCH_API FSearchContext
 {
 	UE_DEPRECATED(5.6, "Use other signaure instead")
 	FSearchContext(float InDesiredPermutationTimeOffset = 0.f, const FPoseIndicesHistory* InPoseIndicesHistory = nullptr,
-		const FSearchResult& InCurrentResult = FSearchResult(), const FFloatInterval& InPoseJumpThresholdTime = FFloatInterval(0.f, 0.f), bool bInUseCachedChannelData = false);
+		const FSearchResult& InCurrentResult = FSearchResult::DefaultSearchResult, const FFloatInterval& InPoseJumpThresholdTime = FFloatInterval(0.f, 0.f), bool bInUseCachedChannelData = false);
 
 	FSearchContext(float InDesiredPermutationTimeOffset, const FPoseIndicesHistory* InPoseIndicesHistory,
 		const FSearchResult& InCurrentResult, const FFloatInterval& InPoseJumpThresholdTime, const FPoseSearchEvent& InEventToSearch);
@@ -246,12 +246,12 @@ private:
 	
 	// if AssetsToConsider is not empty, we'll search only for poses from UObject(s) that are in the AssetsToConsider
 	TConstArrayView<const UObject*> AssetsToConsider;
-	const FPoseSearchEvent& EventToSearch;
+	const FPoseSearchEvent EventToSearch;
 
 	const float DesiredPermutationTimeOffset = 0.f;
 	const FPoseIndicesHistory* PoseIndicesHistory = nullptr;
 	const FSearchResult& CurrentResult;
-	const FFloatInterval& PoseJumpThresholdTime;
+	const FFloatInterval PoseJumpThresholdTime;
 	bool bUseCachedChannelData = false;
 
 	TConstArrayView<float> CurrentResultPoseVector;
diff --git a/Engine/Plugins/Animation/PoseSearch/Source/Runtime/Public/PoseSearch/PoseSearchResult.h b/Engine/Plugins/Animation/PoseSearch/Source/Runtime/Public/PoseSearch/PoseSearchResult.h
--- a/Engine/Plugins/Animation/PoseSearch/Source/Runtime/Public/PoseSearch/PoseSearchResult.h
+++ b/Engine/Plugins/Animation/PoseSearch/Source/Runtime/Public/PoseSearch/PoseSearchResult.h
@@ -92,6 +92,9 @@ struct FSearchResult
 			AssetTime == Other.AssetTime &&
 			bIsContinuingPoseSearch == Other.bIsContinuingPoseSearch;
 	}
+
+	// Experimental, this feature might be removed without warning, not for production use
+	static const FSearchResult DefaultSearchResult;
 };
 
 } // namespace UE::PoseSearch
