Add explit current rotation parameter to UPoseSearchTrajectorLibrary::UpdateHistory_WorldSpace instead of assuming it is correct in the sample buffer.

diff --git a/Engine/Plugins/Animation/PoseSearch/Source/Runtime/Private/PoseSearchTrajectoryLibrary.cpp b/Engine/Plugins/Animation/PoseSearch/Source/Runtime/Private/PoseSearchTrajectoryLibrary.cpp
--- a/Engine/Plugins/Animation/PoseSearch/Source/Runtime/Private/PoseSearchTrajectoryLibrary.cpp
+++ b/Engine/Plugins/Animation/PoseSearch/Source/Runtime/Private/PoseSearchTrajectoryLibrary.cpp
@@ -307,6 +307,7 @@ void UPoseSearchTrajectoryLibrary::UpdateHistory_TransformHistory(
 void UPoseSearchTrajectoryLibrary::UpdateHistory_WorldSpace(
 	FPoseSearchQueryTrajectory& Trajectory,
 	FVector CurrentPosition,
+	FQuat CurrentRotation,
 	const FPoseSearchTrajectoryData::FSampling& TrajectoryDataSampling,
 	float DeltaTime)
 {
@@ -331,7 +332,7 @@ void UPoseSearchTrajectoryLibrary::UpdateHistory_WorldSpace(
 			// Note we add current frame at delta time in the future since t = 0 should be the previous frame
 			Trajectory.Samples[NumHistorySamples - 1].AccumulatedSeconds = DeltaTime;
 			Trajectory.Samples[NumHistorySamples - 1].Position = CurrentPosition;
-			Trajectory.Samples[NumHistorySamples - 1].Facing = Trajectory.Samples[NumHistorySamples].Facing;
+			Trajectory.Samples[NumHistorySamples - 1].Facing = CurrentRotation;
 		}
 		else
 		{
diff --git a/Engine/Plugins/Animation/PoseSearch/Source/Runtime/Public/PoseSearch/PoseSearchTrajectoryLibrary.h b/Engine/Plugins/Animation/PoseSearch/Source/Runtime/Public/PoseSearch/PoseSearchTrajectoryLibrary.h
--- a/Engine/Plugins/Animation/PoseSearch/Source/Runtime/Public/PoseSearch/PoseSearchTrajectoryLibrary.h
+++ b/Engine/Plugins/Animation/PoseSearch/Source/Runtime/Public/PoseSearch/PoseSearchTrajectoryLibrary.h
@@ -129,7 +129,7 @@ public:
 
 	// Update the history purely based on current position, without taking into account ground velocities
 	// DeltaTime is the time since the last call to this method
-	static void UpdateHistory_WorldSpace(FPoseSearchQueryTrajectory& Trajectory, FVector CurrentPosition, const FPoseSearchTrajectoryData::FSampling& TrajectoryDataSampling, float DeltaTime);
+	static void UpdateHistory_WorldSpace(FPoseSearchQueryTrajectory& Trajectory, FVector CurrentPosition, FQuat CurrentRotation, const FPoseSearchTrajectoryData::FSampling& TrajectoryDataSampling, float DeltaTime);
 
 	// Update prediction by simulating the movement math for ground locomotion from UCharacterMovementComponent.
 	static void UpdatePrediction_SimulateCharacterMovement(FPoseSearchQueryTrajectory& Trajectory, const FPoseSearchTrajectoryData& TrajectoryData, const FPoseSearchTrajectoryData::FDerived& TrajectoryDataDerived, const FPoseSearchTrajectoryData::FSampling& TrajectoryDataSampling, float DeltaTime);
