PoseSearch - fix for

PoseSearchResult.cpp(13,36): error: default initialization of an object of const type 'const FSearchResult' without a user-provided default constructor
13 | const FSearchResult FSearchResult::DefaultSearchResult;

on monolithic dedicated server with Clang + AutoRTFM

diff --git a/Engine/Plugins/Animation/PoseSearch/Source/Runtime/Private/PoseSearchResult.cpp b/Engine/Plugins/Animation/PoseSearch/Source/Runtime/Private/PoseSearchResult.cpp
--- a/Engine/Plugins/Animation/PoseSearch/Source/Runtime/Private/PoseSearchResult.cpp
+++ b/Engine/Plugins/Animation/PoseSearch/Source/Runtime/Private/PoseSearchResult.cpp
@@ -9,9 +9,6 @@
 namespace UE::PoseSearch
 {
 
-// Experimental, this feature might be removed without warning, not for production use
-const FSearchResult FSearchResult::DefaultSearchResult;
-
 void FSearchResult::Update(float NewAssetTime)
 {
 	if (IsValid())
diff --git a/Engine/Plugins/Animation/PoseSearch/Source/Runtime/Public/PoseSearch/PoseSearchContext.h b/Engine/Plugins/Animation/PoseSearch/Source/Runtime/Public/PoseSearch/PoseSearchContext.h
--- a/Engine/Plugins/Animation/PoseSearch/Source/Runtime/Public/PoseSearch/PoseSearchContext.h
+++ b/Engine/Plugins/Animation/PoseSearch/Source/Runtime/Public/PoseSearch/PoseSearchContext.h
@@ -138,7 +138,7 @@ struct POSESEARCH_API FSearchContext
 {
 	UE_DEPRECATED(5.6, "Use other signaure instead")
 	FSearchContext(float InDesiredPermutationTimeOffset = 0.f, const FPoseIndicesHistory* InPoseIndicesHistory = nullptr,
-		const FSearchResult& InCurrentResult = FSearchResult::DefaultSearchResult, const FFloatInterval& InPoseJumpThresholdTime = FFloatInterval(0.f, 0.f), bool bInUseCachedChannelData = false);
+		const FSearchResult& InCurrentResult = FSearchResult(), const FFloatInterval& InPoseJumpThresholdTime = FFloatInterval(0.f, 0.f), bool bInUseCachedChannelData = false);
 
 	FSearchContext(float InDesiredPermutationTimeOffset, const FPoseIndicesHistory* InPoseIndicesHistory,
 		const FSearchResult& InCurrentResult, const FFloatInterval& InPoseJumpThresholdTime, const FPoseSearchEvent& InEventToSearch);
diff --git a/Engine/Plugins/Animation/PoseSearch/Source/Runtime/Public/PoseSearch/PoseSearchResult.h b/Engine/Plugins/Animation/PoseSearch/Source/Runtime/Public/PoseSearch/PoseSearchResult.h
--- a/Engine/Plugins/Animation/PoseSearch/Source/Runtime/Public/PoseSearch/PoseSearchResult.h
+++ b/Engine/Plugins/Animation/PoseSearch/Source/Runtime/Public/PoseSearch/PoseSearchResult.h
@@ -92,9 +92,6 @@ struct FSearchResult
 			AssetTime == Other.AssetTime &&
 			bIsContinuingPoseSearch == Other.bIsContinuingPoseSearch;
 	}
-
-	// Experimental, this feature might be removed without warning, not for production use
-	static const FSearchResult DefaultSearchResult;
 };
 
 } // namespace UE::PoseSearch
