PoseSearch - MM interaction tick dependencies construction improvements (removing dependency on CMC/Mover - WIP disabling AnimNext support)

diff --git a/Engine/Plugins/Animation/PoseSearch/Source/Runtime/Private/PoseSearchInteractionIsland.cpp b/Engine/Plugins/Animation/PoseSearch/Source/Runtime/Private/PoseSearchInteractionIsland.cpp
--- a/Engine/Plugins/Animation/PoseSearch/Source/Runtime/Private/PoseSearchInteractionIsland.cpp
+++ b/Engine/Plugins/Animation/PoseSearch/Source/Runtime/Private/PoseSearchInteractionIsland.cpp
@@ -7,7 +7,6 @@
 #include "DrawDebugHelpers.h"
 #include "Engine/SkeletalMesh.h"
 #include "Features/IModularFeatures.h"
-#include "GameFramework/CharacterMovementComponent.h"
 #include "PoseSearch/AnimNode_PoseSearchHistoryCollector.h"
 #include "PoseSearch/MultiAnimAsset.h"
 #include "PoseSearch/PoseSearchContext.h"
@@ -150,52 +149,30 @@ static void InitSearchResults(TArray<FInteractionSearchResult>& SearchResults, T
 	}
 }
 
-static bool FindComponentsForTickDependencies(const UObject* AnimContext, UActorComponent*& PreTickComponent, UActorComponent*& PostTickComponent)
+static UActorComponent* FindComponentForTickDependencies(const UObject* AnimContext)
 {
 	check(AnimContext);
-
-	const AActor* Actor = GetContextOwningActor(AnimContext);
-
-	FName MotionMatchingInteractionPreTag = "MotionMatchingInteractionPre";
-	if (UActorComponent* MotionMatchingInteractionPre = Actor->FindComponentByTag<UActorComponent>(MotionMatchingInteractionPreTag))
-	{
-		PreTickComponent = MotionMatchingInteractionPre;
-	}
-	else if (UActorComponent* CharacterMovementComponent = Actor->GetComponentByClass<UCharacterMovementComponent>())
-	{
-		PreTickComponent = CharacterMovementComponent;
-	}
-	else
+	
+	if (const UAnimInstance* AnimInstance = Cast<UAnimInstance>(AnimContext))
 	{
-		UE_LOG(LogPoseSearch, Error, TEXT("FindComponentsForTickDependencies requires AActor %s to have a UCharacterMovementComponent, or a component tagged as 'MotionMatchingInteractionPre' to work!"), *Actor->GetName());
-		PreTickComponent = nullptr;
+		return AnimInstance->GetSkelMeshComponent();
 	}
+	
+	// this is the AnimNext case
+	return const_cast<UActorComponent*>(Cast<UActorComponent>(AnimContext));
+}
 
-	FName MotionMatchingInteractionPostTag = "MotionMatchingInteractionPost";
-	if (UActorComponent* MotionMatchingInteractionPost = Actor->FindComponentByTag<UActorComponent>(MotionMatchingInteractionPostTag))
-	{
-		PostTickComponent = MotionMatchingInteractionPost;
-	}
-	else if (const UAnimInstance* AnimInstance = Cast<UAnimInstance>(AnimContext))
-	{
-		// this is the UAnimInstance / AnimationBluePrint case
-		if (UActorComponent* SkeletalMeshComponent = Actor->GetComponentByClass<USkeletalMeshComponent>())
-		{
-			PostTickComponent = SkeletalMeshComponent;
-		}
-		else
-		{
-			UE_LOG(LogPoseSearch, Error, TEXT("FindComponentsForTickDependencies requires AActor %s to have a USkeletalMeshComponent, or a component tagged as 'MotionMatchingInteractionPost' to work!"), *Actor->GetName());
-			PostTickComponent = nullptr;
-		}
-	}
-	else
+static void AddPrerequisite(FTickFunction& TickFunction, UObject* TargetObject, FTickFunction& TargetTickFunction)
+{
+	check(TargetObject);
+	TickFunction.AddPrerequisite(TargetObject, TargetTickFunction);
+
+#if ENABLE_ANIM_DEBUG
+	if (!TickFunction.GetPrerequisites().Contains(FTickPrerequisite(TargetObject, TargetTickFunction)))
 	{
-		// this is the AnimNext case
-		PostTickComponent = const_cast<UActorComponent*>(Cast<UActorComponent>(AnimContext));
+		UE_LOG(LogPoseSearch, Error, TEXT("UE::PoseSearch::AddPrerequisite, Failed to add prerequisite from [%s] to [%s, %s]!"), *TickFunction.DiagnosticMessage(), *TargetObject->GetName(), *TargetTickFunction.DiagnosticMessage());
 	}
-
-	return PreTickComponent && PostTickComponent;
+#endif // ENABLE_ANIM_DEBUG
 }
 
 // FInteractionSearchContextBase
@@ -495,42 +472,40 @@ IInteractionIslandDependency* FInteractionIsland::FindCustomDependency(UActorCom
 	return nullptr;
 }
 
-void FInteractionIsland::AddTickDependencies(UActorComponent* InPreTickComponent, UActorComponent* InPostTickComponent, bool bInIsMainActor)
+void FInteractionIsland::AddTickDependencies(UActorComponent* TickActorComponent, bool bInIsMainActor)
 {
-	// Add pre-tick function dependencies
-	if(IInteractionIslandDependency* IslandDependency = FindCustomDependency(InPreTickComponent))
-	{
-		IslandDependency->AddPrerequisite(InteractionSubsystem, PreTickFunction, InPreTickComponent);
+	if (IInteractionIslandDependency* IslandDependency = FindCustomDependency(TickActorComponent))
+	{
+		unimplemented();
+		//IslandDependency->AddPrerequisite(InteractionSubsystem, PreTickFunction, InPreTickComponent);
+		//if (bInIsMainActor)
+		//{
+		//	IslandDependency->AddSubsequent(InteractionSubsystem, PreTickFunction, InPostTickComponent);
+		//	IslandDependency->AddPrerequisite(InteractionSubsystem, PostTickFunction, InPostTickComponent);
+
+		//	check(!bHasTickDependencies);
+		//	bHasTickDependencies = true;
+		//}
+		//else
+		//{
+		//	IslandDependency->AddSubsequent(InteractionSubsystem, PostTickFunction, InPostTickComponent);
+		//	check(bHasTickDependencies);
+		//}
 	}
 	else
 	{
-		PreTickFunction.AddPrerequisite(InPreTickComponent, InPreTickComponent->PrimaryComponentTick);
-	}
-
-	// Add post-tick function dependencies
-	if(IInteractionIslandDependency* IslandDependency = FindCustomDependency(InPostTickComponent))
-	{
-		if(bInIsMainActor)
+		// adding to PreTickFunction all the tick dependencies TickActorComponent has, so it runs after all the tick dependencies of ALL the TickActorComponents in this FInteractionIsland
+		for (FTickPrerequisite& TickActorComponentPrerequisite : TickActorComponent->PrimaryComponentTick.GetPrerequisites())
 		{
-			IslandDependency->AddSubsequent(InteractionSubsystem, PreTickFunction, InPostTickComponent);
-			IslandDependency->AddPrerequisite(InteractionSubsystem, PostTickFunction, InPostTickComponent);
-
-			check(!bHasTickDependencies);
-			bHasTickDependencies = true;
+			AddPrerequisite(PreTickFunction, TickActorComponentPrerequisite.PrerequisiteObject.Get(), *TickActorComponentPrerequisite.PrerequisiteTickFunction);
 		}
-		else
-		{
-			IslandDependency->AddSubsequent(InteractionSubsystem, PostTickFunction, InPostTickComponent);
-			check(bHasTickDependencies);
-		}
-	}
-	else
-	{
-		if(bInIsMainActor)
+
+		// Add post-tick function dependencies
+		if (bInIsMainActor)
 		{
 			// it should be this island, but it's not a UObject, so we use the InteractionSubsystem
-			InPostTickComponent->PrimaryComponentTick.AddPrerequisite(InteractionSubsystem, PreTickFunction);
-			PostTickFunction.AddPrerequisite(InPostTickComponent, InPostTickComponent->PrimaryComponentTick);
+			AddPrerequisite(TickActorComponent->PrimaryComponentTick, InteractionSubsystem, PreTickFunction);
+			AddPrerequisite(PostTickFunction, TickActorComponent, TickActorComponent->PrimaryComponentTick);
 
 			check(!bHasTickDependencies);
 			bHasTickDependencies = true;
@@ -538,82 +513,18 @@ void FInteractionIsland::AddTickDependencies(UActorComponent* InPreTickComponent
 		else
 		{
 			// it should be this island, but it's not a UObject, so we use the InteractionSubsystem
-			InPostTickComponent->PrimaryComponentTick.AddPrerequisite(InteractionSubsystem, PostTickFunction);
+			AddPrerequisite(TickActorComponent->PrimaryComponentTick, InteractionSubsystem, PostTickFunction);
 			check(bHasTickDependencies);
 		}
 	}
 }
 
-void FInteractionIsland::InjectToActor(const UObject* AnimContext, bool bAddTickDependencies)
-{
-	check(IsInGameThread());
-
-	// Called by UPoseSearchInteractionSubsystem::Tick when there aren't animation jobs flying. No need to FScopeLock Lock(&Mutex);
-	if (AnimContext)
-	{
-#if ENABLE_ANIM_DEBUG
-		if (bPreTickFunctionExecuted || bPostTickFunctionExecuted)
-		{
-			if (bPreTickFunctionExecuted)
-			{
-				UE_LOG(LogPoseSearch, Error, TEXT("FInteractionIsland::InjectToActor, unexpected FPreTickFunction::ExecuteTick run!"));
-			}
-			else
-			{
-				UE_LOG(LogPoseSearch, Error, TEXT("FInteractionIsland::InjectToActor, unexpected FPostTickFunction::ExecuteTick run!"));
-			}
-				
-			LogTickDependencies();
-		}
-#endif // ENABLE_ANIM_DEBUG
-
-		UActorComponent* PreTickComponent = nullptr;
-		UActorComponent* PostTickComponent = nullptr;
-		if (FindComponentsForTickDependencies(AnimContext, PreTickComponent, PostTickComponent))
-		{
-			const bool bIsMainActor = IslandAnimContexts.IsEmpty();
-
-			//	tick order: 
-			//		PreTickComponents (UCharacterMovementComponent, or UCharacterMoverComponent) ->
-			//			Island.PreTickFunction ->
-			//				first injected actor PostTickComponent (USkeletalMeshComponent, or WIP UAnimNextComponent) ->
-			//					Island.PostTickFunction ->
-			//						other PostTickComponent(s)
-			PreTickComponents.AddUnique(PreTickComponent);
-			PostTickComponents.AddUnique(PostTickComponent);
-			IslandAnimContexts.AddUnique(AnimContext);
-
-			// making sure that if we add a unique PreTickComponent, we add as well a unique PostTickComponent
-			// (so we can remove them later on in a consistent fashion)
-			check(PreTickComponents.Num() == PostTickComponents.Num() && PreTickComponents.Num() == IslandAnimContexts.Num());
-
-			if (bAddTickDependencies)
-			{
-				AddTickDependencies(PreTickComponent, PostTickComponent, bIsMainActor);
-			}
-			else
-			{
-				check(!bHasTickDependencies);
-			}
-		}
-	}
-}
-
-void FInteractionIsland::AddSearchContext(const FInteractionSearchContext& SearchContext)
-{
-#if DO_CHECK
-	check(SearchContext.CheckForConsistency());
-#endif
-	check(IsInGameThread());
-	SearchContexts.Add(SearchContext);
-}
-
 void FInteractionIsland::RemoveTickDependencies(bool bValidateTickDependencies)
 {
 	// Called by UPoseSearchInteractionSubsystem::Tick when there aren't animation jobs flying. No need to FScopeLock Lock(&Mutex);
 	check(IsInGameThread());
 
-	check(PreTickComponents.Num() == PostTickComponents.Num() && PreTickComponents.Num() == IslandAnimContexts.Num());
+	check(TickActorComponents.Num() == IslandAnimContexts.Num());
 
 	if (!bHasTickDependencies)
 	{
@@ -651,48 +562,42 @@ void FInteractionIsland::RemoveTickDependencies(bool bValidateTickDependencies)
 		}
 #endif // ENABLE_ANIM_DEBUG
 
-		for (TWeakObjectPtr<UActorComponent>& PreTickComponentPtr : PreTickComponents)
+		// removing ALL the prerequisites form the PreTickFunction
+		while (!PreTickFunction.GetPrerequisites().IsEmpty())
 		{
-			if (UActorComponent* PreTickComponent = PreTickComponentPtr.Get())
-			{
-				if(IInteractionIslandDependency* IslandDependency = FindCustomDependency(PreTickComponent))
-				{
-					IslandDependency->RemovePrerequisite(InteractionSubsystem, PreTickFunction, PreTickComponent);
-				}
-				else
-				{
-					PreTickFunction.RemovePrerequisite(PreTickComponent, PreTickComponent->PrimaryComponentTick);
-				}
-			}
+			const FTickPrerequisite& PreTickFunctionPrerequisite = PreTickFunction.GetPrerequisites().Last();
+			PreTickFunction.RemovePrerequisite(PreTickFunctionPrerequisite.PrerequisiteObject.Get(), *PreTickFunctionPrerequisite.PrerequisiteTickFunction);
 		}
 
 		bool bMainActor = true;
-		for (TWeakObjectPtr<UActorComponent>& PostTickComponentPtr : PostTickComponents)
+		for (TWeakObjectPtr<UActorComponent>& TickActorComponentPtr : TickActorComponents)
 		{
-			if (UActorComponent* PostTickComponent = PostTickComponentPtr.Get())
+			if (UActorComponent* TickActorComponent = TickActorComponentPtr.Get())
 			{
-				if(IInteractionIslandDependency* IslandDependency = FindCustomDependency(PostTickComponent))
+				if(IInteractionIslandDependency* IslandDependency = FindCustomDependency(TickActorComponent))
 				{
-					if (bMainActor)
-					{
-						IslandDependency->RemoveSubsequent(InteractionSubsystem, PreTickFunction, PostTickComponent);
-						IslandDependency->RemovePrerequisite(InteractionSubsystem, PostTickFunction, PostTickComponent);
-					}
-					else
-					{
-						IslandDependency->RemoveSubsequent(InteractionSubsystem, PostTickFunction, PostTickComponent);
-					}
+					unimplemented();
+
+					//if (bMainActor)
+					//{
+					//	IslandDependency->RemoveSubsequent(InteractionSubsystem, PreTickFunction, PostTickComponent);
+					//	IslandDependency->RemovePrerequisite(InteractionSubsystem, PostTickFunction, PostTickComponent);
+					//}
+					//else
+					//{
+					//	IslandDependency->RemoveSubsequent(InteractionSubsystem, PostTickFunction, PostTickComponent);
+					//}
 				}
 				else
 				{
 					if (bMainActor)
 					{
-						PostTickComponent->PrimaryComponentTick.RemovePrerequisite(InteractionSubsystem, PreTickFunction);
-						PostTickFunction.RemovePrerequisite(PostTickComponent, PostTickComponent->PrimaryComponentTick);
+						TickActorComponent->PrimaryComponentTick.RemovePrerequisite(InteractionSubsystem, PreTickFunction);
+						PostTickFunction.RemovePrerequisite(TickActorComponent, TickActorComponent->PrimaryComponentTick);
 					}
 					else
 					{
-						PostTickComponent->PrimaryComponentTick.RemovePrerequisite(InteractionSubsystem, PostTickFunction);
+						TickActorComponent->PrimaryComponentTick.RemovePrerequisite(InteractionSubsystem, PostTickFunction);
 					}
 				}
 			}
@@ -711,14 +616,74 @@ void FInteractionIsland::RemoveTickDependencies(bool bValidateTickDependencies)
 #endif // ENABLE_ANIM_DEBUG
 }
 
+void FInteractionIsland::InjectToActor(const UObject* AnimContext, bool bAddTickDependencies)
+{
+	check(IsInGameThread());
+
+	// Called by UPoseSearchInteractionSubsystem::Tick when there aren't animation jobs flying. No need to FScopeLock Lock(&Mutex);
+	if (AnimContext)
+	{
+#if ENABLE_ANIM_DEBUG
+		if (bPreTickFunctionExecuted || bPostTickFunctionExecuted)
+		{
+			if (bPreTickFunctionExecuted)
+			{
+				UE_LOG(LogPoseSearch, Error, TEXT("FInteractionIsland::InjectToActor, unexpected FPreTickFunction::ExecuteTick run!"));
+			}
+			else
+			{
+				UE_LOG(LogPoseSearch, Error, TEXT("FInteractionIsland::InjectToActor, unexpected FPostTickFunction::ExecuteTick run!"));
+			}
+				
+			LogTickDependencies();
+		}
+#endif // ENABLE_ANIM_DEBUG
+
+		if (UActorComponent* TickActorComponent = FindComponentForTickDependencies(AnimContext))
+		{
+			const bool bIsMainActor = IslandAnimContexts.IsEmpty();
+
+			//	tick order: 
+			//		ALL TickActorComponents prerequisites (ultimately we're looking to have UCharacterMovementComponent or UCharacterMoverComponent ticked) ->
+			//			Island.PreTickFunction ->
+			//				first injected TickActorComponent (USkeletalMeshComponent, or UAnimNextComponent) ->
+			//					Island.PostTickFunction ->
+			//						other TickActorComponent(s)
+			TickActorComponents.AddUnique(TickActorComponent);
+			IslandAnimContexts.AddUnique(AnimContext);
+
+			// making sure that if we add a unique TickActorComponent, we add as well a unique PostTickComponent
+			// (so we can remove them later on in a consistent fashion)
+			check(TickActorComponents.Num() == IslandAnimContexts.Num());
+
+			if (bAddTickDependencies)
+			{
+				AddTickDependencies(TickActorComponent, bIsMainActor);
+			}
+			else
+			{
+				check(!bHasTickDependencies);
+			}
+		}
+	}
+}
+
+void FInteractionIsland::AddSearchContext(const FInteractionSearchContext& SearchContext)
+{
+#if DO_CHECK
+	check(SearchContext.CheckForConsistency());
+#endif
+	check(IsInGameThread());
+	SearchContexts.Add(SearchContext);
+}
+
 void FInteractionIsland::Uninitialize(bool bValidateTickDependencies)
 {
 	RemoveTickDependencies(bValidateTickDependencies);
 
 	if (IsInitialized())
 	{
-		PreTickComponents.Reset();
-		PostTickComponents.Reset();
+		TickActorComponents.Reset();
 		IslandAnimContexts.Reset();
 
 		SearchContexts.Reset();
@@ -727,7 +692,7 @@ void FInteractionIsland::Uninitialize(bool bValidateTickDependencies)
 	}
 	else
 	{
-		check(PreTickComponents.IsEmpty() && PostTickComponents.IsEmpty() && IslandAnimContexts.IsEmpty() && SearchContexts.IsEmpty() && SearchResults.IsEmpty() && !bSearchPerfomed);
+		check(TickActorComponents.IsEmpty() && IslandAnimContexts.IsEmpty() && SearchContexts.IsEmpty() && SearchResults.IsEmpty() && !bSearchPerfomed);
 	}
 }
 
@@ -766,7 +731,29 @@ bool FInteractionIsland::CanPerformSearch(const UObject* AnimContext) const
 void FInteractionIsland::LogTickDependencies() const
 {
 	const ENamedThreads::Type Type = IsInGameThread() ? ENamedThreads::GameThread : ENamedThreads::AnyThread;
-	const_cast<FPostTickFunction&>(PostTickFunction).LogTickFunction(Type, true, 1);
+
+	for (const TWeakObjectPtr<UActorComponent>& TickActorComponentPtr : TickActorComponents)
+	{
+		if (UActorComponent* TickActorComponent = TickActorComponentPtr.Get())
+		{
+			const int32 InteractionIslandIndex = InteractionSubsystem->GetInteractionIslands().IndexOfByKey(this);
+
+			UE_LOG(LogPoseSearch, Log, TEXT("============== %s (Island %d) =============="), *TickActorComponent->GetOwner()->GetName(), InteractionIslandIndex);
+			if (IInteractionIslandDependency* IslandDependency = FindCustomDependency(TickActorComponent))
+			{
+				// AnimNextComponent case
+				// @todo: figure out a better way of gathering AnimNextComponent tick function
+				UE_LOG(LogPoseSearch, Log, TEXT("tick %s"), *TickActorComponent->GetName());
+				const_cast<FPostTickFunction&>(PostTickFunction).LogTickFunction(Type, true, 1);
+			}
+			else
+			{
+				// SkeletalMeshComponent / AnimInstance case
+				TickActorComponent->PrimaryComponentTick.LogTickFunction(Type, true, 1);
+			}
+		}
+	}
+
 }
 #endif // ENABLE_ANIM_DEBUG
 
diff --git a/Engine/Plugins/Animation/PoseSearch/Source/Runtime/Private/PoseSearchInteractionLibrary.cpp b/Engine/Plugins/Animation/PoseSearch/Source/Runtime/Private/PoseSearchInteractionLibrary.cpp
--- a/Engine/Plugins/Animation/PoseSearch/Source/Runtime/Private/PoseSearchInteractionLibrary.cpp
+++ b/Engine/Plugins/Animation/PoseSearch/Source/Runtime/Private/PoseSearchInteractionLibrary.cpp
@@ -41,35 +41,3 @@ FPoseSearchContinuingProperties UPoseSearchInteractionLibrary::GetMontageContinu
 	}
 	return ContinuingProperties;
 }
-
-TArray<APawn*> UPoseSearchInteractionLibrary::GetInteractingPawns(APawn* Pawn)
-{
-	using namespace UE::PoseSearch;
-
-	TArray<APawn*> InteractingPawns;
-	if (const USkeletalMeshComponent* SkeletalMeshComponent = Pawn->GetComponentByClass<USkeletalMeshComponent>())
-	{
-		if (const UAnimInstance* AnimInstance = SkeletalMeshComponent->GetAnimInstance())
-		{
-			if (const UPoseSearchInteractionSubsystem* InteractionSubsystem = UPoseSearchInteractionSubsystem::GetSubsystem_AnyThread(AnimInstance))
-			{
-				for (const UObject* InteractingAnimContext : InteractionSubsystem->GetInteractingAnimContextsWith(AnimInstance))
-				{
-					check(InteractingAnimContext);
-					if (const UAnimInstance* InteractingAnimInstance = Cast<UAnimInstance>(InteractingAnimContext))
-					{
-						if (APawn* InteractingPawn = Cast<APawn>(InteractingAnimInstance->GetOwningActor()))
-						{
-							InteractingPawns.Add(InteractingPawn);
-						}
-					}
-					else
-					{
-						unimplemented();
-					}
-				}
-			}
-		}
-	}
-	return InteractingPawns;
-}
diff --git a/Engine/Plugins/Animation/PoseSearch/Source/Runtime/Private/PoseSearchInteractionSubsystem.cpp b/Engine/Plugins/Animation/PoseSearch/Source/Runtime/Private/PoseSearchInteractionSubsystem.cpp
--- a/Engine/Plugins/Animation/PoseSearch/Source/Runtime/Private/PoseSearchInteractionSubsystem.cpp
+++ b/Engine/Plugins/Animation/PoseSearch/Source/Runtime/Private/PoseSearchInteractionSubsystem.cpp
@@ -22,6 +22,9 @@ static FAutoConsoleVariableRef CVarPoseSearchInteractionEnabled(TEXT("a.PoseSear
 
 static bool GVarPoseSearchInteractionCacheIslands = true;
 static FAutoConsoleVariableRef CVarPoseSearchInteractionCacheIslands(TEXT("a.PoseSearchInteraction.CacheIslands"), GVarPoseSearchInteractionCacheIslands, TEXT("Cache Pose Search Interaction Islands for future reuse instead of destrying them"));
+
+static bool GVarPoseSearchInteractionLoglandsTickDependencies = false;
+static FAutoConsoleVariableRef CVarPoseSearchInteractionLoglandsTickDependencies(TEXT("a.PoseSearchInteraction.LoglandsTickDependencies"), GVarPoseSearchInteractionLoglandsTickDependencies, TEXT("Log islands tick dependencies"));
 #endif // !NO_CVARS
 
 FString FInteractionAvailability::GetPoseHistoryName() const
@@ -287,6 +290,9 @@ void UPoseSearchInteractionSubsystem::RegenerateAllIslands(TConstArrayView<UE::P
 			Island.AddSearchContext(SearchContext);
 		}
 
+		// @todo: we sort SearchContextGroup.AnimContexts depending on how much any of them want 
+		//        to be the MainAnimContext, to avoid external conflicting tick dependencies
+		
 		// injecting tick dependencies between island AnimContext
 		for (const UObject* AnimContext : SearchContextGroup.AnimContexts)
 		{
@@ -300,10 +306,8 @@ bool UPoseSearchInteractionSubsystem::ValidateAllIslands() const
 {
 	using namespace UE::PoseSearch;
 
-	TSet<TWeakObjectPtr<UActorComponent>> PreTickComponents;
-	TSet<TWeakObjectPtr<UActorComponent>> PostTickComponents;
+	TSet<TWeakObjectPtr<UActorComponent>> TickActorComponents;
 
-	// @todo: validate IslandAnimContexts reference the same actors as PreTickComponents and PostTickComponents
 	typedef TSet<const UObject*> FIslandAnimContexts;
 	TArray<FIslandAnimContexts> IslandsAnimContexts;
 
@@ -313,18 +317,9 @@ bool UPoseSearchInteractionSubsystem::ValidateAllIslands() const
 	bool bAlreadyInSet = false;
 	for (const FInteractionIsland* Island : Islands)
 	{
-		for (const TWeakObjectPtr<UActorComponent>& PreTickComponent : Island->GetPreTickComponents())
-		{
-			PreTickComponents.Add(PreTickComponent, &bAlreadyInSet);
-			if (bAlreadyInSet)
-			{
-				return false;
-			}
-		}
-
-		for (const TWeakObjectPtr<UActorComponent>& PostTickComponent : Island->GetPostTickComponents())
+		for (const TWeakObjectPtr<UActorComponent>& TickActorComponent : Island->GetTickActorComponents())
 		{
-			PostTickComponents.Add(PostTickComponent, &bAlreadyInSet);
+			TickActorComponents.Add(TickActorComponent, &bAlreadyInSet);
 			if (bAlreadyInSet)
 			{
 				return false;
@@ -1095,11 +1090,15 @@ void UPoseSearchInteractionSubsystem::Tick(float DeltaSeconds)
 
 	// populating the continuing pose properties for the SearchContexts from the current Islands
 	PopulateContinuingProperties(DeltaSeconds, SearchContexts);
-	
+
 #if ENABLE_DRAW_DEBUG && ENABLE_ANIM_DEBUG
 	// drawing the current frame islands to be consistent with the search, before regenerating the islands with the newly published availabilities
 	DebugDrawIslands();
 #endif // ENABLE_DRAW_DEBUG && ENABLE_ANIM_DEBUG
+	
+#if ENABLE_ANIM_DEBUG
+	DebugLogTickDependencies();
+#endif // ENABLE_ANIM_DEBUG
 
 	RegenerateAllIslands(SearchContexts);
 
@@ -1245,3 +1244,24 @@ void UPoseSearchInteractionSubsystem::DebugDrawIslands() const
 	}
 }
 #endif // ENABLE_DRAW_DEBUG && ENABLE_ANIM_DEBUG
+
+#if ENABLE_ANIM_DEBUG
+void UPoseSearchInteractionSubsystem::DebugLogTickDependencies() const
+{
+#if !NO_CVARS
+	using namespace UE::PoseSearch;
+
+	if (GVarPoseSearchInteractionLoglandsTickDependencies)
+	{
+		UE_LOG(LogPoseSearch, Log, TEXT("=================================================================="));
+		for (const FInteractionIsland* Island : Islands)
+		{
+			if (Island->IsInitialized())
+			{
+				Island->LogTickDependencies();
+			}
+		}
+	}
+#endif // !NO_CVARS
+}
+#endif // ENABLE_ANIM_DEBUG
diff --git a/Engine/Plugins/Animation/PoseSearch/Source/Runtime/Private/PoseSearchInteractionValidator.cpp b/Engine/Plugins/Animation/PoseSearch/Source/Runtime/Private/PoseSearchInteractionValidator.cpp
--- a/Engine/Plugins/Animation/PoseSearch/Source/Runtime/Private/PoseSearchInteractionValidator.cpp
+++ b/Engine/Plugins/Animation/PoseSearch/Source/Runtime/Private/PoseSearchInteractionValidator.cpp
@@ -41,7 +41,7 @@ FInteractionValidator::FInteractionValidator(const UObject* AnimContext)
 		const int32 PreviousCounter = ValidatingIsland->InteractionIslandThreadSafeCounter.Set(-1);
 		if (PreviousCounter != 0)
 		{
-			UE_LOG(LogPoseSearch, Error, TEXT("Non thread safe call! Why is there any other actor running while we schedule the MainActor? %d"), PreviousCounter);
+			UE_LOG(LogPoseSearch, Error, TEXT("Non thread safe call! Why is there any other actor running while we schedule the MainActor (%s)? %d"), *MainActor->GetName(), PreviousCounter);
 			ValidatingIsland->LogTickDependencies();
 		}
 	}
@@ -51,7 +51,7 @@ FInteractionValidator::FInteractionValidator(const UObject* AnimContext)
 		const int32 PreviousCounter = ValidatingIsland->InteractionIslandThreadSafeCounter.Add(1);
 		if (PreviousCounter < 0)
 		{
-			UE_LOG(LogPoseSearch, Error, TEXT("Non thread safe call! The MainActor is running! Nobody else in the same island should be running! %d"), PreviousCounter);
+			UE_LOG(LogPoseSearch, Error, TEXT("Non thread safe call! The MainActor (%s) is running! Nobody else (%s) in the same island should be running! %d"), *MainActor->GetName(), *Actor->GetName(),  PreviousCounter);
 			ValidatingIsland->LogTickDependencies();
 		}
 	}
@@ -115,7 +115,7 @@ FInteractionValidator::~FInteractionValidator()
 		const int32 PreviousCounter = Island->InteractionIslandThreadSafeCounter.Set(0);
 		if (PreviousCounter >= 0)
 		{
-			UE_LOG(LogPoseSearch, Error, TEXT("Non thread safe call! Why was there some other actor running while we ended the scheduling the MainActor? %d"), PreviousCounter);
+			UE_LOG(LogPoseSearch, Error, TEXT("Non thread safe call! Why was there some other actor running while we ended the scheduling the MainActor (%s)? %d"), *MainActor->GetName(), PreviousCounter);
 		}
 	}
 	else
@@ -124,7 +124,7 @@ FInteractionValidator::~FInteractionValidator()
 		const int32 PreviousCounter = Island->InteractionIslandThreadSafeCounter.Add(-1);
 		if (PreviousCounter <= 0)
 		{
-			UE_LOG(LogPoseSearch, Error, TEXT("Non thread safe call! The MainActor is running! Nobody else in the same island should be running! %d"), PreviousCounter);
+			UE_LOG(LogPoseSearch, Error, TEXT("Non thread safe call! The MainActor (%s) is running! Nobody else (%s) in the same island should be running! %d"), *MainActor->GetName(), *Actor->GetName(), PreviousCounter);
 		}
 	}
 }
diff --git a/Engine/Plugins/Animation/PoseSearch/Source/Runtime/Public/PoseSearch/PoseSearchInteractionIsland.h b/Engine/Plugins/Animation/PoseSearch/Source/Runtime/Public/PoseSearch/PoseSearchInteractionIsland.h
--- a/Engine/Plugins/Animation/PoseSearch/Source/Runtime/Public/PoseSearch/PoseSearchInteractionIsland.h
+++ b/Engine/Plugins/Animation/PoseSearch/Source/Runtime/Public/PoseSearch/PoseSearchInteractionIsland.h
@@ -7,7 +7,6 @@
 
 struct FPoseSearchInteractionBlueprintResult;
 struct FPoseSearchContinuingProperties;
-class UCharacterMovementComponent;
 class UMultiAnimAsset;
 class UPoseSearchInteractionSubsystem;
 
@@ -77,8 +76,8 @@ struct FInteractionSearchResult : public FSearchResult
 };
 
 // Experimental, this feature might be removed without warning, not for production use
-// FInteractionIsland contains ticks functions injected between the interacting actors PreTickComponent (UCharacterMovementComponent, or UCharacterMoverComponent)
-// and PostTickComponent (USkeletalMeshComponent, or WIP UAnimNextComponent)
+// FInteractionIsland contains ticks functions injected between the interacting actors TickActorComponents (UCharacterMovementComponent, or UCharacterMoverComponent)
+// and PostTickComponent (USkeletalMeshComponent, or UAnimNextComponent)
 // to create a execution threading fence to be able to perform motion matching searches between the involved characters in a thread safe manner.
 // Look at UPoseSearchInteractionSubsystem "Execution model and threading details" for additional information
 struct FInteractionIsland
@@ -90,8 +89,7 @@ struct FInteractionIsland
 
 	bool DoSearch_AnyThread(const UObject* AnimContext, const FPoseSearchContinuingProperties& ContinuingProperties, FPoseSearchInteractionBlueprintResult& Result);
 
-	const TArray<TWeakObjectPtr<UActorComponent>>& GetPreTickComponents() const { return PreTickComponents; }
-	const TArray<TWeakObjectPtr<UActorComponent>>& GetPostTickComponents() const { return PostTickComponents; }
+	const TArray<TWeakObjectPtr<UActorComponent>>& GetTickActorComponents() const { return TickActorComponents; }
 	const TArray<TWeakObjectPtr<const UObject>>& GetIslandAnimContexts() const { return IslandAnimContexts; }
 	const TArray<FInteractionSearchContext>& GetSearchContexts() const { return SearchContexts; }
 	const TArray<FInteractionSearchResult>& GetSearchResults() const { return SearchResults; }
@@ -104,9 +102,13 @@ struct FInteractionIsland
 
 	bool HasTickDependencies() const;
 	void InjectToActor(const UObject* AnimContext, bool bAddTickDependencies);
+	
+#if ENABLE_ANIM_DEBUG
+	void LogTickDependencies() const;
+#endif // ENABLE_ANIM_DEBUG
 
 private:
-	void AddTickDependencies(UActorComponent* InPreTickComponent, UActorComponent* InPostTickComponent, bool bInIsMainActor);
+	void AddTickDependencies(UActorComponent* TickActorComponent, bool bInIsMainActor);
 	void RemoveTickDependencies(bool bValidateTickDependencies);
 	IInteractionIslandDependency* FindCustomDependency(UActorComponent* InTickComponent) const;
 
@@ -115,10 +117,6 @@ private:
 	const AActor* GetMainActor() const;
 	bool CanPerformSearch(const UObject* AnimContext) const;
 
-#if ENABLE_ANIM_DEBUG
-	void LogTickDependencies() const;
-#endif // ENABLE_ANIM_DEBUG
-	
 	struct FPreTickFunction : public FTickFunction
 	{
 		virtual void ExecuteTick(float DeltaTime, enum ELevelTick TickType, ENamedThreads::Type CurrentThread, const FGraphEventRef& MyCompletionGraphEvent) override;
@@ -136,8 +134,7 @@ private:
 	FPostTickFunction PostTickFunction;
 	bool bHasTickDependencies = false;
 
-	TArray<TWeakObjectPtr<UActorComponent>> PreTickComponents;
-	TArray<TWeakObjectPtr<UActorComponent>> PostTickComponents;
+	TArray<TWeakObjectPtr<UActorComponent>> TickActorComponents;
 	// all the AnimContext in this island. Each SearchContexts will contain a subset of IslandAnimContexts 
 	TArray<TWeakObjectPtr<const UObject>> IslandAnimContexts;
 
diff --git a/Engine/Plugins/Animation/PoseSearch/Source/Runtime/Public/PoseSearch/PoseSearchInteractionLibrary.h b/Engine/Plugins/Animation/PoseSearch/Source/Runtime/Public/PoseSearch/PoseSearchInteractionLibrary.h
--- a/Engine/Plugins/Animation/PoseSearch/Source/Runtime/Public/PoseSearch/PoseSearchInteractionLibrary.h
+++ b/Engine/Plugins/Animation/PoseSearch/Source/Runtime/Public/PoseSearch/PoseSearchInteractionLibrary.h
@@ -132,8 +132,4 @@ public:
 
 	UFUNCTION(Experimental, BlueprintPure, Category = "Animation|Pose Search", meta = (BlueprintThreadSafe))
 	static FPoseSearchContinuingProperties GetMontageContinuingProperties(UAnimInstance* AnimInstance);
-
-	// @todo: make it BlueprintThreadSafe
-	UFUNCTION(Experimental, BlueprintPure, Category = "Animation|Pose Search")
-	static TArray<APawn*> GetInteractingPawns(APawn* Pawn);
 };
\ No newline at end of file
diff --git a/Engine/Plugins/Animation/PoseSearch/Source/Runtime/Public/PoseSearch/PoseSearchInteractionSubsystem.h b/Engine/Plugins/Animation/PoseSearch/Source/Runtime/Public/PoseSearch/PoseSearchInteractionSubsystem.h
--- a/Engine/Plugins/Animation/PoseSearch/Source/Runtime/Public/PoseSearch/PoseSearchInteractionSubsystem.h
+++ b/Engine/Plugins/Animation/PoseSearch/Source/Runtime/Public/PoseSearch/PoseSearchInteractionSubsystem.h
@@ -97,19 +97,19 @@ struct FTagToDatabases;
 //			{
 //				for each Actor[k][i]
 //				{
-//					Tick CharacterMovementComponent[k][i] in parallel
+//					Tick all the TickActorComponents prerequisites, such as CharacterMovementComponent[k][i] (or Mover) in parallel
 //				}
 // 
-//				Tick Island[k].PreTickFunction (that eventually generates the trajectories with all the updated CMCs)
+//				Tick Island[k].PreTickFunction (that eventually generates the trajectories with all the updated CMCs or Mover)
 //
-//				Tick Actor[k][0].SkeletalMeshComponent (that performs the MotionMatchInteraction queries for all the involved actors via DoSearch_AnyThread)
+//				Tick Actor[k][0].SkeletalMeshComponent (or AnimNextComponent, that performs the MotionMatchInteraction queries for all the involved actors via DoSearch_AnyThread)
 // 
-//				Tick Island[k].PostTickFunction (currently just a threading fence for the execution of all the other SkeletalMeshComponent(s))
+//				Tick Island[k].PostTickFunction (currently just a threading fence for the execution of all the other SkeletalMeshComponent(s) or AnimNextComponent(s))
 // 
 //				for each Actor[k][i]
 //				{
 //					if (i != 0)
-//						Tick SkeletalMeshComponent[k][i] (that DoSearch_AnyThread get the cached result calculated by Tick Actor[k][0].SkeletalMeshComponent) in parallel
+//						Tick SkeletalMeshComponent[k][i] (or AnimNextComponent(s) that DoSearch_AnyThread get the cached result calculated by Tick Actor[k][0].SkeletalMeshComponent) in parallel
 //				}
 //			}
 // - next frame UPoseSearchInteractionLibrary::MotionMatchInteraction_Pure(TArray<FPoseSearchInteractionAvailability> Availabilities, UObject* AnimInstance), with the context of all the published 
@@ -148,6 +148,8 @@ public:
 	// returns and array containing all the other AnimContext(s) that the input AnimContext is currently interacting with
 	TArray<const UObject*> GetInteractingAnimContextsWith(const UObject* AnimContext) const;
 
+	TConstArrayView<UE::PoseSearch::FInteractionIsland*> GetInteractionIslands() const { return Islands; }
+
 private:
 #if ENABLE_ANIM_DEBUG
 	friend struct UE::PoseSearch::FInteractionValidator;
@@ -171,6 +173,10 @@ private:
 	void DebugDrawIslands() const;
 	const UE::PoseSearch::FInteractionAvailabilities* FindAvailabilities(const UObject* AnimContext) const;
 #endif // ENABLE_DRAW_DEBUG && ENABLE_ANIM_DEBUG
+	
+#if ENABLE_ANIM_DEBUG
+	void DebugLogTickDependencies() const;
+#endif // ENABLE_ANIM_DEBUG
 
 	bool IsThereAnyAvailability() const;
 	void AddAvailabilities(const TArrayView<const FPoseSearchInteractionAvailability> Availabilities, const UObject* AnimContext, FName PoseHistoryName, const UE::PoseSearch::IPoseHistory* PoseHistory);
