Motion Matching: Fix for time normalization for multi anim blendspaces

diff --git a/Engine/Plugins/Animation/PoseSearch/Source/Runtime/Private/PoseSearchDatabase.cpp b/Engine/Plugins/Animation/PoseSearch/Source/Runtime/Private/PoseSearchDatabase.cpp
--- a/Engine/Plugins/Animation/PoseSearch/Source/Runtime/Private/PoseSearchDatabase.cpp
+++ b/Engine/Plugins/Animation/PoseSearch/Source/Runtime/Private/PoseSearchDatabase.cpp
@@ -298,6 +298,21 @@ static bool IsBlendSpaceRootMotionEnabled(const UBlendSpace* BlendSpace)
 	return bIsRootMotionUsedInBlendSpace;
 }
 
+static bool IsDatabaseAssetIndexBlendSpace(const UPoseSearchDatabase& Database, int32 AnimationAssetIndex)
+{
+	if (Database.GetDatabaseAnimationAsset<FPoseSearchDatabaseBlendSpace>(AnimationAssetIndex))
+	{
+		return true;
+	}
+	else if (const FPoseSearchDatabaseMultiAnimAsset* MultiAnimAsset =  Database.GetDatabaseAnimationAsset<FPoseSearchDatabaseMultiAnimAsset>(AnimationAssetIndex))
+	{
+		MultiAnimAsset->GetAnimationAssetForRole(Database.Schema->GetDefaultRole());
+		return true;
+	}
+
+	return false;
+}
+
 } // namespace UE::PoseSearch
 
 //////////////////////////////////////////////////////////////////////////
@@ -1267,7 +1282,7 @@ float UPoseSearchDatabase::GetNormalizedAssetTime(int32 PoseIdx) const
 {
 	check(Schema);
 	const UE::PoseSearch::FSearchIndexAsset& Asset = GetSearchIndex().GetAssetForPose(PoseIdx);
-	const bool bIsBlendSpace = GetDatabaseAnimationAsset<FPoseSearchDatabaseBlendSpace>(Asset.GetSourceAssetIdx()) != nullptr;
+	const bool bIsBlendSpace = UE::PoseSearch::IsDatabaseAssetIndexBlendSpace(*this, Asset.GetSourceAssetIdx());
 
 	// sequences or anim composites
 	float AssetTime = Asset.GetTimeFromPoseIndex(PoseIdx, Schema->SampleRate);
diff --git a/Engine/Plugins/Experimental/AnimNextPoseSearch/Source/AnimNextPoseSearch/Private/MotionMatchingTrait.cpp b/Engine/Plugins/Experimental/AnimNextPoseSearch/Source/AnimNextPoseSearch/Private/MotionMatchingTrait.cpp
--- a/Engine/Plugins/Experimental/AnimNextPoseSearch/Source/AnimNextPoseSearch/Private/MotionMatchingTrait.cpp
+++ b/Engine/Plugins/Experimental/AnimNextPoseSearch/Source/AnimNextPoseSearch/Private/MotionMatchingTrait.cpp
@@ -137,7 +137,7 @@ namespace UE::AnimNext
 			{
 				if (const FPoseSearchDatabaseAnimationAssetBase* DatabaseAnimationAssetBase = MotionMatchingState.CurrentSearchResult.Database->GetDatabaseAnimationAsset<FPoseSearchDatabaseAnimationAssetBase>(*SearchIndexAsset))
 				{
-					const UObject* AnimationAsset = DatabaseAnimationAssetBase->GetAnimationAsset();
+					const UObject* AnimationAsset = DatabaseAnimationAssetBase->GetAnimationAssetForRole(MotionMatchingState.CurrentSearchResult.Database->Schema->GetDefaultRole());
 					if (const UBlendSpace* BlendSpace = Cast<UBlendSpace>(AnimationAsset))
 					{
 						// Convert to normalized time.
