AnimNextPoseSearch - reenabling AnimNext support for MM interactions

diff --git a/Engine/Plugins/Animation/PoseSearch/Source/Runtime/Private/PoseSearchInteractionIsland.cpp b/Engine/Plugins/Animation/PoseSearch/Source/Runtime/Private/PoseSearchInteractionIsland.cpp
--- a/Engine/Plugins/Animation/PoseSearch/Source/Runtime/Private/PoseSearchInteractionIsland.cpp
+++ b/Engine/Plugins/Animation/PoseSearch/Source/Runtime/Private/PoseSearchInteractionIsland.cpp
@@ -474,23 +474,35 @@ IInteractionIslandDependency* FInteractionIsland::FindCustomDependency(UActorCom
 
 void FInteractionIsland::AddTickDependencies(UActorComponent* TickActorComponent, bool bInIsMainActor)
 {
+	check(TickActorComponent);
+
 	if (IInteractionIslandDependency* IslandDependency = FindCustomDependency(TickActorComponent))
 	{
-		unimplemented();
-		//IslandDependency->AddPrerequisite(InteractionSubsystem, PreTickFunction, InPreTickComponent);
-		//if (bInIsMainActor)
-		//{
-		//	IslandDependency->AddSubsequent(InteractionSubsystem, PreTickFunction, InPostTickComponent);
-		//	IslandDependency->AddPrerequisite(InteractionSubsystem, PostTickFunction, InPostTickComponent);
-
-		//	check(!bHasTickDependencies);
-		//	bHasTickDependencies = true;
-		//}
-		//else
-		//{
-		//	IslandDependency->AddSubsequent(InteractionSubsystem, PostTickFunction, InPostTickComponent);
-		//	check(bHasTickDependencies);
-		//}
+		if (const FTickFunction* TickActorComponentTickFunction = IslandDependency->FindTickFunction(TickActorComponent))
+		{
+			for (const FTickPrerequisite& TickActorComponentPrerequisite : TickActorComponentTickFunction->GetPrerequisites())
+			{
+				AddPrerequisite(PreTickFunction, TickActorComponentPrerequisite.PrerequisiteObject.Get(), *TickActorComponentPrerequisite.PrerequisiteTickFunction);
+			}
+		}
+		else
+		{
+			UE_LOG(LogPoseSearch, Error, TEXT("FInteractionIsland::AddTickDependencies, error while retrieving the tick function for %s"), *TickActorComponent->GetName());
+		}
+
+		if(bInIsMainActor)
+		{
+			IslandDependency->AddSubsequent(InteractionSubsystem, PreTickFunction, TickActorComponent);
+			IslandDependency->AddPrerequisite(InteractionSubsystem, PostTickFunction, TickActorComponent);
+
+			check(!bHasTickDependencies);
+			bHasTickDependencies = true;
+		}
+		else
+		{
+			IslandDependency->AddSubsequent(InteractionSubsystem, PostTickFunction, TickActorComponent);
+			check(bHasTickDependencies);
+		}
 	}
 	else
 	{
@@ -576,17 +588,15 @@ void FInteractionIsland::RemoveTickDependencies(bool bValidateTickDependencies)
 			{
 				if(IInteractionIslandDependency* IslandDependency = FindCustomDependency(TickActorComponent))
 				{
-					unimplemented();
-
-					//if (bMainActor)
-					//{
-					//	IslandDependency->RemoveSubsequent(InteractionSubsystem, PreTickFunction, PostTickComponent);
-					//	IslandDependency->RemovePrerequisite(InteractionSubsystem, PostTickFunction, PostTickComponent);
-					//}
-					//else
-					//{
-					//	IslandDependency->RemoveSubsequent(InteractionSubsystem, PostTickFunction, PostTickComponent);
-					//}
+					if (bMainActor)
+					{
+						IslandDependency->RemoveSubsequent(InteractionSubsystem, PreTickFunction, TickActorComponent);
+						IslandDependency->RemovePrerequisite(InteractionSubsystem, PostTickFunction, TickActorComponent);
+					}
+					else
+					{
+						IslandDependency->RemoveSubsequent(InteractionSubsystem, PostTickFunction, TickActorComponent);
+					}
 				}
 				else
 				{
@@ -742,9 +752,14 @@ void FInteractionIsland::LogTickDependencies() const
 			if (IInteractionIslandDependency* IslandDependency = FindCustomDependency(TickActorComponent))
 			{
 				// AnimNextComponent case
-				// @todo: figure out a better way of gathering AnimNextComponent tick function
-				UE_LOG(LogPoseSearch, Log, TEXT("tick %s"), *TickActorComponent->GetName());
-				const_cast<FPostTickFunction&>(PostTickFunction).LogTickFunction(Type, true, 1);
+				if (const FTickFunction* TickActorComponentTickFunction = IslandDependency->FindTickFunction(TickActorComponent))
+				{
+					const_cast<FTickFunction*>(TickActorComponentTickFunction)->LogTickFunction(Type, true, 1);
+				}
+				else
+				{
+					UE_LOG(LogPoseSearch, Error, TEXT("FInteractionIsland::LogTickDependencies, error while retrieving the tick function for to %s"), *TickActorComponent->GetName());
+				}
 			}
 			else
 			{
diff --git a/Engine/Plugins/Animation/PoseSearch/Source/Runtime/Public/PoseSearch/PoseSearchInteractionIsland.h b/Engine/Plugins/Animation/PoseSearch/Source/Runtime/Public/PoseSearch/PoseSearchInteractionIsland.h
--- a/Engine/Plugins/Animation/PoseSearch/Source/Runtime/Public/PoseSearch/PoseSearchInteractionIsland.h
+++ b/Engine/Plugins/Animation/PoseSearch/Source/Runtime/Public/PoseSearch/PoseSearchInteractionIsland.h
@@ -172,6 +172,8 @@ struct IInteractionIslandDependency : public IModularFeature
 	virtual ~IInteractionIslandDependency() = default;
 
 	virtual bool CanMakeDependency(const UObject* InIslandObject, const UObject* InObject) const = 0;
+	virtual const FTickFunction* FindTickFunction(UObject* InObject) const = 0;
+
 	virtual void AddPrerequisite(UObject* InIslandObject, FTickFunction& InIslandTickFunction, UObject* InObject) const = 0;
 	virtual void AddSubsequent(UObject* InIslandObject, FTickFunction& InIslandTickFunction, UObject* InObject) const = 0;
 	virtual void RemovePrerequisite(UObject* InIslandObject, FTickFunction& InIslandTickFunction, UObject* InObject) const = 0;
diff --git a/Engine/Plugins/Experimental/AnimNext/Source/AnimNext/Internal/Component/AnimNextWorldSubsystem.h b/Engine/Plugins/Experimental/AnimNext/Source/AnimNext/Internal/Component/AnimNextWorldSubsystem.h
--- a/Engine/Plugins/Experimental/AnimNext/Source/AnimNext/Internal/Component/AnimNextWorldSubsystem.h
+++ b/Engine/Plugins/Experimental/AnimNext/Source/AnimNext/Internal/Component/AnimNextWorldSubsystem.h
@@ -148,12 +148,17 @@ protected:
 	// Disables the initialize event for a module once it has run once
 	void DisableInitializeEvent(UE::AnimNext::FModuleHandle InHandle);
 
+	// Find the const tick function for the specified event. Useful to analyze/log prerequisites. Use AddDependencyHandle, RemoveDependencyHandle to actually modify dependencies
+	// @param	InHandle			The handle of the module
+	// @param	InEventName			The event associated to the wanted tick function
+	const FTickFunction* FindTickFunctionHandle(UE::AnimNext::FModuleHandle InHandle, FName InEventName) const;
+
 	// Add a dependency on a tick function to the specified event
 	// @param	InHandle			The handle of the module whose dependencies are being modified
 	// @param	InObject			The object that owns the tick function
 	// @param	InTickFunction		The tick function to depend on
 	// @param	InEventName			The event to add the dependency to
-	// @param	InDependency			The kind of dependency to add
+	// @param	InDependency		The kind of dependency to add
 	void AddDependencyHandle(UE::AnimNext::FModuleHandle InHandle, UObject* InObject, FTickFunction& InTickFunction, FName InEventName, EDependency InDependency);
 	
 	// Remove a dependency on a tick function from the specified event
@@ -161,7 +166,7 @@ protected:
 	// @param	InObject			The object that owns the tick function
 	// @param	InTickFunction		The tick function that was depended on
 	// @param	InEventName			The event to remove the dependency from
-	// @param	InDependency			The kind of dependency to remove
+	// @param	InDependency		The kind of dependency to remove
 	void RemoveDependencyHandle(UE::AnimNext::FModuleHandle InHandle, UObject* InObject, FTickFunction& InTickFunction, FName InEventName, EDependency InDependency);
 
 protected:
diff --git a/Engine/Plugins/Experimental/AnimNext/Source/AnimNext/Internal/Module/AnimNextModuleInstance.h b/Engine/Plugins/Experimental/AnimNext/Source/AnimNext/Internal/Module/AnimNextModuleInstance.h
--- a/Engine/Plugins/Experimental/AnimNext/Source/AnimNext/Internal/Module/AnimNextModuleInstance.h
+++ b/Engine/Plugins/Experimental/AnimNext/Source/AnimNext/Internal/Module/AnimNextModuleInstance.h
@@ -137,6 +137,7 @@ struct FAnimNextModuleInstance : public FAnimNextDataInterfaceInstance
 
 	// Find the tick function for the specified event
 	UE::AnimNext::FModuleEventTickFunction* FindTickFunction(FName InEventName);
+	const UE::AnimNext::FModuleEventTickFunction* FindTickFunction(FName InEventName) const;
 
 	// Run the specified RigVM event
 	ANIMNEXT_API void RunRigVMEvent(FName InEventName, float InDeltaTime);
diff --git a/Engine/Plugins/Experimental/AnimNext/Source/AnimNext/Private/Component/AnimNextComponent.cpp b/Engine/Plugins/Experimental/AnimNext/Source/AnimNext/Private/Component/AnimNextComponent.cpp
--- a/Engine/Plugins/Experimental/AnimNext/Source/AnimNext/Private/Component/AnimNextComponent.cpp
+++ b/Engine/Plugins/Experimental/AnimNext/Source/AnimNext/Private/Component/AnimNextComponent.cpp
@@ -905,6 +905,15 @@ void UAnimNextComponent::QueueInputTraitEvent(FAnimNextTraitEventPtr Event)
 	}
 }
 
+const FTickFunction* UAnimNextComponent::FindTickFunction(FName InEventName) const
+{
+	if(Subsystem)
+	{
+		return Subsystem->FindTickFunction(this, InEventName);
+	}
+	return nullptr;
+}
+
 void UAnimNextComponent::AddPrerequisite(UObject* InObject, FTickFunction& InTickFunction, FName InEventName)
 {
 	if(Subsystem)
diff --git a/Engine/Plugins/Experimental/AnimNext/Source/AnimNext/Private/Component/AnimNextComponentWorldSubsystem.cpp b/Engine/Plugins/Experimental/AnimNext/Source/AnimNext/Private/Component/AnimNextComponentWorldSubsystem.cpp
--- a/Engine/Plugins/Experimental/AnimNext/Source/AnimNext/Private/Component/AnimNextComponentWorldSubsystem.cpp
+++ b/Engine/Plugins/Experimental/AnimNext/Source/AnimNext/Private/Component/AnimNextComponentWorldSubsystem.cpp
@@ -67,6 +67,14 @@ void UAnimNextComponentWorldSubsystem::QueueInputTraitEvent(UAnimNextComponent*
 	QueueInputTraitEventHandle(InComponent->ModuleHandle, Event);
 }
 
+const FTickFunction* UAnimNextComponentWorldSubsystem::FindTickFunction(const UAnimNextComponent* InComponent, FName InEventName) const
+{
+	using namespace UE::AnimNext;
+	check(IsInGameThread());
+	check(InComponent);
+	return FindTickFunctionHandle(InComponent->ModuleHandle, InEventName);
+}
+
 void UAnimNextComponentWorldSubsystem::AddDependency(UAnimNextComponent* InComponent, UObject* InObject, FTickFunction& InTickFunction, FName InEventName, EDependency InDependency)
 {
 	using namespace UE::AnimNext;
diff --git a/Engine/Plugins/Experimental/AnimNext/Source/AnimNext/Private/Component/AnimNextComponentWorldSubsystem.h b/Engine/Plugins/Experimental/AnimNext/Source/AnimNext/Private/Component/AnimNextComponentWorldSubsystem.h
--- a/Engine/Plugins/Experimental/AnimNext/Source/AnimNext/Private/Component/AnimNextComponentWorldSubsystem.h
+++ b/Engine/Plugins/Experimental/AnimNext/Source/AnimNext/Private/Component/AnimNextComponentWorldSubsystem.h
@@ -49,6 +49,11 @@ class UAnimNextComponentWorldSubsystem : public UAnimNextWorldSubsystem
 	// Queue an input trait event
 	void QueueInputTraitEvent(UAnimNextComponent* InComponent, FAnimNextTraitEventPtr Event);
 
+	// Find the component tick function for the specified event
+	// @param	InComponent			The component being searched for the tick function
+	// @param	InEventName			The event associated to the wanted tick function
+	const FTickFunction* FindTickFunction(const UAnimNextComponent* InComponent, FName InEventName) const;
+
 	// Add a dependency on a tick function to the specified event
 	// @param	InComponent			The component whose dependencies are being modified
 	// @param	InObject			The object that owns the tick function
diff --git a/Engine/Plugins/Experimental/AnimNext/Source/AnimNext/Private/Component/AnimNextWorldSubsystem.cpp b/Engine/Plugins/Experimental/AnimNext/Source/AnimNext/Private/Component/AnimNextWorldSubsystem.cpp
--- a/Engine/Plugins/Experimental/AnimNext/Source/AnimNext/Private/Component/AnimNextWorldSubsystem.cpp
+++ b/Engine/Plugins/Experimental/AnimNext/Source/AnimNext/Private/Component/AnimNextWorldSubsystem.cpp
@@ -355,6 +355,30 @@ void UAnimNextWorldSubsystem::RemoveModuleDependencyHandle(UE::AnimNext::FModule
 	SubsequentInstance->RemovePrerequisite(*PrerequisiteInstance);
 }
 
+const FTickFunction* UAnimNextWorldSubsystem::FindTickFunctionHandle(UE::AnimNext::FModuleHandle InHandle, FName InEventName) const
+{
+	using namespace UE::AnimNext;
+	check(IsInGameThread());
+	if (IsValidHandle(InHandle))
+	{
+		const FAnimNextModuleInstance& Instance = Instances.Get(InHandle);
+		const FModuleEventTickFunction* TickFunction = Instance.FindTickFunction(InEventName);
+		if(TickFunction == nullptr)
+		{
+			UE_LOGFMT(LogAnimation, Warning, "FindTickFunctionHandle: Could not find event '{EventName}' in module '{ModuleName}'", InEventName, Instance.GetDataInterfaceName());
+			return nullptr;
+		}
+		
+		if(TickFunction->bUserEvent)
+		{
+			return TickFunction;
+		}
+
+		UE_LOGFMT(LogAnimation, Warning, "FindTickFunctionHandle: Event '{EventName}' in module '{ModuleName}' is not a bUserEvent, therefore cannot be exposed", InEventName, Instance.GetDataInterfaceName());
+	}
+	return nullptr;
+}
+
 void UAnimNextWorldSubsystem::AddDependencyHandle(UE::AnimNext::FModuleHandle InHandle, UObject* InObject, FTickFunction& InTickFunction, FName InEventName, EDependency InDependency)
 {
 	using namespace UE::AnimNext;
diff --git a/Engine/Plugins/Experimental/AnimNext/Source/AnimNext/Private/Module/AnimNextModuleInstance.cpp b/Engine/Plugins/Experimental/AnimNext/Source/AnimNext/Private/Module/AnimNextModuleInstance.cpp
--- a/Engine/Plugins/Experimental/AnimNext/Source/AnimNext/Private/Module/AnimNextModuleInstance.cpp
+++ b/Engine/Plugins/Experimental/AnimNext/Source/AnimNext/Private/Module/AnimNextModuleInstance.cpp
@@ -563,6 +563,22 @@ UE::AnimNext::FModuleEventTickFunction* FAnimNextModuleInstance::FindTickFunctio
 	return nullptr;
 }
 
+const UE::AnimNext::FModuleEventTickFunction* FAnimNextModuleInstance::FindTickFunction(FName InEventName) const
+{
+	using namespace UE::AnimNext;
+
+	FModuleWriteGuard Guard(this);
+
+	for(const FModuleEventTickFunction& TickFunction : TickFunctions)
+	{
+		if(TickFunction.EventName == InEventName)
+		{
+			return &TickFunction;
+		}
+	}
+	return nullptr;
+}
+
 void FAnimNextModuleInstance::EndExecution(float InDeltaTime)
 {
 	UE::AnimNext::FModuleWriteGuard Guard(this);
diff --git a/Engine/Plugins/Experimental/AnimNext/Source/AnimNext/Public/Component/AnimNextComponent.h b/Engine/Plugins/Experimental/AnimNext/Source/AnimNext/Public/Component/AnimNextComponent.h
--- a/Engine/Plugins/Experimental/AnimNext/Source/AnimNext/Public/Component/AnimNextComponent.h
+++ b/Engine/Plugins/Experimental/AnimNext/Source/AnimNext/Public/Component/AnimNextComponent.h
@@ -205,6 +205,10 @@ public:
 	// Get the handle to the registered module
 	ANIMNEXT_API UE::AnimNext::FModuleHandle GetModuleHandle() const { return ModuleHandle; }
 
+	// Find the tick function for the specified event
+	// @param	InEventName			The event associated to the wanted tick function
+	ANIMNEXT_API const FTickFunction* FindTickFunction(FName InEventName) const;
+
 	// Add a prerequisite tick function dependency to the specified event
 	// @param	InObject			The object that owns the tick function
 	// @param	InTickFunction		The tick function to depend on
diff --git a/Engine/Plugins/Experimental/AnimNextPoseSearch/Source/AnimNextPoseSearch/Private/AnimNextInteractionIslandDependency.cpp b/Engine/Plugins/Experimental/AnimNextPoseSearch/Source/AnimNextPoseSearch/Private/AnimNextInteractionIslandDependency.cpp
--- a/Engine/Plugins/Experimental/AnimNextPoseSearch/Source/AnimNextPoseSearch/Private/AnimNextInteractionIslandDependency.cpp
+++ b/Engine/Plugins/Experimental/AnimNextPoseSearch/Source/AnimNextPoseSearch/Private/AnimNextInteractionIslandDependency.cpp
@@ -1,4 +1,4 @@
-﻿// Copyright Epic Games, Inc. All Rights Reserved.
+// Copyright Epic Games, Inc. All Rights Reserved.
 
 #include "AnimNextInteractionIslandDependency.h"
 
@@ -14,6 +14,12 @@ bool FAnimNextInteractionIslandDependency::CanMakeDependency(const UObject* InIs
 {
 	return InObject->IsA<UAnimNextComponent>();
 }
+	
+const FTickFunction* FAnimNextInteractionIslandDependency::FindTickFunction(UObject* InObject) const
+{
+	const UAnimNextComponent* AnimNextComponent = CastChecked<UAnimNextComponent>(InObject);
+	return AnimNextComponent->FindTickFunction(FRigUnit_AnimNextPrePhysicsEvent::EventName);
+}
 
 void FAnimNextInteractionIslandDependency::AddPrerequisite(UObject* InIslandObject, FTickFunction& InIslandTickFunction, UObject* InObject) const
 {
diff --git a/Engine/Plugins/Experimental/AnimNextPoseSearch/Source/AnimNextPoseSearch/Private/AnimNextInteractionIslandDependency.h b/Engine/Plugins/Experimental/AnimNextPoseSearch/Source/AnimNextPoseSearch/Private/AnimNextInteractionIslandDependency.h
--- a/Engine/Plugins/Experimental/AnimNextPoseSearch/Source/AnimNextPoseSearch/Private/AnimNextInteractionIslandDependency.h
+++ b/Engine/Plugins/Experimental/AnimNextPoseSearch/Source/AnimNextPoseSearch/Private/AnimNextInteractionIslandDependency.h
@@ -1,4 +1,4 @@
-﻿// Copyright Epic Games, Inc. All Rights Reserved.
+// Copyright Epic Games, Inc. All Rights Reserved.
 
 #pragma once
 
@@ -18,6 +18,8 @@ private:
 	friend UE::AnimNext::PoseSearch::FModule;
 	
 	virtual bool CanMakeDependency(const UObject* InIslandObject, const UObject* InObject) const override;
+	virtual const FTickFunction* FindTickFunction(UObject* InObject) const override;
+
 	virtual void AddPrerequisite(UObject* InIslandObject, FTickFunction& InIslandTickFunction, UObject* InObject) const override;
 	virtual void AddSubsequent(UObject* InIslandObject, FTickFunction& InIslandTickFunction, UObject* InObject) const override;
 	virtual void RemovePrerequisite(UObject* InIslandObject, FTickFunction& InIslandTickFunction, UObject* InObject) const override;
