Deprecated PoseSearch specific's FPoseSearchQueryTrajectory and FPoseSearchQueryTrajectorySample in favor of engine's FTransformTrajectory and FTransformTrajectorySample.

This deprecation is meant to not affect licensees assets (AnimBPs, BPs, etc) and is handling that via redirectors, and custom code in special cases such as manually fixing bindings and pins for the Pose Search History node for example. There are niche cases where deprecation is hard to do gracefully, as it is the case for IPoseHistory::GetTrajetory() due to it being a pure virtual function therefore licensees code will not compile until it has been updated.

[RN] minor

diff --git a/Engine/Plugins/Animation/PoseSearch/Config/BasePoseSearch.ini b/Engine/Plugins/Animation/PoseSearch/Config/BasePoseSearch.ini
new file mode 100644
--- /dev/null
+++ b/Engine/Plugins/Animation/PoseSearch/Config/BasePoseSearch.ini
@@ -0,0 +1,33 @@
+﻿[CoreRedirects]
+
+; Trajectory type redirects
++StructRedirects=(OldName="/Script/PoseSearch.PoseSearchQueryTrajectory", NewName="/Script/Engine.TransformTrajectory")
++StructRedirects=(OldName="/Script/PoseSearch.PoseSearchQueryTrajectorySample", NewName="/Script/Engine.TransformTrajectorySample")
+
+; Trajectory variables
++PropertyRedirects=(OldName="/Script/PoseSearch.PoseSearchQueryTrajectory.Samples", NewName="/Script/Engine.TransformTrajectory.Samples")
+
+; Trajectory Sample variables
++PropertyRedirects=(OldName="/Script/PoseSearch.PoseSearchQueryTrajectorySample.Position", NewName="/Script/Engine.TransformTrajectorySample.Position")
++PropertyRedirects=(OldName="/Script/PoseSearch.PoseSearchQueryTrajectorySample.Facing", NewName="/Script/Engine.TransformTrajectorySample.Facing")
++PropertyRedirects=(OldName="/Script/PoseSearch.PoseSearchQueryTrajectorySample.AccumulatedSeconds", NewName="/Script/Engine.TransformTrajectorySample.TimeInSeconds")
+
+; Fix up all pose search trajectory library functions
++FunctionRedirects=(OldName="PoseSearchTrajectoryLibrary.PoseSearchGenerateTrajectory", NewName="PoseSearchTrajectoryLibrary.PoseSearchGenerateTransformTrajectory")
++FunctionRedirects=(OldName="PoseSearchTrajectoryLibrary.PoseSearchGeneratePredictorTrajectory", NewName="PoseSearchTrajectoryLibrary.PoseSearchGeneratePredictorTransformTrajectory")
++FunctionRedirects=(OldName="PoseSearchTrajectoryLibrary.PoseSearchGenerateTrajectoryWithPredictor", NewName="PoseSearchTrajectoryLibrary.PoseSearchGenerateTransformTrajectoryWithPredictor")
++FunctionRedirects=(OldName="PoseSearchTrajectoryLibrary.HandleTrajectoryWorldCollisions", NewName="PoseSearchTrajectoryLibrary.HandleTransformTrajectoryWorldCollisions")
++FunctionRedirects=(OldName="PoseSearchTrajectoryLibrary.HandleTrajectoryWorldCollisionsWithGravity", NewName="PoseSearchTrajectoryLibrary.HandleTransformTrajectoryWorldCollisionsWithGravity")
++FunctionRedirects=(OldName="PoseSearchTrajectoryLibrary.GetTrajectorySampleAtTime", NewName="PoseSearchTrajectoryLibrary.GetTransformTrajectorySampleAtTime")
++FunctionRedirects=(OldName="PoseSearchTrajectoryLibrary.GetTrajectoryVelocity", NewName="PoseSearchTrajectoryLibrary.GetTransformTrajectoryVelocity")
++FunctionRedirects=(OldName="PoseSearchTrajectoryLibrary.GetTrajectoryAngularVelocity", NewName="PoseSearchTrajectoryLibrary.GetTransformTrajectoryAngularVelocity")
++FunctionRedirects=(OldName="PoseSearchTrajectoryLibrary.GetTransform", NewName="PoseSearchTrajectoryLibrary.GetTransformTrajectorySampleTransform")
++FunctionRedirects=(OldName="PoseSearchTrajectoryLibrary.DrawTrajectory", NewName="PoseSearchTrajectoryLibrary.DrawTransformTrajectory")
+
+; Fix up all pose search history collector node library functions
++FunctionRedirects=(OldName="PoseSearchHistoryCollectorAnimNodeLibrary.GetPoseHistoryNodeTrajectory", NewName="PoseSearchHistoryCollectorAnimNodeLibrary.GetPoseHistoryNodeTransformTrajectory")
++FunctionRedirects=(OldName="PoseSearchHistoryCollectorAnimNodeLibrary.SetPoseHistoryNodeTrajectory", NewName="PoseSearchHistoryCollectorAnimNodeLibrary.SetPoseHistoryNodeTransformTrajectory")
+
+; Fix up pose history node stuff
++PropertyRedirects=(OldName="AnimNode_PoseSearchHistoryCollector_Base.Trajectory", NewName="AnimNode_PoseSearchHistoryCollector_Base.TransformTrajectory")
++PropertyRedirects=(OldName="AnimNode_PoseSearchHistoryCollector.Trajectory", NewName="AnimNode_PoseSearchHistoryCollector.TransformTrajectory")
diff --git a/Engine/Plugins/Animation/PoseSearch/Source/Editor/Private/AnimGraphNode_PoseSearchHistoryCollector.cpp b/Engine/Plugins/Animation/PoseSearch/Source/Editor/Private/AnimGraphNode_PoseSearchHistoryCollector.cpp
--- a/Engine/Plugins/Animation/PoseSearch/Source/Editor/Private/AnimGraphNode_PoseSearchHistoryCollector.cpp
+++ b/Engine/Plugins/Animation/PoseSearch/Source/Editor/Private/AnimGraphNode_PoseSearchHistoryCollector.cpp
@@ -2,6 +2,8 @@
 
 #include "AnimGraphNode_PoseSearchHistoryCollector.h"
 #include "AnimationGraphSchema.h"
+#include "AnimGraphNode_Base.h"
+#include "Kismet2/BlueprintEditorUtils.h"
 
 #define LOCTEXT_NAMESPACE "AnimGraphNode_PoseSearchHistoryCollector"
 
@@ -23,6 +25,76 @@ FText UAnimGraphNode_PoseSearchHistoryCollector_Base::GetMenuCategory() const
 	return LOCTEXT("NodeCategory", "Pose Search");
 }
 
+void UAnimGraphNode_PoseSearchHistoryCollector_Base::ValidateNodeDuringCompilation(FCompilerResultsLog& MessageLog) const
+{
+	Super::ValidateNodeDuringCompilation(MessageLog);
+
+	PRAGMA_DISABLE_DEPRECATION_WARNINGS
+	const FName OldPropertyName = GET_MEMBER_NAME_CHECKED(FAnimNode_PoseSearchHistoryCollector_Base, Trajectory);
+	PRAGMA_ENABLE_DEPRECATION_WARNINGS
+	
+	if (HasBinding(OldPropertyName))
+	{
+		MessageLog.Error(*LOCTEXT("OldTrajectoryType", "In node (@@), please manually re-bind pin \"Trajectory\" to it's respective variable now that FPoseSearchQueryTrajectory has been deprecated").ToString(), this);
+	}
+}
+
+void UAnimGraphNode_PoseSearchHistoryCollector_Base::PostReconstructNode()
+{
+	Super::PostReconstructNode();
+
+#if WITH_EDITOR
+	if (!IsTemplate())
+	{
+		// Make sure we're not dealing with a menu node
+		UEdGraph* OuterGraph = GetGraph();
+		if (OuterGraph && OuterGraph->Schema)
+		{
+			PRAGMA_DISABLE_DEPRECATION_WARNINGS
+			const FName OldPropertyName = GET_MEMBER_NAME_CHECKED(FAnimNode_PoseSearchHistoryCollector_Base, Trajectory);
+			PRAGMA_ENABLE_DEPRECATION_WARNINGS
+
+			const FName NewPropertyName = GET_MEMBER_NAME_CHECKED(FAnimNode_PoseSearchHistoryCollector_Base, TransformTrajectory);
+
+			// Fix-up binding(s) to properly redirect from our old member variable, with type FPoseSearchQueryTrajectory, to the new one, with type FTransformTrajectory.
+			if (HasBinding(OldPropertyName))
+			{
+				if (RedirectBinding(OldPropertyName, NewPropertyName))
+				{
+					// Let the graph know to refresh
+					GetGraph()->NotifyNodeChanged(this);
+
+					UBlueprint* Blueprint = GetBlueprint();
+					if (!Blueprint->bBeingCompiled)
+					{
+						FBlueprintEditorUtils::MarkBlueprintAsModified(Blueprint);
+					}
+				}
+			}
+		}
+	}
+#endif
+}
+
+UK2Node::ERedirectType UAnimGraphNode_PoseSearchHistoryCollector_Base::DoPinsMatchForReconstruction(const UEdGraphPin* NewPin, int32 NewPinIndex, const UEdGraphPin* OldPin, int32 OldPinIndex) const
+{
+	// @note: Once the "Trajectory" member variable is removed, we'd have to manually create an FCoreRedirectObjectName that **also** contains
+	// the class type in the string and not just the pin name to prevent our redirector(s) from failing. We are avoiding it rn, since that is not as simple
+	// forcing the redirect type.
+	PRAGMA_DISABLE_DEPRECATION_WARNINGS
+	const FName OldPropertyName = GET_MEMBER_NAME_CHECKED(FAnimNode_PoseSearchHistoryCollector_Base, Trajectory);
+	PRAGMA_ENABLE_DEPRECATION_WARNINGS
+	const FName NewPropertyName = GET_MEMBER_NAME_CHECKED(FAnimNode_PoseSearchHistoryCollector_Base, TransformTrajectory);
+	
+	// Old pin, FPoseSearchQueryTrajectory, input can't be linked to new pin, FTransformTrajectory. CoreRedirects doesn't seem to help here so we have to manually redirect it.
+	if (OldPin->PinName == OldPropertyName && NewPin->PinName == NewPropertyName)
+	{
+		return ERedirectType_Value;
+	}
+
+	return Super::DoPinsMatchForReconstruction(NewPin, NewPinIndex, OldPin, OldPinIndex);
+}
+
 /////////////////////////////////////////////////////
 // UAnimGraphNode_PoseSearchHistoryCollector
 
diff --git a/Engine/Plugins/Animation/PoseSearch/Source/Editor/Private/PoseSearchDatabaseViewModel.cpp b/Engine/Plugins/Animation/PoseSearch/Source/Editor/Private/PoseSearchDatabaseViewModel.cpp
--- a/Engine/Plugins/Animation/PoseSearch/Source/Editor/Private/PoseSearchDatabaseViewModel.cpp
+++ b/Engine/Plugins/Animation/PoseSearch/Source/Editor/Private/PoseSearchDatabaseViewModel.cpp
@@ -7,6 +7,7 @@
 #include "Animation/AnimSequence.h"
 #include "Animation/BlendSpace.h"
 #include "Animation/DebugSkelMeshComponent.h"
+#include "Animation/TrajectoryTypes.h"
 #include "Animation/MirrorDataTable.h"
 #include "Camera/CameraComponent.h"
 #include "Components/CapsuleComponent.h"
@@ -22,7 +23,6 @@
 #include "PoseSearch/PoseSearchDerivedData.h"
 #include "PoseSearch/PoseSearchHistory.h"
 #include "PoseSearch/PoseSearchSchema.h"
-#include "PoseSearch/PoseSearchTrajectoryTypes.h"
 #include "PoseSearchDatabaseAssetTreeNode.h"
 #include "PoseSearchDatabaseDataDetails.h"
 #include "PoseSearchDatabasePreviewScene.h"
@@ -146,12 +146,12 @@ bool FDatabasePreviewActor::SpawnPreviewActor(UWorld* World, const UPoseSearchDa
 			const FTransform IndexAssetPoseTransform = MirrorDataCache.MirrorTransform(Sampler.ExtractRootTransform(IndexAssetPoseTime));
 
 			Trajectory.Samples[Index].SetTransform(IndexAssetPoseTransform);
-			Trajectory.Samples[Index].AccumulatedSeconds = IndexAssetPoseTime;
+			Trajectory.Samples[Index].TimeInSeconds = IndexAssetPoseTime;
 		}
 
 		for (int32 Index = 1; Index < NumPoses; ++Index)
 		{
-			const float DeltaAccumulatedSeconds = FMath::Max(UE_KINDA_SMALL_NUMBER, Trajectory.Samples[Index].AccumulatedSeconds - Trajectory.Samples[Index - 1].AccumulatedSeconds);
+			const float DeltaAccumulatedSeconds = FMath::Max(UE_KINDA_SMALL_NUMBER, Trajectory.Samples[Index].TimeInSeconds - Trajectory.Samples[Index - 1].TimeInSeconds);
 			const FVector& Start = Trajectory.Samples[Index - 1].Position;
 			const FVector& End = Trajectory.Samples[Index].Position;
 			TrajectorySpeed[Index] = (Start - End).Length() / DeltaAccumulatedSeconds;
@@ -335,9 +335,9 @@ bool FDatabasePreviewActor::DrawPreviewActors(TConstArrayView<FDatabasePreviewAc
 			PoseHistoryEntry.SetComponentSpaceTransform(RootBoneIndexType, RefRootBone);
 		}
 
-		for (FPoseSearchQueryTrajectorySample& TrajectorySample : ArchivedPoseHistory.Trajectory.Samples)
+		for (FTransformTrajectorySample& TrajectorySample : ArchivedPoseHistory.Trajectory.Samples)
 		{
-			TrajectorySample.AccumulatedSeconds -= PreviewActor.QuantizedTime;
+			TrajectorySample.TimeInSeconds -= PreviewActor.QuantizedTime;
 		}
 
 		PoseHistories.Add(&ArchivedPoseHistory);
diff --git a/Engine/Plugins/Animation/PoseSearch/Source/Editor/Private/PoseSearchDebuggerView.cpp b/Engine/Plugins/Animation/PoseSearch/Source/Editor/Private/PoseSearchDebuggerView.cpp
--- a/Engine/Plugins/Animation/PoseSearch/Source/Editor/Private/PoseSearchDebuggerView.cpp
+++ b/Engine/Plugins/Animation/PoseSearch/Source/Editor/Private/PoseSearchDebuggerView.cpp
@@ -14,6 +14,7 @@
 #include "PoseSearchDebuggerViewModel.h"
 #include "PoseSearchEditor.h"
 #include "PropertyEditorModule.h"
+#include "Animation/TrajectoryTypes.h"
 #include "Subsystems/AssetEditorSubsystem.h"
 #include "Trace/PoseSearchTraceProvider.h"
 #include "Widgets/Input/SButton.h"
@@ -274,11 +275,11 @@ void SDebuggerView::Tick(const FGeometry& AllottedGeometry, const double InCurre
 		{
 			for (int32 RoleIndex = 0; RoleIndex < NumRoles; ++RoleIndex)
 			{
-				const FPoseSearchQueryTrajectory& Trajectory = State->PoseHistories[RoleIndex].Trajectory;
+				const FTransformTrajectory& Trajectory = State->PoseHistories[RoleIndex].Trajectory;
 				
 				if (bDrawTrajectory)
 				{
-					Trajectory.DebugDrawTrajectory(World);
+					UTransformTrajectoryBlueprintLibrary::DebugDrawTrajectory(Trajectory, World);
 				}
 
 				if (bDrawHistory)
diff --git a/Engine/Plugins/Animation/PoseSearch/Source/Editor/Private/Trace/PoseSearchTraceAnalyzer.cpp b/Engine/Plugins/Animation/PoseSearch/Source/Editor/Private/Trace/PoseSearchTraceAnalyzer.cpp
--- a/Engine/Plugins/Animation/PoseSearch/Source/Editor/Private/Trace/PoseSearchTraceAnalyzer.cpp
+++ b/Engine/Plugins/Animation/PoseSearch/Source/Editor/Private/Trace/PoseSearchTraceAnalyzer.cpp
@@ -7,6 +7,7 @@
 #include "Serialization/MemoryReader.h"
 #include "PoseSearch/Trace/PoseSearchTraceLogger.h"
 #include "TraceServices/Model/AnalysisSession.h"
+#include "PoseSearch/PoseSearchCustomVersion.h"
 
 namespace UE::PoseSearch
 {
@@ -19,14 +20,9 @@ FTraceAnalyzer::FTraceAnalyzer(TraceServices::IAnalysisSession& InSession, FTrac
 void FTraceAnalyzer::OnAnalysisBegin(const FOnAnalysisContext& Context)
 {
 	FInterfaceBuilder& Builder = Context.InterfaceBuilder;
-
-	ANSICHAR LoggerName[NAME_SIZE];
-	ANSICHAR MotionMatchingStateName[NAME_SIZE];
-
-	FTraceLogger::Name.GetPlainANSIString(LoggerName);
-	FTraceMotionMatchingStateMessage::Name.GetPlainANSIString(MotionMatchingStateName);
-
-	Builder.RouteEvent(RouteId_MotionMatchingState, LoggerName, MotionMatchingStateName);
+	
+	Builder.RouteEvent(RouteId_MotionMatchingState, "PoseSearch", "MotionMatchingState");
+	Builder.RouteEvent(RouteId_MotionMatchingState2, "PoseSearch", "MotionMatchingState2");
 }
 
 bool FTraceAnalyzer::OnEvent(uint16 RouteId, EStyle Style, const FOnEventContext& Context)
@@ -39,6 +35,24 @@ bool FTraceAnalyzer::OnEvent(uint16 RouteId, EStyle Style, const FOnEventContext
 	{
 		FTraceMotionMatchingStateMessage Message;
 		FMemoryReaderView Archive(Context.EventData.GetArrayView<uint8>("Data"));
+
+		// Ensure we have not versions.
+		FCustomVersionContainer CustomVersionContainer;
+		Archive.SetCustomVersions(CustomVersionContainer);
+		
+		Archive << Message;
+		TraceProvider.AppendMotionMatchingState(Message, Context.EventTime.AsSeconds(Message.Cycle));
+	}
+	else if (RouteId == RouteId_MotionMatchingState2)
+	{
+		FTraceMotionMatchingStateMessage Message;
+		FMemoryReaderView Archive(Context.EventData.GetArrayView<uint8>("Data"));
+
+		// Ensure we have the version that matches this RouteId event.
+		FCustomVersionContainer CustomVersionContainer;
+		CustomVersionContainer.SetVersion(FPoseSearchCustomVersion::GUID, FPoseSearchCustomVersion::DeprecatedTrajectoryTypes, TEXT("Dev-PoseSearch-Version"));
+		Archive.SetCustomVersions(CustomVersionContainer);
+
 		Archive << Message;
 		TraceProvider.AppendMotionMatchingState(Message, Context.EventTime.AsSeconds(Message.Cycle));
 	}
diff --git a/Engine/Plugins/Animation/PoseSearch/Source/Editor/Private/Trace/PoseSearchTraceAnalyzer.h b/Engine/Plugins/Animation/PoseSearch/Source/Editor/Private/Trace/PoseSearchTraceAnalyzer.h
--- a/Engine/Plugins/Animation/PoseSearch/Source/Editor/Private/Trace/PoseSearchTraceAnalyzer.h
+++ b/Engine/Plugins/Animation/PoseSearch/Source/Editor/Private/Trace/PoseSearchTraceAnalyzer.h
@@ -32,6 +32,7 @@ private:
 	enum : uint16
 	{
 		RouteId_MotionMatchingState,
+		RouteId_MotionMatchingState2,
 	};
 
 	TraceServices::IAnalysisSession& Session;
diff --git a/Engine/Plugins/Animation/PoseSearch/Source/Editor/Public/AnimGraphNode_PoseSearchHistoryCollector.h b/Engine/Plugins/Animation/PoseSearch/Source/Editor/Public/AnimGraphNode_PoseSearchHistoryCollector.h
--- a/Engine/Plugins/Animation/PoseSearch/Source/Editor/Public/AnimGraphNode_PoseSearchHistoryCollector.h
+++ b/Engine/Plugins/Animation/PoseSearch/Source/Editor/Public/AnimGraphNode_PoseSearchHistoryCollector.h
@@ -15,6 +15,14 @@ class UAnimGraphNode_PoseSearchHistoryCollector_Base : public UAnimGraphNode_Bas
 	virtual FLinearColor GetNodeTitleColor() const override;
 	virtual FText GetTooltipText() const override;
 	virtual FText GetMenuCategory() const override;
+
+public:
+
+	//~ Begin UK2Node Interface.
+	virtual void ValidateNodeDuringCompilation(FCompilerResultsLog& MessageLog) const override;
+	virtual void PostReconstructNode() override;
+	virtual ERedirectType DoPinsMatchForReconstruction(const UEdGraphPin* NewPin, int32 NewPinIndex, const UEdGraphPin* OldPin, int32 OldPinIndex) const override;
+	//~ End UK2Node Interface.
 };
 
 UCLASS(MinimalAPI)
diff --git a/Engine/Plugins/Animation/PoseSearch/Source/Editor/Public/PoseSearchDatabaseViewModel.h b/Engine/Plugins/Animation/PoseSearch/Source/Editor/Public/PoseSearchDatabaseViewModel.h
--- a/Engine/Plugins/Animation/PoseSearch/Source/Editor/Public/PoseSearchDatabaseViewModel.h
+++ b/Engine/Plugins/Animation/PoseSearch/Source/Editor/Public/PoseSearchDatabaseViewModel.h
@@ -8,7 +8,7 @@
 #include "PoseSearch/PoseSearchAssetSampler.h"
 #include "PoseSearch/PoseSearchMirrorDataCache.h"
 #include "PoseSearch/PoseSearchRole.h"
-#include "PoseSearch/PoseSearchTrajectoryTypes.h"
+#include "Animation/TrajectoryTypes.h"
 #include "PoseSearchDatabasePreviewScene.h"
 #include "PoseSearch/PoseSearchDatabase.h"
 #include "UObject/GCObject.h"
@@ -59,7 +59,7 @@ namespace UE::PoseSearch
 		float QuantizedTime = 0.f;
 
 		FAnimationAssetSampler Sampler;
-		FPoseSearchQueryTrajectory Trajectory;
+		FTransformTrajectory Trajectory;
 		TArray<float> TrajectorySpeed;
 
 		FRole ActorRole = DefaultRole;
diff --git a/Engine/Plugins/Animation/PoseSearch/Source/Runtime/Private/AnimNode_PoseSearchHistoryCollector.cpp b/Engine/Plugins/Animation/PoseSearch/Source/Runtime/Private/AnimNode_PoseSearchHistoryCollector.cpp
--- a/Engine/Plugins/Animation/PoseSearch/Source/Runtime/Private/AnimNode_PoseSearchHistoryCollector.cpp
+++ b/Engine/Plugins/Animation/PoseSearch/Source/Runtime/Private/AnimNode_PoseSearchHistoryCollector.cpp
@@ -25,12 +25,12 @@ bool FAnimNode_PoseSearchHistoryCollector_Base::GetCurveValueAtTime(float Time,
 	return PoseHistory.GetCurveValueAtTime(Time, CurveName, OutCurveValue, bExtrapolate);
 }
 
-const FPoseSearchQueryTrajectory& FAnimNode_PoseSearchHistoryCollector_Base::GetTrajectory() const
+const FTransformTrajectory& FAnimNode_PoseSearchHistoryCollector_Base::GetTrajectory() const
 {
 	return PoseHistory.GetTrajectory();
 }
 
-void FAnimNode_PoseSearchHistoryCollector_Base::SetTrajectory(const FPoseSearchQueryTrajectory& InTrajectory, float InTrajectorySpeedMultiplier)
+void FAnimNode_PoseSearchHistoryCollector_Base::SetTrajectory(const FTransformTrajectory& InTrajectory, float InTrajectorySpeedMultiplier)
 {
 	PoseHistory.SetTrajectory(InTrajectory, InTrajectorySpeedMultiplier);
 }
@@ -154,7 +154,7 @@ void FAnimNode_PoseSearchHistoryCollector_Base::Update_AnyThread(const FAnimatio
 	}
 	else
 	{
-		PoseHistory.SetTrajectory(Trajectory, TrajectorySpeedMultiplier);
+		PoseHistory.SetTrajectory(TransformTrajectory, TrajectorySpeedMultiplier);
 	}
 
 	bIsTrajectoryGeneratedBeforePreUpdate = false;
diff --git a/Engine/Plugins/Animation/PoseSearch/Source/Runtime/Private/PoseSearchContext.cpp b/Engine/Plugins/Animation/PoseSearch/Source/Runtime/Private/PoseSearchContext.cpp
--- a/Engine/Plugins/Animation/PoseSearch/Source/Runtime/Private/PoseSearchContext.cpp
+++ b/Engine/Plugins/Animation/PoseSearch/Source/Runtime/Private/PoseSearchContext.cpp
@@ -7,7 +7,7 @@
 #include "PoseSearch/PoseSearchDatabase.h"
 #include "PoseSearch/PoseSearchHistory.h"
 #include "PoseSearch/PoseSearchSchema.h"
-#include "PoseSearch/PoseSearchTrajectoryTypes.h"
+#include "Animation/TrajectoryTypes.h"
 #include "PoseSearch/PoseSearchFeatureChannel_Heading.h"
 #include "PoseSearch/PoseSearchFeatureChannel_PermutationTime.h"
 #include "PoseSearch/PoseSearchFeatureChannel_Position.h"
diff --git a/Engine/Plugins/Animation/PoseSearch/Source/Runtime/Private/PoseSearchCustomVersion.cpp b/Engine/Plugins/Animation/PoseSearch/Source/Runtime/Private/PoseSearchCustomVersion.cpp
new file mode 100644
--- /dev/null
+++ b/Engine/Plugins/Animation/PoseSearch/Source/Runtime/Private/PoseSearchCustomVersion.cpp
@@ -0,0 +1,9 @@
+﻿// Copyright Epic Games, Inc. All Rights Reserved.
+
+#include "PoseSearch/PoseSearchCustomVersion.h"
+#include "Serialization/CustomVersion.h"
+
+const FGuid FPoseSearchCustomVersion::GUID("99c277d4-1347-497b-ab23-47514df0f30b");
+
+// Register the custom version with core
+FCustomVersionRegistration GRegisterPoseSearchCustomVersion(FPoseSearchCustomVersion::GUID, FPoseSearchCustomVersion::LatestVersion, TEXT("Dev-PoseSearch-Version"));
diff --git a/Engine/Plugins/Animation/PoseSearch/Source/Runtime/Private/PoseSearchHistory.cpp b/Engine/Plugins/Animation/PoseSearch/Source/Runtime/Private/PoseSearchHistory.cpp
--- a/Engine/Plugins/Animation/PoseSearch/Source/Runtime/Private/PoseSearchHistory.cpp
+++ b/Engine/Plugins/Animation/PoseSearch/Source/Runtime/Private/PoseSearchHistory.cpp
@@ -9,7 +9,9 @@
 #include "Animation/SkeletonRemappingRegistry.h"
 #include "BonePose.h"
 #include "DrawDebugHelpers.h"
+#include "Animation/TrajectoryTypes.h"
 #include "Engine/SkeletalMesh.h"
+#include "PoseSearch/PoseSearchCustomVersion.h"
 #include "PoseSearch/PoseSearchResult.h"
 #include "PoseSearch/PoseSearchDatabase.h"
 #include "PoseSearch/PoseSearchDefines.h"
@@ -638,7 +640,21 @@ FArchive& operator<<(FArchive& Ar, FArchivedPoseHistory& Entry)
 	Ar << Entry.BoneToTransformMap;
 	Ar << Entry.CollectedCurves;
 	Ar << Entry.Entries;
-	Ar << Entry.Trajectory;
+	
+	// Convert old FPoseSearchTrajectory to new FTransformTrajectory type at load time.
+	if (Ar.CustomVer(FPoseSearchCustomVersion::GUID) < FPoseSearchCustomVersion::DeprecatedTrajectoryTypes)
+	{
+		PRAGMA_DISABLE_DEPRECATION_WARNINGS
+		FPoseSearchQueryTrajectory OldTrajectoryType;
+		Ar << OldTrajectoryType;
+		Entry.Trajectory = OldTrajectoryType;
+		PRAGMA_ENABLE_DEPRECATION_WARNINGS
+	}	
+	else
+	{
+		Ar << Entry.Trajectory;
+	}
+	
 	return Ar;
 }
 
@@ -705,8 +721,8 @@ void FPoseHistory::Initialize_AnyThread(int32 InNumPoses, float InSamplingInterv
 
 	MaxNumPoses = InNumPoses;
 	SamplingInterval = InSamplingInterval;
-
-	Trajectory = FPoseSearchQueryTrajectory();
+	
+	Trajectory = FTransformTrajectory();
 	TrajectoryDataState = FPoseSearchTrajectoryData::FState();
 	TrajectorySpeedMultiplier = 1.f;
 
@@ -798,7 +814,7 @@ bool FPoseHistory::GetCurveValueAtTime(float Time, const FName& CurveName, float
 	return false;
 }
 
-const FPoseSearchQueryTrajectory& FPoseHistory::GetTrajectory() const
+const FTransformTrajectory& FPoseHistory::GetTrajectory() const
 {
 	UE_MT_SCOPED_READ_ACCESS(PoseDataThreadSafeCounter);
 	return Trajectory;
@@ -828,6 +844,13 @@ const TConstArrayView<FName> FPoseHistory::GetCollectedCurves() const
 	return MakeConstArrayView(PoseData.CollectedCurves);
 }
 
+PRAGMA_DISABLE_DEPRECATION_WARNINGS
+void FPoseHistory::SetTrajectory(const FPoseSearchQueryTrajectory& InTrajectory, float InTrajectorySpeedMultiplier)
+{
+	SetTrajectory(FTransformTrajectory(InTrajectory), InTrajectorySpeedMultiplier);
+}
+PRAGMA_ENABLE_DEPRECATION_WARNINGS
+	
 int32 FPoseHistory::GetNumEntries() const
 {
 	UE_MT_SCOPED_READ_ACCESS(PoseDataThreadSafeCounter);
@@ -860,7 +883,7 @@ void FPoseHistory::GenerateTrajectory(const UObject* AnimContext, float DeltaTim
 	GenerateTrajectory(AnimContext, DeltaTime, FPoseSearchTrajectoryData(), FPoseSearchTrajectoryData::FSampling());
 }
 
-void FPoseHistory::SetTrajectory(const FPoseSearchQueryTrajectory& InTrajectory, float InTrajectorySpeedMultiplier)
+void FPoseHistory::SetTrajectory(const FTransformTrajectory& InTrajectory, float InTrajectorySpeedMultiplier)
 {
 	if (!InTrajectory.Samples.IsEmpty())
 	{
@@ -874,9 +897,9 @@ void FPoseHistory::SetTrajectory(const FPoseSearchQueryTrajectory& InTrajectory,
 		if (!FMath::IsNearlyEqual(TrajectorySpeedMultiplier, 1.f))
 		{
 			const float TrajectorySpeedMultiplierInv = FMath::IsNearlyZero(TrajectorySpeedMultiplier) ? 1.f : 1.f / TrajectorySpeedMultiplier;
-			for (FPoseSearchQueryTrajectorySample& Sample : Trajectory.Samples)
+			for (FTransformTrajectorySample& Sample : Trajectory.Samples)
 			{
-				Sample.AccumulatedSeconds *= TrajectorySpeedMultiplierInv;
+				Sample.TimeInSeconds *= TrajectorySpeedMultiplierInv;
 			}
 		}
 	}
@@ -1006,7 +1029,7 @@ void FPoseHistory::DebugDraw(FAnimInstanceProxy& AnimInstanceProxy, FColor Color
 		const float DebugThickness = GVarAnimPoseHistoryDebugDrawTrajectoryThickness;
 		const int MaxHistorySamples = GVarAnimPoseHistoryDebugDrawTrajectoryMaxNumOfHistorySamples;
 		const int MaxPredictionSamples = GVarAnimPoseHistoryDebugDrawTrajectoryMaxNumOfPredictionSamples;
-		Trajectory.DebugDrawTrajectory(AnimInstanceProxy, DebugThickness, 0, MaxHistorySamples, MaxPredictionSamples);
+		UTransformTrajectoryBlueprintLibrary::DebugDrawTrajectory(Trajectory, AnimInstanceProxy, DebugThickness, 0, MaxHistorySamples, MaxPredictionSamples);
 	}
 
 	if (Color.A > 0 && GVarAnimPoseHistoryDebugDrawPose)
@@ -1048,7 +1071,7 @@ void FMemStackPoseHistory::Init(const IPoseHistory* InPoseHistory)
 	PoseHistory = InPoseHistory;
 }
 
-void FMemStackPoseHistory::SetTrajectory(const FPoseSearchQueryTrajectory& InTrajectory, float InTrajectorySpeedMultiplier)
+void FMemStackPoseHistory::SetTrajectory(const FTransformTrajectory& InTrajectory, float InTrajectorySpeedMultiplier)
 {
 	check(PoseHistory);
 	// FMemStackPoseHistory should never change the trajectory!
@@ -1179,7 +1202,7 @@ void FMemStackPoseHistory::DebugDraw(FAnimInstanceProxy& AnimInstanceProxy, FCol
 
 	if (Color.A > 0 && !FutureEntries.IsEmpty() && GVarAnimPoseHistoryDebugDrawPose)
 	{
-		const FPoseSearchQueryTrajectory& Trajectory = GetTrajectory();
+		const FTransformTrajectory& Trajectory = GetTrajectory();
 		const bool bValidTrajectory = !Trajectory.Samples.IsEmpty();
 		TArray<FTransform, TInlineAllocator<128>> PrevGlobalTransforms;
 
diff --git a/Engine/Plugins/Animation/PoseSearch/Source/Runtime/Private/PoseSearchHistoryCollectorAnimNodeLibrary.cpp b/Engine/Plugins/Animation/PoseSearch/Source/Runtime/Private/PoseSearchHistoryCollectorAnimNodeLibrary.cpp
--- a/Engine/Plugins/Animation/PoseSearch/Source/Runtime/Private/PoseSearchHistoryCollectorAnimNodeLibrary.cpp
+++ b/Engine/Plugins/Animation/PoseSearch/Source/Runtime/Private/PoseSearchHistoryCollectorAnimNodeLibrary.cpp
@@ -8,6 +8,30 @@ FPoseSearchHistoryCollectorAnimNodeReference UPoseSearchHistoryCollectorAnimNode
 	return FAnimNodeReference::ConvertToType<FPoseSearchHistoryCollectorAnimNodeReference>(Node, Result);
 }
 
+void UPoseSearchHistoryCollectorAnimNodeLibrary::ConvertToPoseHistoryNodePure(const FAnimNodeReference& Node, FPoseSearchHistoryCollectorAnimNodeReference& PoseSearchHistoryCollectorNode, bool& Result)
+{
+	EAnimNodeReferenceConversionResult ConversionResult;
+	PoseSearchHistoryCollectorNode = ConvertToPoseHistoryNode(Node, ConversionResult);
+	Result = (ConversionResult == EAnimNodeReferenceConversionResult::Succeeded);
+}
+
+void UPoseSearchHistoryCollectorAnimNodeLibrary::GetPoseHistoryNodeTransformTrajectory(const FPoseSearchHistoryCollectorAnimNodeReference& PoseSearchHistoryCollectorNode, FTransformTrajectory& Trajectory)
+{
+	if (FAnimNode_PoseSearchHistoryCollector* PoseSearchHistoryCollectorNodePtr = PoseSearchHistoryCollectorNode.GetAnimNodePtr<FAnimNode_PoseSearchHistoryCollector>())
+	{
+		Trajectory = PoseSearchHistoryCollectorNodePtr->GetTrajectory();
+	}
+}
+
+void UPoseSearchHistoryCollectorAnimNodeLibrary::SetPoseHistoryNodeTransformTrajectory(const FPoseSearchHistoryCollectorAnimNodeReference& PoseSearchHistoryCollectorNode, const FTransformTrajectory& Trajectory)
+{
+	if (FAnimNode_PoseSearchHistoryCollector* PoseSearchHistoryCollectorNodePtr = PoseSearchHistoryCollectorNode.GetAnimNodePtr<FAnimNode_PoseSearchHistoryCollector>())
+	{
+		PoseSearchHistoryCollectorNodePtr->SetTrajectory(Trajectory);
+	}
+}
+
+PRAGMA_DISABLE_DEPRECATION_WARNINGS
 void UPoseSearchHistoryCollectorAnimNodeLibrary::GetPoseHistoryNodeTrajectory(const FPoseSearchHistoryCollectorAnimNodeReference& PoseSearchHistoryCollectorNode, FPoseSearchQueryTrajectory& Trajectory)
 {
 	if (FAnimNode_PoseSearchHistoryCollector* PoseSearchHistoryCollectorNodePtr = PoseSearchHistoryCollectorNode.GetAnimNodePtr<FAnimNode_PoseSearchHistoryCollector>())
@@ -23,3 +47,4 @@ void UPoseSearchHistoryCollectorAnimNodeLibrary::SetPoseHistoryNodeTrajectory(co
 		PoseSearchHistoryCollectorNodePtr->SetTrajectory(Trajectory);
 	}
 }
+PRAGMA_ENABLE_DEPRECATION_WARNINGS
\ No newline at end of file
diff --git a/Engine/Plugins/Animation/PoseSearch/Source/Runtime/Private/PoseSearchTrajectoryLibrary.cpp b/Engine/Plugins/Animation/PoseSearch/Source/Runtime/Private/PoseSearchTrajectoryLibrary.cpp
--- a/Engine/Plugins/Animation/PoseSearch/Source/Runtime/Private/PoseSearchTrajectoryLibrary.cpp
+++ b/Engine/Plugins/Animation/PoseSearch/Source/Runtime/Private/PoseSearchTrajectoryLibrary.cpp
@@ -9,6 +9,7 @@
 #include "GameFramework/CharacterMovementComponent.h"
 #include "PoseSearch/PoseSearchDefines.h"
 #include "Kismet/KismetMathLibrary.h"
+#include "Animation/TrajectoryTypes.h"
 
 bool FPoseSearchTrajectoryData::UpdateData(
 	float DeltaTime,
@@ -171,6 +172,8 @@ FVector FPoseSearchTrajectoryData::StepCharacterMovementGroundPrediction(
 	return OutVelocity;
 }
 
+PRAGMA_DISABLE_DEPRECATION_WARNINGS
+
 // Deprecated
 void UPoseSearchTrajectoryLibrary::InitTrajectorySamples(
 	FPoseSearchQueryTrajectory& Trajectory,
@@ -187,6 +190,7 @@ void UPoseSearchTrajectoryLibrary::InitTrajectorySamples(
 }
 
 // Deprecated
+PRAGMA_DISABLE_DEPRECATION_WARNINGS
 void UPoseSearchTrajectoryLibrary::InitTrajectorySamples(
 	FPoseSearchQueryTrajectory& Trajectory,
 	const FPoseSearchTrajectoryData& TrajectoryData,	// Unused parameter
@@ -202,11 +206,22 @@ void UPoseSearchTrajectoryLibrary::InitTrajectorySamples(
 		DeltaTime);
 }
 
+// Deprecated
+
 void UPoseSearchTrajectoryLibrary::InitTrajectorySamples(
 	FPoseSearchQueryTrajectory& Trajectory,
 	FVector DefaultPosition, FQuat DefaultFacing,
 	const FPoseSearchTrajectoryData::FSampling& TrajectoryDataSampling,
 	float DeltaTime)
+{
+	FTransformTrajectory TransformTrajectory = Trajectory;
+	InitTrajectorySamples(TransformTrajectory, DefaultPosition, DefaultFacing, TrajectoryDataSampling, DeltaTime);
+	Trajectory = TransformTrajectory;
+}
+
+PRAGMA_ENABLE_DEPRECATION_WARNINGS
+
+void UPoseSearchTrajectoryLibrary::InitTrajectorySamples(FTransformTrajectory& Trajectory, FVector DefaultPosition, FQuat DefaultFacing, const FPoseSearchTrajectoryData::FSampling& TrajectoryDataSampling, float DeltaTime)
 {
 	const int32 NumHistorySamples = TrajectoryDataSampling.NumHistorySamples;
 	const int32 NumPredictionSamples = TrajectoryDataSampling.NumPredictionSamples;
@@ -224,7 +239,7 @@ void UPoseSearchTrajectoryLibrary::InitTrajectorySamples(
 		{
 			Trajectory.Samples[i].Position = DefaultPosition;
 			Trajectory.Samples[i].Facing = DefaultFacing;
-			Trajectory.Samples[i].AccumulatedSeconds = SecondsPerHistorySample * (i - NumHistorySamples - 1);
+			Trajectory.Samples[i].TimeInSeconds = SecondsPerHistorySample * (i - NumHistorySamples - 1);
 		}
 
 		// Initialize current sample and prediction
@@ -233,11 +248,13 @@ void UPoseSearchTrajectoryLibrary::InitTrajectorySamples(
 		{
 			Trajectory.Samples[i].Position = DefaultPosition;
 			Trajectory.Samples[i].Facing = DefaultFacing;
-			Trajectory.Samples[i].AccumulatedSeconds = SecondsPerPredictionSample * (i - NumHistorySamples) + DeltaTime;
+			Trajectory.Samples[i].TimeInSeconds = SecondsPerPredictionSample * (i - NumHistorySamples) + DeltaTime;
 		}
 	}
 }
 
+PRAGMA_DISABLE_DEPRECATION_WARNINGS
+
 // Deprecated
 void UPoseSearchTrajectoryLibrary::UpdateHistory_TransformHistory(
 	FPoseSearchQueryTrajectory& Trajectory,
@@ -255,12 +272,23 @@ void UPoseSearchTrajectoryLibrary::UpdateHistory_TransformHistory(
 		DeltaTime);
 }
 
+// Deprecated
 void UPoseSearchTrajectoryLibrary::UpdateHistory_TransformHistory(
 	FPoseSearchQueryTrajectory& Trajectory,
 	FVector CurrentPosition,
 	FVector CurrentVelocity, 
 	const FPoseSearchTrajectoryData::FSampling& TrajectoryDataSampling,
 	float DeltaTime)
+{
+	FTransformTrajectory TransformTrajectory = Trajectory;
+	UpdateHistory_TransformHistory(TransformTrajectory, CurrentPosition, CurrentVelocity, TrajectoryDataSampling, DeltaTime);
+	Trajectory = TransformTrajectory;
+}
+
+PRAGMA_ENABLE_DEPRECATION_WARNINGS
+
+void UPoseSearchTrajectoryLibrary::UpdateHistory_TransformHistory(FTransformTrajectory& Trajectory, FVector CurrentPosition, FVector CurrentVelocity,
+	const FPoseSearchTrajectoryData::FSampling& TrajectoryDataSampling, float DeltaTime)
 {
 	const int32 NumHistorySamples = TrajectoryDataSampling.NumHistorySamples;
 	if (NumHistorySamples > 0)
@@ -275,18 +303,18 @@ void UPoseSearchTrajectoryLibrary::UpdateHistory_TransformHistory(
 		const FVector InferredGroundTranslation = TranslationSinceLastFrame - CurrentTranslationFromMover;
 
 		// Shift history Samples when it's time to record a new one.
-		if (SecondsPerHistorySample <= 0.f || FMath::Abs(Trajectory.Samples[NumHistorySamples - 1].AccumulatedSeconds) >= SecondsPerHistorySample)
+		if (SecondsPerHistorySample <= 0.f || FMath::Abs(Trajectory.Samples[NumHistorySamples - 1].TimeInSeconds) >= SecondsPerHistorySample)
 		{
 			for (int32 Index = 0; Index < NumHistorySamples - 1; ++Index)
 			{
-				Trajectory.Samples[Index].AccumulatedSeconds = Trajectory.Samples[Index + 1].AccumulatedSeconds - DeltaTime;
+				Trajectory.Samples[Index].TimeInSeconds = Trajectory.Samples[Index + 1].TimeInSeconds - DeltaTime;
 				Trajectory.Samples[Index].Position = Trajectory.Samples[Index + 1].Position + InferredGroundTranslation;
 				Trajectory.Samples[Index].Facing = Trajectory.Samples[Index + 1].Facing;
 			}
 
 			// Adding a new history record
 			// Note we add current frame at delta time in the future since t = 0 should be the previous frame
-			Trajectory.Samples[NumHistorySamples - 1].AccumulatedSeconds = DeltaTime;
+			Trajectory.Samples[NumHistorySamples - 1].TimeInSeconds = DeltaTime;
 			Trajectory.Samples[NumHistorySamples - 1].Position = CurrentPosition;
 			Trajectory.Samples[NumHistorySamples - 1].Facing = Trajectory.Samples[NumHistorySamples].Facing;
 		}
@@ -296,20 +324,32 @@ void UPoseSearchTrajectoryLibrary::UpdateHistory_TransformHistory(
 
 			for (int32 Index = 0; Index < NumHistorySamples; ++Index)
 			{
-				Trajectory.Samples[Index].AccumulatedSeconds -= DeltaTime;
+				Trajectory.Samples[Index].TimeInSeconds -= DeltaTime;
 				Trajectory.Samples[Index].Position += InferredGroundTranslation;
 			}
 		}
 	}
 }
 
+PRAGMA_DISABLE_DEPRECATION_WARNINGS
 
+// Deprecated
 void UPoseSearchTrajectoryLibrary::UpdateHistory_WorldSpace(
 	FPoseSearchQueryTrajectory& Trajectory,
 	FVector CurrentPosition,
 	FQuat CurrentRotation,
 	const FPoseSearchTrajectoryData::FSampling& TrajectoryDataSampling,
 	float DeltaTime)
+{
+	FTransformTrajectory TransformTrajectory = Trajectory;
+	UpdateHistory_WorldSpace(TransformTrajectory, CurrentPosition, CurrentRotation, TrajectoryDataSampling, DeltaTime);
+	Trajectory = TransformTrajectory;
+}
+
+PRAGMA_ENABLE_DEPRECATION_WARNINGS
+
+void UPoseSearchTrajectoryLibrary::UpdateHistory_WorldSpace(FTransformTrajectory& Trajectory, FVector CurrentPosition, FQuat CurrentRotation, 
+	const FPoseSearchTrajectoryData::FSampling& TrajectoryDataSampling, float DeltaTime)
 {
 	const int32 NumHistorySamples = TrajectoryDataSampling.NumHistorySamples;
 	if (NumHistorySamples > 0)
@@ -319,18 +359,18 @@ void UPoseSearchTrajectoryLibrary::UpdateHistory_WorldSpace(
 		check(NumHistorySamples <= Trajectory.Samples.Num());
 
 		// Shift history Samples when it's time to record a new one.
-		if (SecondsPerHistorySample <= 0.f || FMath::Abs(Trajectory.Samples[NumHistorySamples - 1].AccumulatedSeconds) >= SecondsPerHistorySample)
+		if (SecondsPerHistorySample <= 0.f || FMath::Abs(Trajectory.Samples[NumHistorySamples - 1].TimeInSeconds) >= SecondsPerHistorySample)
 		{
 			for (int32 Index = 0; Index < NumHistorySamples - 1; ++Index)
 			{
-				Trajectory.Samples[Index].AccumulatedSeconds = Trajectory.Samples[Index + 1].AccumulatedSeconds - DeltaTime;
+				Trajectory.Samples[Index].TimeInSeconds = Trajectory.Samples[Index + 1].TimeInSeconds - DeltaTime;
 				Trajectory.Samples[Index].Position = Trajectory.Samples[Index + 1].Position;
 				Trajectory.Samples[Index].Facing = Trajectory.Samples[Index + 1].Facing;
 			}
 
 			// Adding a new history record
 			// Note we add current frame at delta time in the future since t = 0 should be the previous frame
-			Trajectory.Samples[NumHistorySamples - 1].AccumulatedSeconds = DeltaTime;
+			Trajectory.Samples[NumHistorySamples - 1].TimeInSeconds = DeltaTime;
 			Trajectory.Samples[NumHistorySamples - 1].Position = CurrentPosition;
 			Trajectory.Samples[NumHistorySamples - 1].Facing = CurrentRotation;
 		}
@@ -339,10 +379,10 @@ void UPoseSearchTrajectoryLibrary::UpdateHistory_WorldSpace(
 			// Didn't record a new history position, update timers and shift by ground translation
 			for (int32 Index = 0; Index < NumHistorySamples; ++Index)
 			{
-				Trajectory.Samples[Index].AccumulatedSeconds -= DeltaTime;
+				Trajectory.Samples[Index].TimeInSeconds -= DeltaTime;
 			}
 		}
-	}
+	}	
 }
 
 FVector UPoseSearchTrajectoryLibrary::RemapVectorMagnitudeWithCurve(
@@ -363,6 +403,9 @@ FVector UPoseSearchTrajectoryLibrary::RemapVectorMagnitudeWithCurve(
 	return Vector;
 }
 
+PRAGMA_DISABLE_DEPRECATION_WARNINGS
+
+// Deprecated
 void UPoseSearchTrajectoryLibrary::UpdatePrediction_SimulateCharacterMovement(
 	FPoseSearchQueryTrajectory& Trajectory,
 	const FPoseSearchTrajectoryData& TrajectoryData,
@@ -370,7 +413,18 @@ void UPoseSearchTrajectoryLibrary::UpdatePrediction_SimulateCharacterMovement(
 	const FPoseSearchTrajectoryData::FSampling& TrajectoryDataSampling,
 	float DeltaTime)
 {
-	FVector CurrentPositionWS = TrajectoryDataDerived.Position;
+	FTransformTrajectory TransformTrajectory = Trajectory;
+	UpdatePrediction_SimulateCharacterMovement(TransformTrajectory, TrajectoryData, TrajectoryDataDerived, TrajectoryDataSampling, DeltaTime);
+	Trajectory = TransformTrajectory;
+}
+
+PRAGMA_ENABLE_DEPRECATION_WARNINGS
+
+void UPoseSearchTrajectoryLibrary::UpdatePrediction_SimulateCharacterMovement(FTransformTrajectory& Trajectory,
+	const FPoseSearchTrajectoryData& TrajectoryData, const FPoseSearchTrajectoryData::FDerived& TrajectoryDataDerived,
+	const FPoseSearchTrajectoryData::FSampling& TrajectoryDataSampling, float DeltaTime)
+{
+		FVector CurrentPositionWS = TrajectoryDataDerived.Position;
 	FVector CurrentVelocityWS = RemapVectorMagnitudeWithCurve(TrajectoryDataDerived.Velocity, TrajectoryData.bUseSpeedRemappingCurve, TrajectoryData.SpeedRemappingCurve);
 	FVector CurrentAccelerationWS = RemapVectorMagnitudeWithCurve(TrajectoryDataDerived.Acceleration, TrajectoryData.bUseAccelerationRemappingCurve, TrajectoryData.AccelerationRemappingCurve);
 
@@ -414,7 +468,7 @@ void UPoseSearchTrajectoryLibrary::UpdatePrediction_SimulateCharacterMovement(
 		{
 			Trajectory.Samples[Index].Position = CurrentPositionWS;
 			Trajectory.Samples[Index].Facing = CurrentFacingWS;
-			Trajectory.Samples[Index].AccumulatedSeconds = AccumulatedSeconds;
+			Trajectory.Samples[Index].TimeInSeconds = AccumulatedSeconds;
 
 			if (Index == LastIndex)
 			{
@@ -444,6 +498,9 @@ void UPoseSearchTrajectoryLibrary::UpdatePrediction_SimulateCharacterMovement(
 	}
 }
 
+PRAGMA_DISABLE_DEPRECATION_WARNINGS
+
+// Deprecated
 void UPoseSearchTrajectoryLibrary::PoseSearchGenerateTrajectory(
 	const UObject* Context, 
 	UPARAM(ref)	const FPoseSearchTrajectoryData& InTrajectoryData,
@@ -455,6 +512,21 @@ void UPoseSearchTrajectoryLibrary::PoseSearchGenerateTrajectory(
 	int32 InTrajectoryHistoryCount,
 	float InPredictionSamplingInterval,
 	int32 InTrajectoryPredictionCount)
+{
+	FTransformTrajectory InOutTransformTrajectory = InOutTrajectory;
+	FTransformTrajectory OutTransformTrajectory;
+	
+	PoseSearchGenerateTransformTrajectory(Context, InTrajectoryData, InDeltaTime, InOutTransformTrajectory, InOutDesiredControllerYawLastUpdate,
+		OutTransformTrajectory, InHistorySamplingInterval, InTrajectoryHistoryCount, InPredictionSamplingInterval, InTrajectoryPredictionCount);
+
+	OutTrajectory = OutTransformTrajectory;
+}
+
+PRAGMA_ENABLE_DEPRECATION_WARNINGS
+
+void UPoseSearchTrajectoryLibrary::PoseSearchGenerateTransformTrajectory(const UObject* InContext, const FPoseSearchTrajectoryData& InTrajectoryData,
+	float InDeltaTime, FTransformTrajectory& InOutTrajectory, float& InOutDesiredControllerYawLastUpdate, FTransformTrajectory& OutTrajectory,
+	float InHistorySamplingInterval, int32 InTrajectoryHistoryCount, float InPredictionSamplingInterval, int32 InTrajectoryPredictionCount)
 {
 	FPoseSearchTrajectoryData::FSampling TrajectoryDataSampling;
 	TrajectoryDataSampling.NumHistorySamples = InTrajectoryHistoryCount;
@@ -466,7 +538,7 @@ void UPoseSearchTrajectoryLibrary::PoseSearchGenerateTrajectory(
 	TrajectoryDataState.DesiredControllerYawLastUpdate = InOutDesiredControllerYawLastUpdate;
 
 	FPoseSearchTrajectoryData::FDerived TrajectoryDataDerived;
-	InTrajectoryData.UpdateData(InDeltaTime, Context, TrajectoryDataDerived, TrajectoryDataState);
+	InTrajectoryData.UpdateData(InDeltaTime, InContext, TrajectoryDataDerived, TrajectoryDataState);
 	InitTrajectorySamples(InOutTrajectory, TrajectoryDataDerived.Position, TrajectoryDataDerived.Facing, TrajectoryDataSampling, InDeltaTime);
 	UpdateHistory_TransformHistory(InOutTrajectory, TrajectoryDataDerived.Position, TrajectoryDataDerived.Velocity, TrajectoryDataSampling, InDeltaTime);
 	UpdatePrediction_SimulateCharacterMovement(InOutTrajectory, InTrajectoryData, TrajectoryDataDerived, TrajectoryDataSampling, InDeltaTime);
@@ -476,6 +548,27 @@ void UPoseSearchTrajectoryLibrary::PoseSearchGenerateTrajectory(
 	OutTrajectory = InOutTrajectory;
 }
 
+void UPoseSearchTrajectoryLibrary::PoseSearchGeneratePredictorTransformTrajectory(UObject* InPredictor,
+	const FPoseSearchTrajectoryData& InTrajectoryData, float InDeltaTime, FTransformTrajectory& InOutTrajectory,
+	float& InOutDesiredControllerYawLastUpdate, FTransformTrajectory& OutTrajectory, float InHistorySamplingInterval,
+	int32 InTrajectoryHistoryCount, float InPredictionSamplingInterval, int32 InTrajectoryPredictionCount)
+{
+	IPoseSearchTrajectoryPredictorInterface* Predictor = Cast<IPoseSearchTrajectoryPredictorInterface>(InPredictor);
+	if (Predictor != nullptr)
+	{
+		PoseSearchGenerateTransformTrajectoryWithPredictor(InPredictor,
+			InDeltaTime,
+			InOutTrajectory,
+			InOutDesiredControllerYawLastUpdate,
+			OutTrajectory,
+			InHistorySamplingInterval,
+			InTrajectoryHistoryCount,
+			InPredictionSamplingInterval,
+			InTrajectoryPredictionCount);
+	}
+}
+
+PRAGMA_DISABLE_DEPRECATION_WARNINGS
 
 // Deprecated
 void UPoseSearchTrajectoryLibrary::PoseSearchGeneratePredictorTrajectory(
@@ -505,6 +598,7 @@ void UPoseSearchTrajectoryLibrary::PoseSearchGeneratePredictorTrajectory(
 	}
 }
 
+// Deprecated
 void UPoseSearchTrajectoryLibrary::PoseSearchGenerateTrajectoryWithPredictor(
 	TScriptInterface<IPoseSearchTrajectoryPredictorInterface> InPredictor,	// must implement IPoseSearchTrajectoryPredictorInterface
 	float InDeltaTime,
@@ -515,6 +609,21 @@ void UPoseSearchTrajectoryLibrary::PoseSearchGenerateTrajectoryWithPredictor(
 	int32 InTrajectoryHistoryCount,
 	float InPredictionSamplingInterval,
 	int32 InTrajectoryPredictionCount)
+{
+	FTransformTrajectory InOutTransformTrajectory = InOutTrajectory;
+	FTransformTrajectory OutTransformTrajectory;
+	
+	PoseSearchGenerateTransformTrajectoryWithPredictor(InPredictor, InDeltaTime, InOutTransformTrajectory, InOutDesiredControllerYawLastUpdate,
+		OutTransformTrajectory, InHistorySamplingInterval, InTrajectoryHistoryCount, InPredictionSamplingInterval, InTrajectoryPredictionCount);
+
+	OutTrajectory = OutTransformTrajectory;
+}
+
+PRAGMA_ENABLE_DEPRECATION_WARNINGS
+
+void UPoseSearchTrajectoryLibrary::PoseSearchGenerateTransformTrajectoryWithPredictor(TScriptInterface<IPoseSearchTrajectoryPredictorInterface> InPredictor,
+	float InDeltaTime, FTransformTrajectory& InOutTrajectory, float& InOutDesiredControllerYawLastUpdate, FTransformTrajectory& OutTrajectory,
+	float InHistorySamplingInterval, int32 InTrajectoryHistoryCount, float InPredictionSamplingInterval, int32 InTrajectoryPredictionCount)
 {
 	FPoseSearchTrajectoryData::FSampling TrajectoryDataSampling;
 	TrajectoryDataSampling.NumHistorySamples = InTrajectoryHistoryCount;
@@ -547,9 +656,29 @@ void UPoseSearchTrajectoryLibrary::PoseSearchGenerateTrajectoryWithPredictor(
 	OutTrajectory = InOutTrajectory;
 }
 
+PRAGMA_DISABLE_DEPRECATION_WARNINGS
 
+// Deprecated
 void UPoseSearchTrajectoryLibrary::HandleTrajectoryWorldCollisions(const UObject* WorldContextObject, const UAnimInstance* AnimInstance, UPARAM(ref) const FPoseSearchQueryTrajectory& InTrajectory, bool bApplyGravity, float FloorCollisionsOffset, FPoseSearchQueryTrajectory& OutTrajectory, FPoseSearchTrajectory_WorldCollisionResults& CollisionResult,
-	ETraceTypeQuery TraceChannel, bool bTraceComplex, const TArray<AActor*>& ActorsToIgnore, EDrawDebugTrace::Type DrawDebugType, bool bIgnoreSelf, float MaxObstacleHeight, FLinearColor TraceColor, FLinearColor TraceHitColor, float DrawTime)
+                                                                   ETraceTypeQuery TraceChannel, bool bTraceComplex, const TArray<AActor*>& ActorsToIgnore, EDrawDebugTrace::Type DrawDebugType, bool bIgnoreSelf, float MaxObstacleHeight, FLinearColor TraceColor, FLinearColor TraceHitColor, float DrawTime)
+{
+	FTransformTrajectory InTransformTrajectory = InTrajectory;
+	FTransformTrajectory OutTransformTrajectory;
+	
+	HandleTransformTrajectoryWorldCollisions(WorldContextObject, AnimInstance, InTransformTrajectory, bApplyGravity, 
+        FloorCollisionsOffset, OutTransformTrajectory, CollisionResult, TraceChannel, bTraceComplex, ActorsToIgnore, 
+        DrawDebugType, bIgnoreSelf, MaxObstacleHeight, TraceColor, TraceHitColor, DrawTime);
+
+	OutTrajectory = OutTransformTrajectory;
+}
+
+PRAGMA_ENABLE_DEPRECATION_WARNINGS
+
+void UPoseSearchTrajectoryLibrary::HandleTransformTrajectoryWorldCollisions(const UObject* WorldContextObject, const UAnimInstance* AnimInstance,
+	const FTransformTrajectory& InTrajectory, bool bApplyGravity, float FloorCollisionsOffset, FTransformTrajectory& OutTrajectory,
+	FPoseSearchTrajectory_WorldCollisionResults& CollisionResult, ETraceTypeQuery TraceChannel, bool bTraceComplex,
+	const TArray<AActor*>& ActorsToIgnore, EDrawDebugTrace::Type DrawDebugType, bool bIgnoreSelf, float MaxObstacleHeight, FLinearColor TraceColor,
+	FLinearColor TraceHitColor, float DrawTime)
 {
 	FVector StartingVelocity = FVector::ZeroVector;
 	FVector GravityAccel = FVector::ZeroVector;
@@ -565,19 +694,39 @@ void UPoseSearchTrajectoryLibrary::HandleTrajectoryWorldCollisions(const UObject
 		}
 	}
 
-	HandleTrajectoryWorldCollisionsWithGravity(WorldContextObject, InTrajectory, StartingVelocity, bApplyGravity, GravityAccel, 
+	HandleTransformTrajectoryWorldCollisionsWithGravity(WorldContextObject, InTrajectory, StartingVelocity, bApplyGravity, GravityAccel, 
 		FloorCollisionsOffset, OutTrajectory, CollisionResult, TraceChannel, bTraceComplex, ActorsToIgnore, 
 		DrawDebugType, bIgnoreSelf, MaxObstacleHeight, TraceColor, TraceHitColor, DrawTime);
 }
 
+PRAGMA_DISABLE_DEPRECATION_WARNINGS
 
+// Deprecated
 void UPoseSearchTrajectoryLibrary::HandleTrajectoryWorldCollisionsWithGravity(const UObject* WorldContextObject,
-	UPARAM(ref) const FPoseSearchQueryTrajectory& InTrajectory, FVector StartingVelocity, bool bApplyGravity, FVector GravityAccel, float FloorCollisionsOffset, FPoseSearchQueryTrajectory& OutTrajectory, FPoseSearchTrajectory_WorldCollisionResults& CollisionResult,
-	ETraceTypeQuery TraceChannel, bool bTraceComplex, const TArray<AActor*>& ActorsToIgnore, EDrawDebugTrace::Type DrawDebugType, bool bIgnoreSelf, float MaxObstacleHeight, FLinearColor TraceColor, FLinearColor TraceHitColor, float DrawTime)
+                                                                              UPARAM(ref) const FPoseSearchQueryTrajectory& InTrajectory, FVector StartingVelocity, bool bApplyGravity, FVector GravityAccel, float FloorCollisionsOffset, FPoseSearchQueryTrajectory& OutTrajectory, FPoseSearchTrajectory_WorldCollisionResults& CollisionResult,
+                                                                              ETraceTypeQuery TraceChannel, bool bTraceComplex, const TArray<AActor*>& ActorsToIgnore, EDrawDebugTrace::Type DrawDebugType, bool bIgnoreSelf, float MaxObstacleHeight, FLinearColor TraceColor, FLinearColor TraceHitColor, float DrawTime)
+{
+	FTransformTrajectory InTransformTrajectory = InTrajectory;
+	FTransformTrajectory OutTransformTrajectory;
+	
+	HandleTransformTrajectoryWorldCollisionsWithGravity(WorldContextObject, InTransformTrajectory, StartingVelocity, bApplyGravity, GravityAccel, 
+		FloorCollisionsOffset, OutTransformTrajectory, CollisionResult, TraceChannel, bTraceComplex, ActorsToIgnore, 
+		DrawDebugType, bIgnoreSelf, MaxObstacleHeight, TraceColor, TraceHitColor, DrawTime);
+
+	OutTrajectory = OutTransformTrajectory;
+}
+
+PRAGMA_ENABLE_DEPRECATION_WARNINGS
+
+void UPoseSearchTrajectoryLibrary::HandleTransformTrajectoryWorldCollisionsWithGravity(const UObject* WorldContextObject,
+	const FTransformTrajectory& InTrajectory, FVector StartingVelocity, bool bApplyGravity, FVector GravityAccel, float FloorCollisionsOffset,
+	FTransformTrajectory& OutTrajectory, FPoseSearchTrajectory_WorldCollisionResults& CollisionResult, ETraceTypeQuery TraceChannel,
+	bool bTraceComplex, const TArray<AActor*>& ActorsToIgnore, EDrawDebugTrace::Type DrawDebugType, bool bIgnoreSelf, float MaxObstacleHeight,
+	FLinearColor TraceColor, FLinearColor TraceHitColor, float DrawTime)
 {
 	OutTrajectory = InTrajectory;
 
-	TArray<FPoseSearchQueryTrajectorySample>& Samples = OutTrajectory.Samples;
+	TArray<FTransformTrajectorySample>& Samples = OutTrajectory.Samples;
 	const int32 NumSamples = Samples.Num();
 
 	FVector GravityDirection = FVector::ZeroVector;
@@ -593,7 +742,7 @@ void UPoseSearchTrajectoryLibrary::HandleTrajectoryWorldCollisionsWithGravity(co
 		InitialVelocityZ = VelocityOnGravityAxis.Length() * -FMath::Sign(GravityDirection.Dot(VelocityOnGravityAxis));
 	}
 
-	CollisionResult.TimeToLand = OutTrajectory.Samples.Last().AccumulatedSeconds;
+	CollisionResult.TimeToLand = OutTrajectory.Samples.Last().TimeInSeconds;
 
 	if (!FMath::IsNearlyZero(GravityZ))
 	{
@@ -606,13 +755,13 @@ void UPoseSearchTrajectoryLibrary::HandleTrajectoryWorldCollisionsWithGravity(co
 		float FreeFallAccumulatedSeconds = 0.f;
 		for (int32 SampleIndex = 1; SampleIndex < NumSamples; ++SampleIndex)
 		{
-			FPoseSearchQueryTrajectorySample& Sample = Samples[SampleIndex];
-			if (Sample.AccumulatedSeconds > 0.f)
+			FTransformTrajectorySample& Sample = Samples[SampleIndex];
+			if (Sample.TimeInSeconds > 0.f)
 			{
 				const int32 PrevSampleIndex = SampleIndex - 1;
-				const FPoseSearchQueryTrajectorySample& PrevSample = Samples[PrevSampleIndex];
+				const FTransformTrajectorySample& PrevSample = Samples[PrevSampleIndex];
 
-				FreeFallAccumulatedSeconds += Sample.AccumulatedSeconds - PrevSample.AccumulatedSeconds;
+				FreeFallAccumulatedSeconds += Sample.TimeInSeconds - PrevSample.TimeInSeconds;
 
 				if (bIsLastImpactValid)
 				{
@@ -654,8 +803,8 @@ void UPoseSearchTrajectoryLibrary::HandleTrajectoryWorldCollisionsWithGravity(co
 	{
 		for (int32 SampleIndex = 0; SampleIndex < NumSamples; ++SampleIndex)
 		{
-			FPoseSearchQueryTrajectorySample& Sample = OutTrajectory.Samples[SampleIndex];
-			if (Sample.AccumulatedSeconds > 0.f)
+			FTransformTrajectorySample& Sample = OutTrajectory.Samples[SampleIndex];
+			if (Sample.TimeInSeconds > 0.f)
 			{
 				FHitResult HitResult;
 				if (UKismetSystemLibrary::LineTraceSingle(WorldContextObject, Sample.Position + FVector::UpVector * 3000.f, Sample.Position, TraceChannel, bTraceComplex, ActorsToIgnore, DrawDebugType, HitResult, bIgnoreSelf, TraceColor, TraceHitColor, DrawTime))
@@ -669,12 +818,48 @@ void UPoseSearchTrajectoryLibrary::HandleTrajectoryWorldCollisionsWithGravity(co
 	CollisionResult.LandSpeed = InitialVelocityZ + GravityZ * CollisionResult.TimeToLand;
 }
 
+PRAGMA_DISABLE_DEPRECATION_WARNINGS
+
+// Deprecated
 void UPoseSearchTrajectoryLibrary::GetTrajectorySampleAtTime(UPARAM(ref) const FPoseSearchQueryTrajectory& InTrajectory, float Time, FPoseSearchQueryTrajectorySample& OutTrajectorySample, bool bExtrapolate)
 {
 	OutTrajectorySample = InTrajectory.GetSampleAtTime(Time, bExtrapolate);
 }
 
+// Deprecated
 void UPoseSearchTrajectoryLibrary::GetTrajectoryVelocity(UPARAM(ref) const FPoseSearchQueryTrajectory& InTrajectory, float Time1, float Time2, FVector& OutVelocity, bool bExtrapolate)
+{
+	GetTransformTrajectoryVelocity(FTransformTrajectory(InTrajectory), Time1, Time2, OutVelocity, bExtrapolate);
+}
+
+// Deprecated
+void UPoseSearchTrajectoryLibrary::GetTrajectoryAngularVelocity(const FPoseSearchQueryTrajectory& InTrajectory, float Time1, float Time2, FVector& OutAngularVelocity, bool bExtrapolate /*= false*/)
+{
+	GetTransformTrajectoryAngularVelocity(FTransformTrajectory(InTrajectory), Time1, Time2, OutAngularVelocity, bExtrapolate);
+}
+
+// Deprecated
+FTransform UPoseSearchTrajectoryLibrary::GetTransform(const FPoseSearchQueryTrajectorySample& InTrajectorySample)
+{
+	return InTrajectorySample.GetTransform();
+}
+
+// Deprecated
+void UPoseSearchTrajectoryLibrary::DrawTrajectory(const UObject* WorldContextObject, const FPoseSearchQueryTrajectory& InTrajectory, const float DebugThickness, float HeightOffset)
+{
+	DrawTransformTrajectory(WorldContextObject, FTransformTrajectory(InTrajectory), DebugThickness, HeightOffset);
+}
+
+PRAGMA_ENABLE_DEPRECATION_WARNINGS
+
+void UPoseSearchTrajectoryLibrary::GetTransformTrajectorySampleAtTime(const FTransformTrajectory& InTrajectory, float Time,
+	FTransformTrajectorySample& OutTrajectorySample, bool bExtrapolate)
+{
+	OutTrajectorySample = InTrajectory.GetSampleAtTime(Time, bExtrapolate);
+}
+
+void UPoseSearchTrajectoryLibrary::GetTransformTrajectoryVelocity(const FTransformTrajectory& InTrajectory, float Time1, float Time2, FVector& OutVelocity,
+	bool bExtrapolate)
 {
 	if (FMath::IsNearlyEqual(Time1, Time2))
 	{
@@ -683,13 +868,14 @@ void UPoseSearchTrajectoryLibrary::GetTrajectoryVelocity(UPARAM(ref) const FPose
 		return;
 	}
 
-	FPoseSearchQueryTrajectorySample Sample1 = InTrajectory.GetSampleAtTime(Time1, bExtrapolate);
-	FPoseSearchQueryTrajectorySample Sample2 = InTrajectory.GetSampleAtTime(Time2, bExtrapolate);
+	FTransformTrajectorySample Sample1 = InTrajectory.GetSampleAtTime(Time1, bExtrapolate);
+	FTransformTrajectorySample Sample2 = InTrajectory.GetSampleAtTime(Time2, bExtrapolate);
 
 	OutVelocity = (Sample2.Position - Sample1.Position) / (Time2 - Time1);
 }
 
-void UPoseSearchTrajectoryLibrary::GetTrajectoryAngularVelocity(const FPoseSearchQueryTrajectory& InTrajectory, float Time1, float Time2, FVector& OutAngularVelocity, bool bExtrapolate /*= false*/)
+void UPoseSearchTrajectoryLibrary::GetTransformTrajectoryAngularVelocity(const FTransformTrajectory& InTrajectory, float Time1, float Time2,
+	FVector& OutAngularVelocity, bool bExtrapolate)
 {
 	if (FMath::IsNearlyEqual(Time1, Time2))
 	{
@@ -698,8 +884,8 @@ void UPoseSearchTrajectoryLibrary::GetTrajectoryAngularVelocity(const FPoseSearc
 		return;
 	}
 
-	FPoseSearchQueryTrajectorySample Sample1 = InTrajectory.GetSampleAtTime(Time1, bExtrapolate);
-	FPoseSearchQueryTrajectorySample Sample2 = InTrajectory.GetSampleAtTime(Time2, bExtrapolate);
+	FTransformTrajectorySample Sample1 = InTrajectory.GetSampleAtTime(Time1, bExtrapolate);
+	FTransformTrajectorySample Sample2 = InTrajectory.GetSampleAtTime(Time2, bExtrapolate);
 
 	const FQuat DeltaRotation = (Sample2.Facing * Sample1.Facing.Inverse()).GetShortestArcWith(FQuat::Identity);
 	const FVector AngularVelocityInRadians = DeltaRotation.ToRotationVector() / (Time2 - Time1);
@@ -710,12 +896,19 @@ void UPoseSearchTrajectoryLibrary::GetTrajectoryAngularVelocity(const FPoseSearc
 		FMath::RadiansToDegrees(AngularVelocityInRadians.Z));
 }
 
-void UPoseSearchTrajectoryLibrary::DrawTrajectory(const UObject* WorldContextObject, const FPoseSearchQueryTrajectory& InTrajectory, const float DebugThickness, float HeightOffset)
+FTransform UPoseSearchTrajectoryLibrary::GetTransformTrajectorySampleTransform(const FTransformTrajectorySample& InTrajectorySample)
+{
+	return InTrajectorySample.GetTransform();
+}
+
+void UPoseSearchTrajectoryLibrary::DrawTransformTrajectory(const UObject* WorldContextObject, const FTransformTrajectory& InTrajectory,
+	const float DebugThickness, float HeightOffset)
 {
 #if ENABLE_ANIM_DEBUG
 	if (UWorld* World = GEngine->GetWorldFromContextObject(WorldContextObject, EGetWorldErrorMode::LogAndReturnNull))
 	{
-		InTrajectory.DebugDrawTrajectory(World, DebugThickness, HeightOffset);
+		UTransformTrajectoryBlueprintLibrary::DebugDrawTrajectory(InTrajectory, World, DebugThickness, HeightOffset);
 	}
 #endif // ENABLE_ANIM_DEBUG
-}
\ No newline at end of file
+}
+
diff --git a/Engine/Plugins/Animation/PoseSearch/Source/Runtime/Private/PoseSearchTrajectoryTypes.cpp b/Engine/Plugins/Animation/PoseSearch/Source/Runtime/Private/PoseSearchTrajectoryTypes.cpp
--- a/Engine/Plugins/Animation/PoseSearch/Source/Runtime/Private/PoseSearchTrajectoryTypes.cpp
+++ b/Engine/Plugins/Animation/PoseSearch/Source/Runtime/Private/PoseSearchTrajectoryTypes.cpp
@@ -8,6 +8,14 @@
 #include "PoseSearch/PoseSearchDefines.h"
 #include "Math/Axis.h"
 #include "VisualLogger/VisualLogger.h"
+#include "Animation/TrajectoryTypes.h"
+
+PRAGMA_DISABLE_DEPRECATION_WARNINGS
+
+FPoseSearchQueryTrajectorySample::FPoseSearchQueryTrajectorySample(const FQuat& InFacing, const FVector& InPosition, float InAccumulatedSeconds)
+	: Facing(InFacing), Position(InPosition), AccumulatedSeconds(InAccumulatedSeconds)
+{
+}
 
 FPoseSearchQueryTrajectorySample FPoseSearchQueryTrajectorySample::Lerp(const FPoseSearchQueryTrajectorySample& Other, float Alpha) const
 {
@@ -208,8 +216,34 @@ FArchive& operator<<(FArchive& Ar, FPoseSearchQueryTrajectorySample& TrajectoryS
 	return Ar;
 }
 
+FPoseSearchQueryTrajectory::FPoseSearchQueryTrajectory(const FTransformTrajectory& InTrajectory)
+{
+	Samples.Reset(InTrajectory.Samples.Num());
+		
+	for (int i = 0; i < InTrajectory.Samples.Num(); ++i)
+	{
+		Samples[i].Position = InTrajectory.Samples[i].Position;
+		Samples[i].Facing = InTrajectory.Samples[i].Facing;
+		Samples[i].AccumulatedSeconds = InTrajectory.Samples[i].TimeInSeconds;
+	}
+}
+
+FPoseSearchQueryTrajectory::operator FTransformTrajectory() const
+{
+	FTransformTrajectory OutTrajectory;
+
+	for (const auto & Sample : Samples)
+	{
+		OutTrajectory.Samples.Emplace(Sample.Facing, Sample.Position, Sample.AccumulatedSeconds);
+	}
+
+	return OutTrajectory;
+}
+
 FArchive& operator<<(FArchive& Ar, FPoseSearchQueryTrajectory& Trajectory)
 {
 	Ar << Trajectory.Samples;
 	return Ar;
-}
\ No newline at end of file
+}
+
+PRAGMA_ENABLE_DEPRECATION_WARNINGS
\ No newline at end of file
diff --git a/Engine/Plugins/Animation/PoseSearch/Source/Runtime/Private/Trace/PoseSearchTraceLogger.cpp b/Engine/Plugins/Animation/PoseSearch/Source/Runtime/Private/Trace/PoseSearchTraceLogger.cpp
--- a/Engine/Plugins/Animation/PoseSearch/Source/Runtime/Private/Trace/PoseSearchTraceLogger.cpp
+++ b/Engine/Plugins/Animation/PoseSearch/Source/Runtime/Private/Trace/PoseSearchTraceLogger.cpp
@@ -8,18 +8,26 @@
 #include "Serialization/MemoryWriter.h"
 #include "Trace/Trace.inl"
 #include "TraceFilter.h"
+#include "PoseSearch/PoseSearchCustomVersion.h"
 
 UE_TRACE_CHANNEL_DEFINE(PoseSearchChannel);
 
+// No custom versions applied when loading.
 UE_TRACE_EVENT_BEGIN(PoseSearch, MotionMatchingState)
 	UE_TRACE_EVENT_FIELD(uint8[], Data)
 UE_TRACE_EVENT_END()
 
+// We apply the FPoseSearchCustomVersion::DeprecatedTrajectoryTypes version when loading.
+UE_TRACE_EVENT_BEGIN(PoseSearch, MotionMatchingState2)
+	UE_TRACE_EVENT_FIELD(uint8[], Data)
+UE_TRACE_EVENT_END()
+
 namespace UE::PoseSearch
 {
 
+// @todo: These should be deprecated. No need to expose them.
 const FName FTraceLogger::Name("PoseSearch");
-const FName FTraceMotionMatchingStateMessage::Name("MotionMatchingState");
+const FName FTraceMotionMatchingStateMessage::Name("MotionMatchingState2");
 
 FArchive& operator<<(FArchive& Ar, FTraceMessage& State)
 {
@@ -77,8 +85,14 @@ void FTraceMotionMatchingStateMessage::Output()
 #if OBJECT_TRACE_ENABLED
 	TArray<uint8> ArchiveData;
 	FMemoryWriter Archive(ArchiveData);
+
+	// Ensure we have the latest version.
+	FCustomVersionContainer CustomVersionContainer;
+	CustomVersionContainer.SetVersion(FPoseSearchCustomVersion::GUID, FPoseSearchCustomVersion::LatestVersion, TEXT("Dev-PoseSearch-Version"));
+	Archive.SetCustomVersions(CustomVersionContainer);
+	
 	Archive << *this;
-	UE_TRACE_LOG(PoseSearch, MotionMatchingState, PoseSearchChannel) << MotionMatchingState.Data(ArchiveData.GetData(), ArchiveData.Num());
+	UE_TRACE_LOG(PoseSearch, MotionMatchingState2, PoseSearchChannel) << MotionMatchingState2.Data(ArchiveData.GetData(), ArchiveData.Num());
 #endif
 }
 
diff --git a/Engine/Plugins/Animation/PoseSearch/Source/Runtime/Public/PoseSearch/AnimNode_PoseSearchHistoryCollector.h b/Engine/Plugins/Animation/PoseSearch/Source/Runtime/Public/PoseSearch/AnimNode_PoseSearchHistoryCollector.h
--- a/Engine/Plugins/Animation/PoseSearch/Source/Runtime/Public/PoseSearch/AnimNode_PoseSearchHistoryCollector.h
+++ b/Engine/Plugins/Animation/PoseSearch/Source/Runtime/Public/PoseSearch/AnimNode_PoseSearchHistoryCollector.h
@@ -5,6 +5,7 @@
 #include "Animation/AnimNodeBase.h"
 #include "PoseSearch/PoseSearchHistory.h"
 #include "PoseSearch/PoseSearchTrajectoryTypes.h"
+#include "Animation/TrajectoryTypes.h"
 #include "AnimNode_PoseSearchHistoryCollector.generated.h"
 
 USTRUCT(BlueprintInternalUseOnly)
@@ -66,10 +67,19 @@ public:
 	UPROPERTY(EditAnywhere, Category = Experimental)
 	bool bGenerateTrajectory = false;
 
-	// input Trajectory samples for pose search queries in Motion Matching. These are expected to be in the world space of the SkeletalMeshComponent.
-	// the trajectory sample with AccumulatedSeconds equals to zero (Trajectory.Samples[i].AccumulatedSeconds) is the sample of the previous frame of simulation (since MM works by matching the previous character pose)
-	UPROPERTY(EditAnywhere, Category = Settings, meta = (PinShownByDefault, EditCondition="!bGenerateTrajectory", EditConditionHides))
+#if WITH_EDITORONLY_DATA
+
+	PRAGMA_DISABLE_DEPRECATION_WARNINGS
+	UE_DEPRECATED(5.6, "Please use TransformTrajectory variable instead.")
 	FPoseSearchQueryTrajectory Trajectory;
+	PRAGMA_ENABLE_DEPRECATION_WARNINGS
+
+#endif // WITH_EDITORONLY_DATA
+
+	// Input Trajectory samples for pose search queries in Motion Matching. These are expected to be in the world space of the SkeletalMeshComponent.
+	// the trajectory sample with AccumulatedSeconds equals to zero (Trajectory.Samples[i].AccumulatedSeconds) is the sample of the previous frame of simulation (since MM works by matching the previous character pose)
+	UPROPERTY(EditAnywhere, Transient, Category = Settings, DisplayName="Trajectory", meta = (PinShownByDefault, EditCondition="!bGenerateTrajectory", EditConditionHides))
+	FTransformTrajectory TransformTrajectory;
 
 	// Input Trajectory velocity will be multiplied by TrajectorySpeedMultiplier: values below 1 will result in selecting animation slower than requested from the original Trajectory
 	UPROPERTY(EditAnywhere, Category = Experimental, meta = (PinHiddenByDefault, ClampMin="0", EditCondition="!bGenerateTrajectory", EditConditionHides))
@@ -102,9 +112,9 @@ public:
 	// UE::PoseSearch::IPoseHistory interface
 	virtual bool GetTransformAtTime(float Time, FTransform& OutBoneTransform, const USkeleton* BoneIndexSkeleton = nullptr, FBoneIndexType BoneIndexType = UE::PoseSearch::RootBoneIndexType, FBoneIndexType ReferenceBoneIndexType = UE::PoseSearch::ComponentSpaceIndexType, bool bExtrapolate = false) const override;
 	virtual bool GetCurveValueAtTime(float Time, const FName& CurveName, float& OutCurveValue, bool bExtrapolate = false) const override;
-	virtual const FPoseSearchQueryTrajectory& GetTrajectory() const override;
+	virtual const FTransformTrajectory& GetTrajectory() const override;
 	// Experimental, this feature might be removed without warning, not for production use
-	virtual void SetTrajectory(const FPoseSearchQueryTrajectory& InTrajectory, float InTrajectorySpeedMultiplier = 1.f) override;
+	virtual void SetTrajectory(const FTransformTrajectory& InTrajectory, float InTrajectorySpeedMultiplier = 1.f) override;
 	virtual void GenerateTrajectory(const UObject* AnimContext, float DeltaTime) override;
 	virtual float GetTrajectorySpeedMultiplier() const override;
 	virtual bool IsEmpty() const override;
@@ -128,6 +138,17 @@ public:
 	UE_DEPRECATED(5.6, "Use this::IPoseHistory::GenerateTrajectory interface instead")
 	void GenerateTrajectory(const UAnimInstance* InAnimInstance);
 
+	PRAGMA_DISABLE_DEPRECATION_WARNINGS
+	// Note: We need to explicitly disable warnings on these constructors/operators for clang to be happy with deprecated variables
+	// this is a requirement for clang to compile without warnings.
+	 FAnimNode_PoseSearchHistoryCollector_Base() = default;
+	~FAnimNode_PoseSearchHistoryCollector_Base() = default;
+	FAnimNode_PoseSearchHistoryCollector_Base(const FAnimNode_PoseSearchHistoryCollector_Base&) = default;
+	FAnimNode_PoseSearchHistoryCollector_Base(FAnimNode_PoseSearchHistoryCollector_Base&&) = default;
+	FAnimNode_PoseSearchHistoryCollector_Base& operator=(const FAnimNode_PoseSearchHistoryCollector_Base&) = default;
+	FAnimNode_PoseSearchHistoryCollector_Base& operator=(FAnimNode_PoseSearchHistoryCollector_Base&&) = default;
+	PRAGMA_ENABLE_DEPRECATION_WARNINGS
+	
 protected:
 	// returns an array of skeleton bone indexes (encoded as FBoneIndexType)
 	TArray<FBoneIndexType> GetRequiredBones(const FAnimInstanceProxy* AnimInstanceProxy) const;
diff --git a/Engine/Plugins/Animation/PoseSearch/Source/Runtime/Public/PoseSearch/PoseSearchContext.h b/Engine/Plugins/Animation/PoseSearch/Source/Runtime/Public/PoseSearch/PoseSearchContext.h
--- a/Engine/Plugins/Animation/PoseSearch/Source/Runtime/Public/PoseSearch/PoseSearchContext.h
+++ b/Engine/Plugins/Animation/PoseSearch/Source/Runtime/Public/PoseSearch/PoseSearchContext.h
@@ -10,7 +10,7 @@
 #include "PoseSearch/PoseSearchIndex.h"
 #include "PoseSearch/PoseSearchResult.h"
 
-struct FPoseSearchQueryTrajectory;
+struct FTransformTrajectory;
 class UPoseSearchDatabase;
 class UPoseSearchFeatureChannel_Position;
 
diff --git a/Engine/Plugins/Animation/PoseSearch/Source/Runtime/Public/PoseSearch/PoseSearchCustomVersion.h b/Engine/Plugins/Animation/PoseSearch/Source/Runtime/Public/PoseSearch/PoseSearchCustomVersion.h
new file mode 100644
--- /dev/null
+++ b/Engine/Plugins/Animation/PoseSearch/Source/Runtime/Public/PoseSearch/PoseSearchCustomVersion.h
@@ -0,0 +1,29 @@
+// Copyright Epic Games, Inc. All Rights Reserved.
+
+#pragma once
+
+#include "CoreMinimal.h"
+#include "Misc/Guid.h"
+
+// Custom serialization version for all packages containing PoseSearch dependent asset types
+struct POSESEARCH_API FPoseSearchCustomVersion
+{
+	enum Type
+	{
+		// Before any version changes were made in the plugin
+		BeforeCustomVersionWasAdded = 0,
+
+		// Deprecated FPoseSearchQueryTrajectory and FPoseSearchQueryTrajectorySample. This version matches the RouteId_MotionMatchingState_2 event.
+		DeprecatedTrajectoryTypes,
+
+		// -----<new versions can be added above this line>-------------------------------------------------
+		VersionPlusOne,
+		LatestVersion = VersionPlusOne - 1
+	};
+
+	// The GUID for this custom version number
+	const static FGuid GUID;
+
+private:
+	FPoseSearchCustomVersion() {}
+};
diff --git a/Engine/Plugins/Animation/PoseSearch/Source/Runtime/Public/PoseSearch/PoseSearchHistory.h b/Engine/Plugins/Animation/PoseSearch/Source/Runtime/Public/PoseSearch/PoseSearchHistory.h
--- a/Engine/Plugins/Animation/PoseSearch/Source/Runtime/Public/PoseSearch/PoseSearchHistory.h
+++ b/Engine/Plugins/Animation/PoseSearch/Source/Runtime/Public/PoseSearch/PoseSearchHistory.h
@@ -9,6 +9,7 @@
 #include "PoseSearch/PoseHistoryProvider.h"
 #include "PoseSearch/PoseSearchDefines.h"
 #include "PoseSearch/PoseSearchTrajectoryLibrary.h"
+#include "Animation/TrajectoryTypes.h"
 #include "UObject/ObjectKey.h"
 
 struct FAnimInstanceProxy;
@@ -83,9 +84,14 @@ struct POSESEARCH_API IPoseHistory
 	virtual bool GetTransformAtTime(float Time, FTransform& OutBoneTransform, const USkeleton* BoneIndexSkeleton = nullptr, FBoneIndexType BoneIndexType = RootBoneIndexType, FBoneIndexType ReferenceBoneIndexType = ComponentSpaceIndexType, bool bExtrapolate = false) const = 0;
 	// @todo: consider consolidating into a (templated?) get X value at time once we add custom attributes to pose history.
 	virtual bool GetCurveValueAtTime(float Time, const FName& CurveName, float& OutCurveValue, bool bExtrapolate = false) const = 0;
-	virtual const FPoseSearchQueryTrajectory& GetTrajectory() const = 0;
+	
+	// @note for review: Here we are just changing the return type to force users to update immediately. Will cause project to not compile for users who override the pure virtual func getter, not sure how we can gracefully deprecate this.
+	// Another option is to just have a new function signature GetTransformTrajectory() and add a final tag to the GetTrajectory() to also force users to upgrade.
+	// Lastly, to me this seems like a struct not many licensees would inherent from tbh.
+	virtual const FTransformTrajectory& GetTrajectory() const = 0;
+	
 	// Experimental, this feature might be removed without warning, not for production use
-	virtual void SetTrajectory(const FPoseSearchQueryTrajectory& InTrajectory, float InTrajectorySpeedMultiplier = 1.f) = 0;
+	virtual void SetTrajectory(const FTransformTrajectory& InTrajectory, float InTrajectorySpeedMultiplier = 1.f) = 0;
 	// Experimental, this feature might be removed without warning, not for production use
 	virtual void GenerateTrajectory(const UObject* AnimContext, float DeltaTime) = 0;
 	// @todo: deprecate this API. TrajectorySpeedMultiplier should be a global query scaling value passed as input parameter of FSearchContext during config BuildQuery
@@ -111,9 +117,9 @@ struct POSESEARCH_API FArchivedPoseHistory : public IPoseHistory
 	// IPoseHistory interface
 	virtual bool GetTransformAtTime(float Time, FTransform& OutBoneTransform, const USkeleton* BoneIndexSkeleton = nullptr, FBoneIndexType BoneIndexType = RootBoneIndexType, FBoneIndexType ReferenceBoneIndexType = ComponentSpaceIndexType, bool bExtrapolate = false) const override;
 	virtual bool GetCurveValueAtTime(float Time, const FName& CurveName, float& outcurvevalue, bool bExtrapolate = false) const override;
-	virtual const FPoseSearchQueryTrajectory& GetTrajectory() const override { return Trajectory; }
+	virtual const FTransformTrajectory& GetTrajectory() const override { return Trajectory; }
 	// Experimental, this feature might be removed without warning, not for production use
-	virtual void SetTrajectory(const FPoseSearchQueryTrajectory& InTrajectory, float InTrajectorySpeedMultiplier = 1.f) override { Trajectory = InTrajectory; }
+	virtual void SetTrajectory(const FTransformTrajectory& InTrajectory, float InTrajectorySpeedMultiplier = 1.f) override { Trajectory = InTrajectory; }
 	// Experimental, this feature might be removed without warning, not for production use
 	virtual void GenerateTrajectory(const UObject* AnimContext, float DeltaTime) override { checkNoEntry(); }
 	virtual float GetTrajectorySpeedMultiplier() const override { return 1.f; }
@@ -133,13 +139,18 @@ struct POSESEARCH_API FArchivedPoseHistory : public IPoseHistory
 	// @todo: Make this a map if this is expected to be big.
 	TArray<FName> CollectedCurves;
 	TArray<FPoseHistoryEntry> Entries;
-	FPoseSearchQueryTrajectory Trajectory;
+	
+	FTransformTrajectory Trajectory;
 };
 FArchive& operator<<(FArchive& Ar, FArchivedPoseHistory& Entry);
 
 struct POSESEARCH_API FPoseHistory : public IPoseHistory
 {
+	PRAGMA_DISABLE_DEPRECATION_WARNINGS
 	FPoseHistory() = default;
+	virtual ~FPoseHistory() override = default;
+	PRAGMA_ENABLE_DEPRECATION_WARNINGS
+
 	FPoseHistory(const FPoseHistory& Other);
 	FPoseHistory(FPoseHistory&& Other);
 	FPoseHistory& operator=(const FPoseHistory& Other);
@@ -165,15 +176,21 @@ struct POSESEARCH_API FPoseHistory : public IPoseHistory
 	// IPoseHistory interface
 	virtual bool GetTransformAtTime(float Time, FTransform& OutBoneTransform, const USkeleton* BoneIndexSkeleton = nullptr, FBoneIndexType BoneIndexType = RootBoneIndexType, FBoneIndexType ReferenceBoneIndexType = ComponentSpaceIndexType, bool bExtrapolate = false) const override;
 	virtual bool GetCurveValueAtTime(float Time, const FName& CurveName, float& outcurvevalue, bool bExtrapolate = false) const override;
-	virtual const FPoseSearchQueryTrajectory& GetTrajectory() const override;
+	virtual const FTransformTrajectory& GetTrajectory() const override;
 	// Experimental, this feature might be removed without warning, not for production use
-	virtual void SetTrajectory(const FPoseSearchQueryTrajectory& InTrajectory, float InTrajectorySpeedMultiplier = 1.f) override;
+	virtual void SetTrajectory(const FTransformTrajectory& InTrajectory, float InTrajectorySpeedMultiplier = 1.f) override;
 	// Experimental, this feature might be removed without warning, not for production use
 	virtual void GenerateTrajectory(const UObject* AnimContext, float DeltaTime) override;
 	virtual float GetTrajectorySpeedMultiplier() const override;
 	virtual bool IsEmpty() const override;
 	virtual const FBoneToTransformMap& GetBoneToTransformMap() const override;
 	virtual const TConstArrayView<FName> GetCollectedCurves() const override;
+
+	// We need to handle the changes from CL 39422514 since licensees still haven't gotten the move of this function to the IPoseHistory, as of 5.5, therefore this function still needs to be marked as deprecated even tho it was deleted later on.
+	PRAGMA_DISABLE_DEPRECATION_WARNINGS
+	UE_DEPRECATED(5.6, "SetTrajectory() that takes in FPoseSearchQueryTrajectory is deprecated. Use version that takes FTransformTrajectory instead.")
+	void SetTrajectory(const FPoseSearchQueryTrajectory& InTrajectory, float InTrajectorySpeedMultiplier = 1.f);
+	PRAGMA_ENABLE_DEPRECATION_WARNINGS
 	
 	virtual int32 GetNumEntries() const override;
 	virtual const FPoseHistoryEntry& GetEntry(int32 EntryIndex) const override;
@@ -193,8 +210,9 @@ private:
 	int32 MaxNumPoses = 0;
 
 	float SamplingInterval = 0.f;
-
-	FPoseSearchQueryTrajectory Trajectory;
+	
+	FTransformTrajectory Trajectory;
+	
 	FPoseSearchTrajectoryData::FState TrajectoryDataState;
 	// @todo: deprecate this member and expose it via blue print logic or as global query scaling multiplier
 	float TrajectorySpeedMultiplier = 1.f;
@@ -228,9 +246,9 @@ struct POSESEARCH_API FMemStackPoseHistory : public IPoseHistory
 	// IPoseHistory interface
 	virtual bool GetTransformAtTime(float Time, FTransform& OutBoneTransform, const USkeleton* BoneIndexSkeleton = nullptr, FBoneIndexType BoneIndexType = RootBoneIndexType, FBoneIndexType ReferenceBoneIndexType = ComponentSpaceIndexType, bool bExtrapolate = false) const override;
 	virtual bool GetCurveValueAtTime(float Time, const FName& CurveName, float& outcurvevalue, bool bExtrapolate = false) const override;
-	virtual const FPoseSearchQueryTrajectory& GetTrajectory() const override { check(PoseHistory); return PoseHistory->GetTrajectory(); }
+	virtual const FTransformTrajectory& GetTrajectory() const override { check(PoseHistory); return PoseHistory->GetTrajectory(); }
 	// Experimental, this feature might be removed without warning, not for production use
-	virtual void SetTrajectory(const FPoseSearchQueryTrajectory& InTrajectory, float InTrajectorySpeedMultiplier = 1.f) override;
+	virtual void SetTrajectory(const FTransformTrajectory& InTrajectory, float InTrajectorySpeedMultiplier = 1.f) override;
 	// Experimental, this feature might be removed without warning, not for production use
 	virtual void GenerateTrajectory(const UObject* AnimContext, float DeltaTime) override;
 	virtual float GetTrajectorySpeedMultiplier() const override { check(PoseHistory); return PoseHistory->GetTrajectorySpeedMultiplier(); }
diff --git a/Engine/Plugins/Animation/PoseSearch/Source/Runtime/Public/PoseSearch/PoseSearchHistoryCollectorAnimNodeLibrary.h b/Engine/Plugins/Animation/PoseSearch/Source/Runtime/Public/PoseSearch/PoseSearchHistoryCollectorAnimNodeLibrary.h
--- a/Engine/Plugins/Animation/PoseSearch/Source/Runtime/Public/PoseSearch/PoseSearchHistoryCollectorAnimNodeLibrary.h
+++ b/Engine/Plugins/Animation/PoseSearch/Source/Runtime/Public/PoseSearch/PoseSearchHistoryCollectorAnimNodeLibrary.h
@@ -7,6 +7,10 @@
 #include "Kismet/BlueprintFunctionLibrary.h"
 #include "PoseSearchHistoryCollectorAnimNodeLibrary.generated.h"
 
+PRAGMA_DISABLE_DEPRECATION_WARNINGS
+struct FPoseSearchQueryTrajectory;
+PRAGMA_ENABLE_DEPRECATION_WARNINGS
+struct FTransformTrajectory;
 struct FAnimNode_PoseSearchHistoryCollector;
 
 USTRUCT(BlueprintType)
@@ -30,16 +34,21 @@ public:
 
 	/** Get a Pose History node context from an anim node context (pure) */
 	UFUNCTION(BlueprintPure, Category = "Animation|PoseHistory", meta = (BlueprintThreadSafe, DisplayName = "Convert to Pose History Node"))
-	static void ConvertToPoseHistoryNodePure(const FAnimNodeReference& Node, FPoseSearchHistoryCollectorAnimNodeReference& PoseSearchHistoryCollectorNode, bool& Result)
-	{
-		EAnimNodeReferenceConversionResult ConversionResult;
-		PoseSearchHistoryCollectorNode = ConvertToPoseHistoryNode(Node, ConversionResult);
-		Result = (ConversionResult == EAnimNodeReferenceConversionResult::Succeeded);
-	}
+	static void ConvertToPoseHistoryNodePure(const FAnimNodeReference& Node, FPoseSearchHistoryCollectorAnimNodeReference& PoseSearchHistoryCollectorNode, bool& Result);
 
 	UFUNCTION(BlueprintPure, Category = "Animation|PoseHistory", meta = (BlueprintThreadSafe, DisplayName = "Get Pose History Node Trajectory"))
-	static void GetPoseHistoryNodeTrajectory(const FPoseSearchHistoryCollectorAnimNodeReference& PoseSearchHistoryCollectorNode, FPoseSearchQueryTrajectory& Trajectory);
+	static void GetPoseHistoryNodeTransformTrajectory(const FPoseSearchHistoryCollectorAnimNodeReference& PoseSearchHistoryCollectorNode, FTransformTrajectory& Trajectory);
 
 	UFUNCTION(BlueprintCallable, Category = "Animation|PoseHistory", meta = (BlueprintThreadSafe, DisplayName = "Set Pose History Node Trajectory"))
+	static void SetPoseHistoryNodeTransformTrajectory(const FPoseSearchHistoryCollectorAnimNodeReference& PoseSearchHistoryCollectorNode, const FTransformTrajectory& Trajectory);
+	
+	PRAGMA_DISABLE_DEPRECATION_WARNINGS
+	
+	UE_DEPRECATED(5.6, "Please use GetPoseHistoryNodeTrajectory that takes in FTransformTrajectory instead")
+	static void GetPoseHistoryNodeTrajectory(const FPoseSearchHistoryCollectorAnimNodeReference& PoseSearchHistoryCollectorNode, FPoseSearchQueryTrajectory& Trajectory);
+	
+	UE_DEPRECATED(5.6, "Please use GetPoseHistoryNodeTrajectory that takes in FTransformTrajectory instead")
 	static void SetPoseHistoryNodeTrajectory(const FPoseSearchHistoryCollectorAnimNodeReference& PoseSearchHistoryCollectorNode, const FPoseSearchQueryTrajectory& Trajectory);
+	
+	PRAGMA_ENABLE_DEPRECATION_WARNINGS
 };
\ No newline at end of file
diff --git a/Engine/Plugins/Animation/PoseSearch/Source/Runtime/Public/PoseSearch/PoseSearchTrajectoryLibrary.h b/Engine/Plugins/Animation/PoseSearch/Source/Runtime/Public/PoseSearch/PoseSearchTrajectoryLibrary.h
--- a/Engine/Plugins/Animation/PoseSearch/Source/Runtime/Public/PoseSearch/PoseSearchTrajectoryLibrary.h
+++ b/Engine/Plugins/Animation/PoseSearch/Source/Runtime/Public/PoseSearch/PoseSearchTrajectoryLibrary.h
@@ -9,6 +9,8 @@
 
 class UAnimInstance;
 class IPoseSearchTrajectoryPredictorInterface;
+struct FTransformTrajectory;
+struct FTransformTrajectorySample;
 
 USTRUCT(BlueprintType)
 struct POSESEARCH_API FPoseSearchTrajectoryData
@@ -105,6 +107,11 @@ class POSESEARCH_API UPoseSearchTrajectoryLibrary : public UBlueprintFunctionLib
 
 public:
 
+	///////////////////////////////////////////////////////////////////////////
+	// DEPRECATED /////////////////////////////////////////////////////////////
+	///////////////////////////////////////////////////////////////////////////
+	
+	PRAGMA_DISABLE_DEPRECATION_WARNINGS
 	UE_DEPRECATED(5.5, "Use the InitTrajectorySamples function that takes DefaultPosition and DefaultFacing arguments instead of TrajectoryDataDerived")
 	static void InitTrajectorySamples(FPoseSearchQueryTrajectory& Trajectory, const FPoseSearchTrajectoryData& TrajectoryData, const FPoseSearchTrajectoryData::FDerived& TrajectoryDataDerived, const FPoseSearchTrajectoryData::FSampling& TrajectoryDataSampling, float DeltaTime);
 
@@ -113,6 +120,7 @@ public:
 	static void InitTrajectorySamples(FPoseSearchQueryTrajectory& Trajectory, const FPoseSearchTrajectoryData& TrajectoryData, FVector DefaultPosition, FQuat DefaultFacing, const FPoseSearchTrajectoryData::FSampling& TrajectoryDataSampling, float DeltaTime);
 
 	// Initialize history and predicted samples based on sampling settings and a default state
+	UE_DEPRECATED(5.6, "Use the IniTrajctorySamples function that takes in FTransformTrajectory instead of FPoseSearchQueryTrajectory.")
 	static void InitTrajectorySamples(FPoseSearchQueryTrajectory& Trajectory, FVector DefaultPosition, FQuat DefaultFacing, const FPoseSearchTrajectoryData::FSampling& TrajectoryDataSampling, float DeltaTime);
 
 	UE_DEPRECATED(5.5, "Use the UpdateHistory_TransformHistory function that takes CurrentPosition and CurrentVelocity arguments instead of TrajectoryDataDerived")
@@ -120,76 +128,136 @@ public:
 
 	UE_DEPRECATED(5.6, "Use the UpdateHistory_TransformHistory function that doesn't take FPoseSearchTrajectoryData")
 	static void UpdateHistory_TransformHistory(FPoseSearchQueryTrajectory& Trajectory, const FPoseSearchTrajectoryData& TrajectoryData, FVector CurrentPosition, FVector CurrentVelocity, const FPoseSearchTrajectoryData::FSampling& TrajectoryDataSampling, float DeltaTime);
-
-	// Update history by tracking offsets that result from character intent (e.g. movement component velocity) and applying
-	// that to the current world transform. This works well on moving platforms as it only stores a history of movement
-	// that results from character intent, not movement from platforms.
-	// Important: CurrentVelocity should be the velocity relative to the ground as reported by the character movement component or character mover etc
+	
+	UE_DEPRECATED(5.6, "Use the UpdateHistory_TransformHistory function thats in FTransformTrajectory instead of FPoseSearchQueryTrajectory.")
 	static void UpdateHistory_TransformHistory(FPoseSearchQueryTrajectory& Trajectory, FVector CurrentPosition, FVector CurrentVelocity, const FPoseSearchTrajectoryData::FSampling& TrajectoryDataSampling, float DeltaTime);
 
-	// Update the history purely based on current position, without taking into account ground velocities
-	// DeltaTime is the time since the last call to this method
+	UE_DEPRECATED(5.6, "Use the UpdateHistory_WorldSpace function thats in FTransformTrajectory instead of FPoseSearchQueryTrajectory.")
 	static void UpdateHistory_WorldSpace(FPoseSearchQueryTrajectory& Trajectory, FVector CurrentPosition, FQuat CurrentRotation, const FPoseSearchTrajectoryData::FSampling& TrajectoryDataSampling, float DeltaTime);
 
-	// Update prediction by simulating the movement math for ground locomotion from UCharacterMovementComponent.
+	UE_DEPRECATED(5.6, "Use the UpdatePrediction_SimulateCharacterMovement function thats in FTransformTrajectory instead of FPoseSearchQueryTrajectory.")
 	static void UpdatePrediction_SimulateCharacterMovement(FPoseSearchQueryTrajectory& Trajectory, const FPoseSearchTrajectoryData& TrajectoryData, const FPoseSearchTrajectoryData::FDerived& TrajectoryDataDerived, const FPoseSearchTrajectoryData::FSampling& TrajectoryDataSampling, float DeltaTime);
 
-	// @todo: rename InAnimInstance into Context, since it can now be either an AnimInstance or a Character
 	// Generates a prediction trajectory based of the current character intent. For use with Character actors.
-	UFUNCTION(BlueprintCallable, Category = "Animation|PoseSearch", meta = (BlueprintThreadSafe, DisplayName = "Pose Search Generate Trajectory (for Character)"))
+	UE_DEPRECATED(5.6, "Use PoseSearchGenerateTrajectory function that takes in FTransformTrajectory.")
 	static void PoseSearchGenerateTrajectory(const UObject* InAnimInstance, 
 		UPARAM(ref) const FPoseSearchTrajectoryData& InTrajectoryData, float InDeltaTime,
 		UPARAM(ref) FPoseSearchQueryTrajectory& InOutTrajectory, UPARAM(ref) float& InOutDesiredControllerYawLastUpdate, FPoseSearchQueryTrajectory& OutTrajectory,
 		float InHistorySamplingInterval = 0.04f, int32 InTrajectoryHistoryCount = 10, float InPredictionSamplingInterval = 0.2f, int32 InTrajectoryPredictionCount = 8);
 
 	// Generates a prediction trajectory based of the current movement intent. For use with predictors. InPredictor must implement IPoseSearchTrajectoryPredictorInterface
-	UE_DEPRECATED(5.6, "Use PoseSearchGenerateTrajectoryWithPredictor function that doesn't take FPoseSearchTrajectoryData")
-	UFUNCTION(BlueprintCallable, Category = "Animation|PoseSearch|Experimental", meta = (BlueprintThreadSafe, DisplayName = "Pose Search Generate Trajectory (using Predictor) (DEPRECATED)"))
+	UE_DEPRECATED(5.6, "Use PoseSearchGenerateTrajectoryWithPredictor function that doesn't take FPoseSearchTrajectoryData and also takes in FTransformTrajectory")
 	static void PoseSearchGeneratePredictorTrajectory(UObject* InPredictor,
 		UPARAM(ref) const FPoseSearchTrajectoryData& InTrajectoryData, float InDeltaTime,
 		UPARAM(ref) FPoseSearchQueryTrajectory& InOutTrajectory, UPARAM(ref) float& InOutDesiredControllerYawLastUpdate, FPoseSearchQueryTrajectory& OutTrajectory,
 		float InHistorySamplingInterval = 0.04f, int32 InTrajectoryHistoryCount = 10, float InPredictionSamplingInterval = 0.2f, int32 InTrajectoryPredictionCount = 8);
 
 	// Generates a prediction trajectory based of the current movement intent. For use with predictors. InPredictor must implement IPoseSearchTrajectoryPredictorInterface
-	UFUNCTION(BlueprintCallable, Category = "Animation|PoseSearch|Experimental", meta = (BlueprintThreadSafe, DisplayName = "Pose Search Generate Trajectory (using Predictor)"))
+	UE_DEPRECATED(5.6, "Use PoseSearchGenerateTrajectoryWithPredictor function that takes in FTransformTrajectory")
 	static void PoseSearchGenerateTrajectoryWithPredictor(TScriptInterface<IPoseSearchTrajectoryPredictorInterface> InPredictor,
 		float InDeltaTime,
 		UPARAM(ref) FPoseSearchQueryTrajectory& InOutTrajectory, UPARAM(ref) float& InOutDesiredControllerYawLastUpdate, FPoseSearchQueryTrajectory& OutTrajectory,
 		float InHistorySamplingInterval = 0.04f, int32 InTrajectoryHistoryCount = 10, float InPredictionSamplingInterval = 0.2f, int32 InTrajectoryPredictionCount = 8);
 
-	// Experimental: Process InTrajectory to apply gravity and handle collisions. Eventually returns the modified OutTrajectory.
-	// If bApplyGravity is true, gravity from the UCharacterMovementComponent will be applied.
-	// If FloorCollisionsOffset > 0, vertical collision will be performed to every sample of the trajectory to have the samples float over the geometry (by FloorCollisionsOffset).
-	UFUNCTION(BlueprintCallable, Category="Animation|PoseSearch|Experimental", meta=(BlueprintThreadSafe, WorldContext="WorldContextObject", AutoCreateRefTerm="ActorsToIgnore", AdvancedDisplay="TraceChannel,bTraceComplex,ActorsToIgnore,DrawDebugType,bIgnoreSelf,MaxObstacleHeight,TraceColor,TraceHitColor,DrawTime"))
+	UE_DEPRECATED(5.6, "Use HandleTrajectoryWorldCollisions function that takes in FTransformTrajectory")
 	static void HandleTrajectoryWorldCollisions(const UObject* WorldContextObject, const UAnimInstance* AnimInstance, UPARAM(ref) const FPoseSearchQueryTrajectory& InTrajectory, bool bApplyGravity, float FloorCollisionsOffset, FPoseSearchQueryTrajectory& OutTrajectory, FPoseSearchTrajectory_WorldCollisionResults& CollisionResult,
 		ETraceTypeQuery TraceChannel, bool bTraceComplex, const TArray<AActor*>& ActorsToIgnore, EDrawDebugTrace::Type DrawDebugType, bool bIgnoreSelf = true, float MaxObstacleHeight = 10000.f, FLinearColor TraceColor = FLinearColor::Red, FLinearColor TraceHitColor = FLinearColor::Green, float DrawTime = 5.0f);
 
-	// Experimental: Process InTrajectory to apply gravity and handle collisions. Eventually returns the modified OutTrajectory.
-	// If bApplyGravity is true, GravityAccel will be applied.
-	// If FloorCollisionsOffset > 0, vertical collision will be performed to every sample of the trajectory to have the samples float over the geometry (by FloorCollisionsOffset).
-	UFUNCTION(BlueprintCallable, Category="Animation|PoseSearch|Experimental", meta=(BlueprintThreadSafe, WorldContext="WorldContextObject", AutoCreateRefTerm="ActorsToIgnore", AdvancedDisplay="TraceChannel,bTraceComplex,ActorsToIgnore,DrawDebugType,bIgnoreSelf,MaxObstacleHeight,TraceColor,TraceHitColor,DrawTime"))
+	UE_DEPRECATED(5.6, "Use HandleTrajectoryWorldCollisionsWithGravity that takes in FTransformTrajectory.")
 	static void HandleTrajectoryWorldCollisionsWithGravity(const UObject* WorldContextObject, 
 		UPARAM(ref) const FPoseSearchQueryTrajectory& InTrajectory, FVector StartingVelocity, bool bApplyGravity, FVector GravityAccel, float FloorCollisionsOffset, FPoseSearchQueryTrajectory& OutTrajectory, FPoseSearchTrajectory_WorldCollisionResults& CollisionResult,
 		ETraceTypeQuery TraceChannel, bool bTraceComplex, const TArray<AActor*>& ActorsToIgnore, EDrawDebugTrace::Type DrawDebugType, bool bIgnoreSelf = true, float MaxObstacleHeight = 10000.f, FLinearColor TraceColor = FLinearColor::Red, FLinearColor TraceHitColor = FLinearColor::Green, float DrawTime = 5.0f);
 
-	UFUNCTION(BlueprintPure, meta = (BlueprintThreadSafe), Category="Animation|PoseSearch")
+	UE_DEPRECATED(5.6, "Use GetTrajectorySampleAtTime that takes in FTransformTrajectory and FTransformTrajectorySample.")
 	static void GetTrajectorySampleAtTime(UPARAM(ref) const FPoseSearchQueryTrajectory& InTrajectory, float Time, FPoseSearchQueryTrajectorySample& OutTrajectorySample, bool bExtrapolate = false);
 
-	UFUNCTION(BlueprintPure, meta = (BlueprintThreadSafe), Category="Animation|PoseSearch")
+	UE_DEPRECATED(5.6, "Use GetTrajectorySampleAtTime that takes in FTransformTrajectory.")
 	static void GetTrajectoryVelocity(UPARAM(ref) const FPoseSearchQueryTrajectory& InTrajectory, float Time1, float Time2, FVector& OutVelocity, bool bExtrapolate = false);
-	
-	UFUNCTION(BlueprintPure, meta = (BlueprintThreadSafe), Category="Animation|PoseSearch")
+
+	UE_DEPRECATED(5.6, "Use GetTrajectorySampleAtTime that takes in FTransformTrajectory.")
 	static void GetTrajectoryAngularVelocity(UPARAM(ref) const FPoseSearchQueryTrajectory& InTrajectory, float Time1, float Time2, FVector& OutAngularVelocity, bool bExtrapolate = false);
 
-	UFUNCTION(BlueprintPure, meta = (BlueprintThreadSafe), Category="Animation|PoseSearch")
-	static FTransform GetTransform(UPARAM(ref) const FPoseSearchQueryTrajectorySample& InTrajectorySample)
-	{
-		return InTrajectorySample.GetTransform();
-	};
-	
-	UFUNCTION(BlueprintCallable, Category="Animation|PoseSearch|Experimental", meta=(WorldContext="WorldContextObject"))
+	UE_DEPRECATED(5.6, "Use GetTrajectorySampleAtTime that takes in FTransformTrajectorySample.")
+	static FTransform GetTransform(UPARAM(ref) const FPoseSearchQueryTrajectorySample& InTrajectorySample);
+
+	UE_DEPRECATED(5.6, "Use DrawTrajectory that takes in FTransformTrajectory.")
 	static void DrawTrajectory(const UObject* WorldContextObject, UPARAM(ref) const FPoseSearchQueryTrajectory& InTrajectory, const float DebugThickness, float HeightOffset);
+	PRAGMA_ENABLE_DEPRECATION_WARNINGS
+
+	///////////////////////////////////////////////////////////////////////////
+	// NON- DEPRECATED ////////////////////////////////////////////////////////
+	///////////////////////////////////////////////////////////////////////////
+	
+	// Initialize history and predicted samples based on sampling settings and a default state
+	static void InitTrajectorySamples(FTransformTrajectory& Trajectory, FVector DefaultPosition, FQuat DefaultFacing, const FPoseSearchTrajectoryData::FSampling& TrajectoryDataSampling, float DeltaTime);
+
+	// Update history by tracking offsets that result from character intent (e.g. movement component velocity) and applying
+	// that to the current world transform. This works well on moving platforms as it only stores a history of movement
+	// that results from character intent, not movement from platforms.
+	// Important: CurrentVelocity should be the velocity relative to the ground as reported by the character movement component or character mover etc
+	static void UpdateHistory_TransformHistory(FTransformTrajectory& Trajectory, FVector CurrentPosition, FVector CurrentVelocity, const FPoseSearchTrajectoryData::FSampling& TrajectoryDataSampling, float DeltaTime);
 
+	// Update prediction by simulating the movement math for ground locomotion from UCharacterMovementComponent.
+	static void UpdatePrediction_SimulateCharacterMovement(FTransformTrajectory& Trajectory, const FPoseSearchTrajectoryData& TrajectoryData, const FPoseSearchTrajectoryData::FDerived& TrajectoryDataDerived, const FPoseSearchTrajectoryData::FSampling& TrajectoryDataSampling, float DeltaTime);
+
+	// Update the history purely based on current position, without taking into account ground velocities
+	// DeltaTime is the time since the last call to this method	
+	static void UpdateHistory_WorldSpace(FTransformTrajectory& Trajectory, FVector CurrentPosition, FQuat CurrentRotation, const FPoseSearchTrajectoryData::FSampling& TrajectoryDataSampling, float DeltaTime);
+	
+	// Generates a prediction trajectory based of the current character intent. For use with Character actors.
+	UFUNCTION(BlueprintCallable, Category = "Animation|PoseSearch", meta = (BlueprintThreadSafe, DisplayName = "Pose Search Generate Trajectory (for Character)"))
+	static void PoseSearchGenerateTransformTrajectory(const UObject* InAnimInstance, 
+		UPARAM(ref) const FPoseSearchTrajectoryData& InTrajectoryData, float InDeltaTime,
+		UPARAM(ref) FTransformTrajectory& InOutTrajectory, UPARAM(ref) float& InOutDesiredControllerYawLastUpdate, FTransformTrajectory& OutTrajectory,
+		float InHistorySamplingInterval = 0.04f, int32 InTrajectoryHistoryCount = 10, float InPredictionSamplingInterval = 0.2f, int32 InTrajectoryPredictionCount = 8);
+
+	// Generates a prediction trajectory based of the current movement intent. For use with predictors. InPredictor must implement IPoseSearchTrajectoryPredictorInterface
+	UFUNCTION(BlueprintCallable, Category = "Animation|PoseSearch|Experimental", meta = (BlueprintThreadSafe, DisplayName = "Pose Search Generate Trajectory (using Predictor)", DeprecatedFunction, DeprecationMessage="Use PoseSearchGeneratePredictorTrajectory that does not take FPoseSearchTrajectoryData."))
+	static void PoseSearchGeneratePredictorTransformTrajectory(UObject* InPredictor,
+	UPARAM(ref) const FPoseSearchTrajectoryData& InTrajectoryData, float InDeltaTime,
+	UPARAM(ref) FTransformTrajectory& InOutTrajectory, UPARAM(ref) float& InOutDesiredControllerYawLastUpdate, FTransformTrajectory& OutTrajectory,
+	float InHistorySamplingInterval = 0.04f, int32 InTrajectoryHistoryCount = 10, float InPredictionSamplingInterval = 0.2f, int32 InTrajectoryPredictionCount = 8);
+	
+	// Generates a prediction trajectory based of the current movement intent. For use with predictors. InPredictor must implement IPoseSearchTrajectoryPredictorInterface
+	UFUNCTION(BlueprintCallable, Category = "Animation|PoseSearch|Experimental", meta = (BlueprintThreadSafe, DisplayName = "Pose Search Generate Trajectory (using Predictor)"))
+	static void PoseSearchGenerateTransformTrajectoryWithPredictor(TScriptInterface<IPoseSearchTrajectoryPredictorInterface> InPredictor,
+		float InDeltaTime,
+		UPARAM(ref) FTransformTrajectory& InOutTrajectory, UPARAM(ref) float& InOutDesiredControllerYawLastUpdate, FTransformTrajectory& OutTrajectory,
+		float InHistorySamplingInterval = 0.04f, int32 InTrajectoryHistoryCount = 10, float InPredictionSamplingInterval = 0.2f, int32 InTrajectoryPredictionCount = 8);
+
+	// Experimental: Process InTrajectory to apply gravity and handle collisions. Eventually returns the modified OutTrajectory.
+	// If bApplyGravity is true, gravity from the UCharacterMovementComponent will be applied.
+	// If FloorCollisionsOffset > 0, vertical collision will be performed to every sample of the trajectory to have the samples float over the geometry (by FloorCollisionsOffset).
+	UE_EXPERIMENTAL(5.6, "Solution for CMC based workflow not Mover.")
+	UFUNCTION(BlueprintCallable, Category="Animation|PoseSearch|Experimental", meta=(BlueprintThreadSafe, DisplayName="HandleTrajectoryWorldCollisions", WorldContext="WorldContextObject", AutoCreateRefTerm="ActorsToIgnore", AdvancedDisplay="TraceChannel,bTraceComplex,ActorsToIgnore,DrawDebugType,bIgnoreSelf,MaxObstacleHeight,TraceColor,TraceHitColor,DrawTime"))
+	static void HandleTransformTrajectoryWorldCollisions(const UObject* WorldContextObject, const UAnimInstance* AnimInstance, UPARAM(ref) const FTransformTrajectory& InTrajectory, bool bApplyGravity, float FloorCollisionsOffset, FTransformTrajectory& OutTrajectory, FPoseSearchTrajectory_WorldCollisionResults& CollisionResult,
+		ETraceTypeQuery TraceChannel, bool bTraceComplex, const TArray<AActor*>& ActorsToIgnore, EDrawDebugTrace::Type DrawDebugType, bool bIgnoreSelf = true, float MaxObstacleHeight = 10000.f, FLinearColor TraceColor = FLinearColor::Red, FLinearColor TraceHitColor = FLinearColor::Green, float DrawTime = 5.0f);
+
+	// Experimental: Process InTrajectory to apply gravity and handle collisions. Eventually returns the modified OutTrajectory.
+	// If bApplyGravity is true, GravityAccel will be applied.
+	// If FloorCollisionsOffset > 0, vertical collision will be performed to every sample of the trajectory to have the samples float over the geometry (by FloorCollisionsOffset).
+	UE_EXPERIMENTAL(5.6, "Solution for CMC based workflow not Mover.")
+	UFUNCTION(BlueprintCallable, Category="Animation|PoseSearch|Experimental", meta=(BlueprintThreadSafe, DisplayName="HandleTrajectoryWorldCollisionsWithGravity", WorldContext="WorldContextObject", AutoCreateRefTerm="ActorsToIgnore", AdvancedDisplay="TraceChannel,bTraceComplex,ActorsToIgnore,DrawDebugType,bIgnoreSelf,MaxObstacleHeight,TraceColor,TraceHitColor,DrawTime"))
+	static void HandleTransformTrajectoryWorldCollisionsWithGravity(const UObject* WorldContextObject, 
+		UPARAM(ref) const FTransformTrajectory& InTrajectory, FVector StartingVelocity, bool bApplyGravity, FVector GravityAccel, float FloorCollisionsOffset, FTransformTrajectory& OutTrajectory, FPoseSearchTrajectory_WorldCollisionResults& CollisionResult,
+		ETraceTypeQuery TraceChannel, bool bTraceComplex, const TArray<AActor*>& ActorsToIgnore, EDrawDebugTrace::Type DrawDebugType, bool bIgnoreSelf = true, float MaxObstacleHeight = 10000.f, FLinearColor TraceColor = FLinearColor::Red, FLinearColor TraceHitColor = FLinearColor::Green, float DrawTime = 5.0f);
+
+	UFUNCTION(BlueprintPure, Category="Animation|PoseSearch", meta = (BlueprintThreadSafe, DisplayName="GetTrajectorySampleAtTime"))
+	static void GetTransformTrajectorySampleAtTime(UPARAM(ref) const FTransformTrajectory& InTrajectory, float Time, FTransformTrajectorySample& OutTrajectorySample, bool bExtrapolate = false);
+	
+	UFUNCTION(BlueprintPure, Category="Animation|PoseSearch", meta = (BlueprintThreadSafe, DisplayName="GetTrajectoryVelocity"))
+	static void GetTransformTrajectoryVelocity(UPARAM(ref) const FTransformTrajectory& InTrajectory, float Time1, float Time2, FVector& OutVelocity, bool bExtrapolate = false);
+	
+	UFUNCTION(BlueprintPure, Category="Animation|PoseSearch", meta = (BlueprintThreadSafe, DisplayName="GetTrajectoryAngularVelocity"))
+	static void GetTransformTrajectoryAngularVelocity(UPARAM(ref) const FTransformTrajectory& InTrajectory, float Time1, float Time2, FVector& OutAngularVelocity, bool bExtrapolate = false);
+	
+	UFUNCTION(BlueprintPure, Category="Animation|PoseSearch", meta = (BlueprintThreadSafe, DisplayName="GetTrajectorySampleTransform"))
+	static FTransform GetTransformTrajectorySampleTransform(UPARAM(ref) const FTransformTrajectorySample& InTrajectorySample);
+	
+	UFUNCTION(BlueprintCallable, Category="Animation|PoseSearch|Experimental", meta=(DisplayName="DrawTrajectory", WorldContext="WorldContextObject"))
+	static void DrawTransformTrajectory(const UObject* WorldContextObject, UPARAM(ref) const FTransformTrajectory& InTrajectory, const float DebugThickness, float HeightOffset);
+	
 private:
+	
 	static FVector RemapVectorMagnitudeWithCurve(const FVector& Vector, bool bUseCurve, const FRuntimeFloatCurve& Curve);
 };
\ No newline at end of file
diff --git a/Engine/Plugins/Animation/PoseSearch/Source/Runtime/Public/PoseSearch/PoseSearchTrajectoryPredictor.h b/Engine/Plugins/Animation/PoseSearch/Source/Runtime/Public/PoseSearch/PoseSearchTrajectoryPredictor.h
--- a/Engine/Plugins/Animation/PoseSearch/Source/Runtime/Public/PoseSearch/PoseSearchTrajectoryPredictor.h
+++ b/Engine/Plugins/Animation/PoseSearch/Source/Runtime/Public/PoseSearch/PoseSearchTrajectoryPredictor.h
@@ -5,7 +5,7 @@
 #include "UObject/Interface.h"
 #include "PoseSearchTrajectoryPredictor.generated.h"
 
-struct FPoseSearchQueryTrajectory;
+struct FTransformTrajectory;
 
 
 UINTERFACE(BlueprintType, NotBlueprintable, Experimental)
@@ -22,9 +22,9 @@ class IPoseSearchTrajectoryPredictorInterface : public IInterface
 	GENERATED_BODY()
 
 public:
-
+	
 	UFUNCTION(BlueprintCallable, Category = "Animation|PoseSearch|Experimental")
-	virtual void Predict(FPoseSearchQueryTrajectory& InOutTrajectory, int32 NumPredictionSamples, float SecondsPerPredictionSample, int32 NumHistorySamples) = 0;
+	virtual void Predict(FTransformTrajectory& InOutTrajectory, int32 NumPredictionSamples, float SecondsPerPredictionSample, int32 NumHistorySamples) = 0;
 
 	UFUNCTION(BlueprintCallable, Category="Animation|PoseSearch|Experimental")
 	virtual void GetGravity(FVector& OutGravityAccel) = 0;
diff --git a/Engine/Plugins/Animation/PoseSearch/Source/Runtime/Public/PoseSearch/PoseSearchTrajectoryTypes.h b/Engine/Plugins/Animation/PoseSearch/Source/Runtime/Public/PoseSearch/PoseSearchTrajectoryTypes.h
--- a/Engine/Plugins/Animation/PoseSearch/Source/Runtime/Public/PoseSearch/PoseSearchTrajectoryTypes.h
+++ b/Engine/Plugins/Animation/PoseSearch/Source/Runtime/Public/PoseSearch/PoseSearchTrajectoryTypes.h
@@ -3,47 +3,92 @@
 #pragma once
 
 #include "Animation/AnimTypes.h"
-#include "Animation/MotionTrajectoryTypes.h"
-
 #include "PoseSearchTrajectoryTypes.generated.h"
 
 struct FAnimInstanceProxy;
+struct FTransformTrajectory;
+struct FTransformTrajectorySample;
+
 
 USTRUCT(BlueprintType, Category="Pose Search Trajectory")
-struct POSESEARCH_API FPoseSearchQueryTrajectorySample
+struct UE_DEPRECATED(5.6, "FPoseSearchQueryTrajectorySample has been deprecated. Use FTransformTrajectorySample instead.")
+POSESEARCH_API FPoseSearchQueryTrajectorySample 
 {
 	GENERATED_BODY()
 
-	UPROPERTY(EditAnywhere, BlueprintReadWrite, Category = "Pose Search Query Trajectory")
+	FPoseSearchQueryTrajectorySample(const FQuat& InFacing, const FVector& InPosition, float InAccumulatedSeconds);
+
+	PRAGMA_DISABLE_DEPRECATION_WARNINGS
+	// Note: We need to explicitly disable warnings on these constructors/operators for clang to be happy with deprecated variables
+	FPoseSearchQueryTrajectorySample() = default;
+	~FPoseSearchQueryTrajectorySample() = default;
+	FPoseSearchQueryTrajectorySample(const FPoseSearchQueryTrajectorySample&) = default;
+	FPoseSearchQueryTrajectorySample(FPoseSearchQueryTrajectorySample&&) = default;
+	FPoseSearchQueryTrajectorySample& operator=(const FPoseSearchQueryTrajectorySample&) = default;
+	FPoseSearchQueryTrajectorySample& operator=(FPoseSearchQueryTrajectorySample&&) = default;
+	PRAGMA_ENABLE_DEPRECATION_WARNINGS
+	
+	UE_DEPRECATED(5.6, "Use FTransfromTrajectorySample instead. See FTransfromTrajectorySample::Facing.")
+	UPROPERTY(meta = (DeprecatedProperty, DeprecationMessage = "Use FTransfromTrajectorySample instead. See FTransfromTrajectorySample::Facing."))
 	FQuat Facing = FQuat::Identity;
 
-	UPROPERTY(EditAnywhere, BlueprintReadWrite, Category = "Pose Search Query Trajectory")
+	UE_DEPRECATED(5.6, "Use FTransfromTrajectorySample instead. See FTransfromTrajectorySample::Position.")
+	UPROPERTY(meta = (DeprecatedProperty, DeprecationMessage = "Use FTransfromTrajectorySample instead. See FTransfromTrajectorySample::Position."))
 	FVector Position = FVector::ZeroVector;
 
-	UPROPERTY(EditAnywhere, BlueprintReadWrite, Category = "Pose Search Query Trajectory")
+	UE_DEPRECATED(5.6, "Use FTransfromTrajectorySample instead. See FTransfromTrajectorySample::AccumulatedSeconds.")
+	UPROPERTY( meta = (DeprecatedProperty, DeprecationMessage = "Use FTransfromTrajectorySample instead. See FTransfromTrajectorySample::AccumulatedSeconds."))
 	float AccumulatedSeconds = 0.f;
 
+	PRAGMA_DISABLE_DEPRECATION_WARNINGS
 	FPoseSearchQueryTrajectorySample Lerp(const FPoseSearchQueryTrajectorySample& Other, float Alpha) const;
 	void SetTransform(const FTransform& Transform);
 	FTransform GetTransform() const { return FTransform(Facing, Position); }
+	PRAGMA_ENABLE_DEPRECATION_WARNINGS
 };
+
+PRAGMA_DISABLE_DEPRECATION_WARNINGS
 POSESEARCH_API FArchive& operator<<(FArchive& Ar, FPoseSearchQueryTrajectorySample& TrajectorySample);
+PRAGMA_ENABLE_DEPRECATION_WARNINGS
 
-USTRUCT(BlueprintType, Category = "Pose Search Trajectory")
-struct POSESEARCH_API FPoseSearchQueryTrajectory
+USTRUCT(BlueprintType, Category="Pose Search Trajectory")
+struct UE_DEPRECATED(5.6, "FPoseSearchQueryTrajectory has been deprecated. Use FTransformTrajectory instead.")
+POSESEARCH_API FPoseSearchQueryTrajectory
 {
 	GENERATED_BODY()
 
+	PRAGMA_DISABLE_DEPRECATION_WARNINGS
+	// Note: We need to explicitly disable warnings on these constructors/operators for clang to be happy with deprecated variables
+	FPoseSearchQueryTrajectory() = default;
+	~FPoseSearchQueryTrajectory() = default;
+	FPoseSearchQueryTrajectory(const FPoseSearchQueryTrajectory&) = default;
+	FPoseSearchQueryTrajectory(FPoseSearchQueryTrajectory&&) = default;
+	FPoseSearchQueryTrajectory& operator=(const FPoseSearchQueryTrajectory&) = default;
+	FPoseSearchQueryTrajectory& operator=(FPoseSearchQueryTrajectory&&) = default;
+	PRAGMA_ENABLE_DEPRECATION_WARNINGS
+
+	// Implicit conversion constructor in oreder to facilitate deprecations.
+	FPoseSearchQueryTrajectory(const FTransformTrajectory& InTrajectory);
+	
+	// Implicit conversion operator in order to facilitate deprecations.
+	operator FTransformTrajectory() const;
+	
+	PRAGMA_DISABLE_DEPRECATION_WARNINGS
 	// This contains zero or more history samples, a current sample, and zero or more future predicted samples.
-	UPROPERTY(EditAnywhere, BlueprintReadWrite, Category = "Pose Search Query Trajectory")
+	UE_DEPRECATED(5.6, "Use FTransformTrajectory instead. See FTransfromTrajectory::Samples.")
+	UPROPERTY(meta = (DeprecatedProperty, DeprecationMessage = "Use FTransfromTrajectory instead. See FTransfromTrajectory:Samples."))
 	TArray<FPoseSearchQueryTrajectorySample> Samples;
-
-	FPoseSearchQueryTrajectorySample GetSampleAtTime(float Time, bool bExtrapolate = false) const;
 	
+	FPoseSearchQueryTrajectorySample GetSampleAtTime(float Time, bool bExtrapolate = false) const;
+	PRAGMA_ENABLE_DEPRECATION_WARNINGS
+
 #if ENABLE_ANIM_DEBUG
 	void DebugDrawTrajectory(const UWorld* World, const float DebugThickness = 0.f, float HeightOffset = 0.f) const;
 	void DebugDrawTrajectory(FAnimInstanceProxy& AnimInstanceProxy, const float DebugThickness = 0.f, float HeightOffset = 0.f, int MaxHistorySamples = -1, int MaxPredictionSamples = -1) const;
 	void DebugDrawTrajectory(const UObject* Owner, const FLogCategoryBase& Category, ELogVerbosity::Type Verbosity, const float DebugThickness = 0.f, float HeightOffset = 0.f) const;
 #endif // ENABLE_ANIM_DEBUG
 };
+
+PRAGMA_DISABLE_DEPRECATION_WARNINGS
 POSESEARCH_API FArchive& operator<<(FArchive& Ar, FPoseSearchQueryTrajectory& Trajectory);
+PRAGMA_ENABLE_DEPRECATION_WARNINGS
\ No newline at end of file
diff --git a/Engine/Plugins/Animation/PoseSearch/Source/Runtime/Public/PoseSearch/Trace/PoseSearchTraceLogger.h b/Engine/Plugins/Animation/PoseSearch/Source/Runtime/Public/PoseSearch/Trace/PoseSearchTraceLogger.h
--- a/Engine/Plugins/Animation/PoseSearch/Source/Runtime/Public/PoseSearch/Trace/PoseSearchTraceLogger.h
+++ b/Engine/Plugins/Animation/PoseSearch/Source/Runtime/Public/PoseSearch/Trace/PoseSearchTraceLogger.h
@@ -4,7 +4,6 @@
 
 #include "PoseSearch/PoseSearchContext.h"
 #include "PoseSearch/PoseSearchHistory.h"
-#include "PoseSearch/PoseSearchTrajectoryTypes.h"
 
 UE_TRACE_CHANNEL_EXTERN(PoseSearchChannel, POSESEARCH_API);
 
diff --git a/Engine/Plugins/Experimental/AnimNextPoseSearch/Source/AnimNextPoseSearch/Private/HistoryCollectorTrait.h b/Engine/Plugins/Experimental/AnimNextPoseSearch/Source/AnimNextPoseSearch/Private/HistoryCollectorTrait.h
--- a/Engine/Plugins/Experimental/AnimNextPoseSearch/Source/AnimNextPoseSearch/Private/HistoryCollectorTrait.h
+++ b/Engine/Plugins/Experimental/AnimNextPoseSearch/Source/AnimNextPoseSearch/Private/HistoryCollectorTrait.h
@@ -5,7 +5,7 @@
 #include "CoreMinimal.h"
 #include "IPoseHistory.h"
 #include "PoseSearch/PoseSearchHistory.h"
-#include "PoseSearch/PoseSearchTrajectoryTypes.h"
+#include "Animation/TrajectoryTypes.h"
 #include "TraitCore/Trait.h"
 #include "TraitInterfaces/IEvaluate.h"
 #include "TraitInterfaces/IUpdate.h"
@@ -63,8 +63,8 @@ struct FAnimNextHistoryCollectorTraitSharedData : public FAnimNextTraitSharedDat
 
 	// input Trajectory samples for pose search queries in Motion Matching. These are expected to be in the world space of the SkeletalMeshComponent.
 	// the trajectory sample with AccumulatedSeconds equals to zero (Trajectory.Samples[i].AccumulatedSeconds) is the sample of the previous frame of simulation (since MM works by matching the previous character pose)
-	UPROPERTY(EditAnywhere, Category = Settings, meta = (PinShownByDefault, EditCondition="!bGenerateTrajectory", EditConditionHides))
-	FPoseSearchQueryTrajectory Trajectory;
+	UPROPERTY(EditAnywhere, Transient, Category = Settings, meta = (PinShownByDefault, EditCondition="!bGenerateTrajectory", EditConditionHides))
+	FTransformTrajectory Trajectory;
 
 	// Input Trajectory velocity will be multiplied by TrajectorySpeedMultiplier: values below 1 will result in selecting animation slower than requested from the original Trajectory
 	UPROPERTY(EditAnywhere, Category = Experimental, meta = (PinHiddenByDefault, ClampMin="0", EditCondition="!bGenerateTrajectory", EditConditionHides))
diff --git a/Engine/Plugins/Experimental/AnimNextPoseSearch/Source/AnimNextPoseSearch/Private/Module.cpp b/Engine/Plugins/Experimental/AnimNextPoseSearch/Source/AnimNextPoseSearch/Private/Module.cpp
--- a/Engine/Plugins/Experimental/AnimNextPoseSearch/Source/AnimNextPoseSearch/Private/Module.cpp
+++ b/Engine/Plugins/Experimental/AnimNextPoseSearch/Source/AnimNextPoseSearch/Private/Module.cpp
@@ -12,7 +12,7 @@
 #include "PoseSearch/PoseSearchDatabase.h"
 #include "PoseSearch/PoseSearchInteractionLibrary.h"
 #include "PoseSearch/PoseSearchSchema.h"
-#include "PoseSearch/PoseSearchTrajectoryTypes.h"
+#include "Animation/TrajectoryTypes.h"
 
 namespace UE::AnimNext::PoseSearch
 {
@@ -30,10 +30,10 @@ public:
 
 		static UScriptStruct* const AllowedStructTypes[] =
 		{
+			FTransformTrajectorySample::StaticStruct(),
+			FTransformTrajectory::StaticStruct(),
 			FAlphaBlendArgs::StaticStruct(),
 			FPoseSearchInteractionAvailability::StaticStruct(),
-			FPoseSearchQueryTrajectory::StaticStruct(),
-			FPoseSearchQueryTrajectorySample::StaticStruct(),
 		};
 
 		FRigVMRegistry& RigVMRegistry = FRigVMRegistry::Get();
diff --git a/Engine/Plugins/Experimental/AnimNextPoseSearch/Source/AnimNextPoseSearch/Private/RigUnit_DebugDrawTrajectory.cpp b/Engine/Plugins/Experimental/AnimNextPoseSearch/Source/AnimNextPoseSearch/Private/RigUnit_DebugDrawTrajectory.cpp
--- a/Engine/Plugins/Experimental/AnimNextPoseSearch/Source/AnimNextPoseSearch/Private/RigUnit_DebugDrawTrajectory.cpp
+++ b/Engine/Plugins/Experimental/AnimNextPoseSearch/Source/AnimNextPoseSearch/Private/RigUnit_DebugDrawTrajectory.cpp
@@ -1,7 +1,8 @@
 // Copyright Epic Games, Inc. All Rights Reserved.
 
 #include "RigUnit_DebugDrawTrajectory.h"
-#include "PoseSearch/PoseSearchTrajectoryLibrary.h"
+
+#include "Animation/TrajectoryTypes.h"
 #include "Component/AnimNextComponent.h"
 #include "PoseSearch/PoseSearchDefines.h"
 
@@ -20,6 +21,6 @@ FRigUnit_DebugDrawTrajectory_Execute()
 	// into a module component that can be used to access GT data where needed. 
 	UAnimNextComponent* Component = Cast<UAnimNextComponent>(ModuleInstance.GetObject());
 
-	Trajectory.DebugDrawTrajectory(Component, LogPoseSearch, ELogVerbosity::Display, DebugThickness, DebugOffset);
+	UTransformTrajectoryBlueprintLibrary::DebugDrawTrajectory(Trajectory, Component, LogPoseSearch, ELogVerbosity::Display, DebugThickness, DebugOffset);
 #endif
 }
diff --git a/Engine/Plugins/Experimental/AnimNextPoseSearch/Source/AnimNextPoseSearch/Private/RigUnit_DebugDrawTrajectory.h b/Engine/Plugins/Experimental/AnimNextPoseSearch/Source/AnimNextPoseSearch/Private/RigUnit_DebugDrawTrajectory.h
--- a/Engine/Plugins/Experimental/AnimNextPoseSearch/Source/AnimNextPoseSearch/Private/RigUnit_DebugDrawTrajectory.h
+++ b/Engine/Plugins/Experimental/AnimNextPoseSearch/Source/AnimNextPoseSearch/Private/RigUnit_DebugDrawTrajectory.h
@@ -4,7 +4,7 @@
 
 #include "AnimNextExecuteContext.h"
 #include "Graph/RigUnit_AnimNextBase.h"
-#include "PoseSearch/PoseSearchTrajectoryTypes.h"
+#include "Animation/TrajectoryTypes.h"
 
 #include "RigUnit_DebugDrawTrajectory.generated.h"
 
@@ -17,7 +17,7 @@ struct FRigUnit_DebugDrawTrajectory: public FRigUnit_AnimNextBase
 	void Execute();
 
 	UPROPERTY(EditAnywhere, Category = "Graph", meta = (Input, Output))
-	FPoseSearchQueryTrajectory Trajectory;
+	FTransformTrajectory Trajectory;
 
 	UPROPERTY(EditAnywhere, Category = "Graph", meta = (Input))
 	float DebugThickness = 0.f;
diff --git a/Engine/Plugins/Experimental/Animation/MotionTrajectory/Source/MotionTrajectory/Private/CharacterTrajectoryComponent.cpp b/Engine/Plugins/Experimental/Animation/MotionTrajectory/Source/MotionTrajectory/Private/CharacterTrajectoryComponent.cpp
--- a/Engine/Plugins/Experimental/Animation/MotionTrajectory/Source/MotionTrajectory/Private/CharacterTrajectoryComponent.cpp
+++ b/Engine/Plugins/Experimental/Animation/MotionTrajectory/Source/MotionTrajectory/Private/CharacterTrajectoryComponent.cpp
@@ -6,6 +6,7 @@
 #include "GameFramework/Character.h"
 #include "HAL/IConsoleManager.h"
 #include "MotionTrajectory.h"
+#include "Animation/TrajectoryTypes.h"
 
 #if ENABLE_ANIM_DEBUG
 TAutoConsoleVariable<int32> CVarCharacterTrajectoryDebug(TEXT("a.CharacterTrajectory.Debug"), 0, TEXT("Turn on debug drawing for Character trajectory"));
@@ -102,7 +103,7 @@ void UCharacterTrajectoryComponent::OnMovementUpdated(float DeltaSeconds, FVecto
 #if ENABLE_ANIM_DEBUG
 	if (CVarCharacterTrajectoryDebug.GetValueOnAnyThread())
 	{
-		Trajectory.DebugDrawTrajectory(GetWorld());
+		UTransformTrajectoryBlueprintLibrary::DebugDrawTrajectory(Trajectory, GetWorld());
 	}
 #endif // ENABLE_ANIM_DEBUG
 }
\ No newline at end of file
diff --git a/Engine/Plugins/Experimental/Animation/MotionTrajectory/Source/MotionTrajectory/Private/MotionTrajectoryLibrary.cpp b/Engine/Plugins/Experimental/Animation/MotionTrajectory/Source/MotionTrajectory/Private/MotionTrajectoryLibrary.cpp
--- a/Engine/Plugins/Experimental/Animation/MotionTrajectory/Source/MotionTrajectory/Private/MotionTrajectoryLibrary.cpp
+++ b/Engine/Plugins/Experimental/Animation/MotionTrajectory/Source/MotionTrajectory/Private/MotionTrajectoryLibrary.cpp
@@ -151,7 +151,7 @@ FVector FCharacterTrajectoryData::StepCharacterMovementGroundPrediction(float De
 	return OutVelocity;
 }
 
-void FMotionTrajectoryLibrary::InitTrajectorySamples(FPoseSearchQueryTrajectory& Trajectory,
+void FMotionTrajectoryLibrary::InitTrajectorySamples(FTransformTrajectory& Trajectory,
 	const FTrajectorySamplingData& SamplingData, const FVector& Position, const FQuat& Facing)
 {
 	// History + current sample + prediction
@@ -162,7 +162,7 @@ void FMotionTrajectoryLibrary::InitTrajectorySamples(FPoseSearchQueryTrajectory&
 	{
 		Trajectory.Samples[i].Position = Position;
 		Trajectory.Samples[i].Facing = Facing;
-		Trajectory.Samples[i].AccumulatedSeconds = SamplingData.SecondsPerHistorySample * (i - SamplingData.NumHistorySamples);
+		Trajectory.Samples[i].TimeInSeconds = SamplingData.SecondsPerHistorySample * (i - SamplingData.NumHistorySamples);
 	}
 
 	// Initialize current sample and prediction
@@ -170,11 +170,11 @@ void FMotionTrajectoryLibrary::InitTrajectorySamples(FPoseSearchQueryTrajectory&
 	{
 		Trajectory.Samples[i].Position = Position;
 		Trajectory.Samples[i].Facing = Facing;
-		Trajectory.Samples[i].AccumulatedSeconds = SamplingData.SecondsPerPredictionSample * (i - SamplingData.NumHistorySamples);
+		Trajectory.Samples[i].TimeInSeconds = SamplingData.SecondsPerPredictionSample * (i - SamplingData.NumHistorySamples);
 	}
 }
 
-void FMotionTrajectoryLibrary::UpdateHistory_TransformHistory(FPoseSearchQueryTrajectory& Trajectory, TArrayView<FVector> TranslationHistory,
+void FMotionTrajectoryLibrary::UpdateHistory_TransformHistory(FTransformTrajectory& Trajectory, TArrayView<FVector> TranslationHistory,
 	const FCharacterTrajectoryData& CharacterTrajectoryData, const FTrajectorySamplingData& SamplingData, float DeltaSeconds)
 {
 	check(SamplingData.NumHistorySamples <= Trajectory.Samples.Num());
@@ -183,16 +183,16 @@ void FMotionTrajectoryLibrary::UpdateHistory_TransformHistory(FPoseSearchQueryTr
 	FVector CurrentTranslation = CharacterTrajectoryData.Velocity * DeltaSeconds;
 
 	// Shift history Samples when it's time to record a new one.
-	if (SamplingData.NumHistorySamples > 0 && FMath::Abs(Trajectory.Samples[SamplingData.NumHistorySamples - 1].AccumulatedSeconds) >= SamplingData.SecondsPerHistorySample)
+	if (SamplingData.NumHistorySamples > 0 && FMath::Abs(Trajectory.Samples[SamplingData.NumHistorySamples - 1].TimeInSeconds) >= SamplingData.SecondsPerHistorySample)
 	{
 		for (int32 Index = 0; Index < SamplingData.NumHistorySamples - 1; ++Index)
 		{
-			Trajectory.Samples[Index].AccumulatedSeconds = Trajectory.Samples[Index + 1].AccumulatedSeconds;
+			Trajectory.Samples[Index].TimeInSeconds = Trajectory.Samples[Index + 1].TimeInSeconds;
 			TranslationHistory[Index] = TranslationHistory[Index + 1] + CurrentTranslation;
 			Trajectory.Samples[Index].Facing = Trajectory.Samples[Index + 1].Facing;
 		}
 
-		Trajectory.Samples[SamplingData.NumHistorySamples - 1].AccumulatedSeconds = 0.f;
+		Trajectory.Samples[SamplingData.NumHistorySamples - 1].TimeInSeconds = 0.f;
 		TranslationHistory[SamplingData.NumHistorySamples - 1] = CurrentTranslation;
 		Trajectory.Samples[SamplingData.NumHistorySamples - 1].Facing = CharacterTrajectoryData.Facing;
 	}
@@ -207,7 +207,7 @@ void FMotionTrajectoryLibrary::UpdateHistory_TransformHistory(FPoseSearchQueryTr
 	// Update trajectory samples by applying the tracked translations to the current world position.
 	for (int32 Index = 0; Index < SamplingData.NumHistorySamples; ++Index)
 	{
-		Trajectory.Samples[Index].AccumulatedSeconds -= DeltaSeconds;
+		Trajectory.Samples[Index].TimeInSeconds -= DeltaSeconds;
 		Trajectory.Samples[Index].Position = CharacterTrajectoryData.Position - TranslationHistory[Index];
 	}
 }
@@ -227,7 +227,7 @@ FVector FMotionTrajectoryLibrary::RemapVectorMagnitudeWithCurve(const FVector& V
 	return Vector;
 }
 
-void FMotionTrajectoryLibrary::UpdatePrediction_SimulateCharacterMovement(FPoseSearchQueryTrajectory& Trajectory,
+void FMotionTrajectoryLibrary::UpdatePrediction_SimulateCharacterMovement(FTransformTrajectory& Trajectory,
 	const FCharacterTrajectoryData& CharacterTrajectoryData, const FTrajectorySamplingData& SamplingData)
 {
 	FVector CurrentPositionWS = CharacterTrajectoryData.Position;
@@ -275,7 +275,7 @@ void FMotionTrajectoryLibrary::UpdatePrediction_SimulateCharacterMovement(FPoseS
 		{
 			Trajectory.Samples[Index].Position = CurrentPositionWS;
 			Trajectory.Samples[Index].Facing = CurrentFacingWS;
-			Trajectory.Samples[Index].AccumulatedSeconds = AccumulatedSeconds;
+			Trajectory.Samples[Index].TimeInSeconds = AccumulatedSeconds;
 
 			if (Index == LastIndex)
 			{
diff --git a/Engine/Plugins/Experimental/Animation/MotionTrajectory/Source/MotionTrajectory/Public/CharacterTrajectoryComponent.h b/Engine/Plugins/Experimental/Animation/MotionTrajectory/Source/MotionTrajectory/Public/CharacterTrajectoryComponent.h
--- a/Engine/Plugins/Experimental/Animation/MotionTrajectory/Source/MotionTrajectory/Public/CharacterTrajectoryComponent.h
+++ b/Engine/Plugins/Experimental/Animation/MotionTrajectory/Source/MotionTrajectory/Public/CharacterTrajectoryComponent.h
@@ -4,7 +4,7 @@
 
 #include "Components/ActorComponent.h"
 #include "MotionTrajectoryLibrary.h"
-#include "PoseSearch/PoseSearchTrajectoryTypes.h"
+#include "Animation/TrajectoryTypes.h"
 #include "CharacterTrajectoryComponent.generated.h"
 
 // Component for generating trajectories usable by Motion Matching. This component generates trajectories from ACharacter.
@@ -31,7 +31,7 @@ protected:
 
 	// Trajectory stored in world space so it can be directly passed to Motion Matching.
 	UPROPERTY(Transient, BlueprintReadOnly, Category = "Trajectory")
-	FPoseSearchQueryTrajectory Trajectory;
+	FTransformTrajectory Trajectory;
 
 	UPROPERTY(Transient, BlueprintReadOnly, Category = "Trajectory Settings")
 	FTrajectorySamplingData SamplingData;
diff --git a/Engine/Plugins/Experimental/Animation/MotionTrajectory/Source/MotionTrajectory/Public/MotionTrajectoryLibrary.h b/Engine/Plugins/Experimental/Animation/MotionTrajectory/Source/MotionTrajectory/Public/MotionTrajectoryLibrary.h
--- a/Engine/Plugins/Experimental/Animation/MotionTrajectory/Source/MotionTrajectory/Public/MotionTrajectoryLibrary.h
+++ b/Engine/Plugins/Experimental/Animation/MotionTrajectory/Source/MotionTrajectory/Public/MotionTrajectoryLibrary.h
@@ -3,7 +3,7 @@
 #pragma once
 
 #include "Curves/CurveFloat.h"
-#include "PoseSearch/PoseSearchTrajectoryTypes.h"
+#include "Animation/TrajectoryTypes.h"
 #include "MotionTrajectoryLibrary.generated.h"
 
 class ACharacter;
@@ -103,17 +103,17 @@ public:
 struct MOTIONTRAJECTORY_API FMotionTrajectoryLibrary
 {
 public:
-	static void InitTrajectorySamples(FPoseSearchQueryTrajectory& Trajectory,
+	static void InitTrajectorySamples(FTransformTrajectory& Trajectory,
 		const FTrajectorySamplingData& SamplingData, const FVector& Position, const FQuat& Facing);
 
 	// Update history by tracking offsets that result from character intent (e.g. movement component velocity) and applying
 	// that to the current world transform. This works well on moving platforms as it only stores a history of movement
 	// that results from character intent, not movement from platforms.
-	static void UpdateHistory_TransformHistory(FPoseSearchQueryTrajectory& Trajectory, TArrayView<FVector> TranslationHistory,
+	static void UpdateHistory_TransformHistory(FTransformTrajectory& Trajectory, TArrayView<FVector> TranslationHistory,
 		const FCharacterTrajectoryData& CharacterTrajectoryData, const FTrajectorySamplingData& SamplingData, float DeltaSeconds);
 
 	// Update prediction by simulating the movement math for ground locomotion from UCharacterMovementComponent.
-	static void UpdatePrediction_SimulateCharacterMovement(FPoseSearchQueryTrajectory& Trajectory,
+	static void UpdatePrediction_SimulateCharacterMovement(FTransformTrajectory& Trajectory,
 		const FCharacterTrajectoryData& CharacterTrajectoryData, const FTrajectorySamplingData& SamplingData);
 
 private:
diff --git a/Engine/Plugins/Experimental/Mover/Source/Mover/Private/MoverPoseSearchTrajectoryPredictor.cpp b/Engine/Plugins/Experimental/Mover/Source/Mover/Private/MoverPoseSearchTrajectoryPredictor.cpp
--- a/Engine/Plugins/Experimental/Mover/Source/Mover/Private/MoverPoseSearchTrajectoryPredictor.cpp
+++ b/Engine/Plugins/Experimental/Mover/Source/Mover/Private/MoverPoseSearchTrajectoryPredictor.cpp
@@ -1,11 +1,11 @@
 // Copyright Epic Games, Inc. All Rights Reserved.
 
 #include "MoverPoseSearchTrajectoryPredictor.h"
- #include "PoseSearch/PoseSearchTrajectoryTypes.h"
+ #include "Animation/TrajectoryTypes.h"
  #include "MoverComponent.h"
  #include "MoveLibrary/MovementUtils.h"
  
-void UMoverTrajectoryPredictor::Predict(FPoseSearchQueryTrajectory& InOutTrajectory,
+void UMoverTrajectoryPredictor::Predict(FTransformTrajectory& InOutTrajectory,
 	int32 NumPredictionSamples, float SecondsPerPredictionSample, int NumHistorySamples)
 {
 	if (!MoverComponent)
@@ -17,7 +17,7 @@ void UMoverTrajectoryPredictor::Predict(FPoseSearchQueryTrajectory& InOutTraject
 	Predict(*MoverComponent, InOutTrajectory, NumPredictionSamples, SecondsPerPredictionSample, NumHistorySamples);
 }
 
-void UMoverTrajectoryPredictor::Predict(UMoverComponent& MoverComponent, FPoseSearchQueryTrajectory& InOutTrajectory, int32 NumPredictionSamples, float SecondsPerPredictionSample, int32 NumHistorySamples)
+void UMoverTrajectoryPredictor::Predict(UMoverComponent& MoverComponent, FTransformTrajectory& InOutTrajectory, int32 NumPredictionSamples, float SecondsPerPredictionSample, int32 NumHistorySamples)
 {
 	FMoverPredictTrajectoryParams PredictParams;
 	PredictParams.NumPredictionSamples = NumPredictionSamples;
@@ -42,7 +42,7 @@ void UMoverTrajectoryPredictor::Predict(UMoverComponent& MoverComponent, FPoseSe
 
 			InOutTrajectory.Samples[PoseSampleIdx].Position = MoverPredictionSamples[PredictedSampleIdx].Transform.GetLocation();
 			InOutTrajectory.Samples[PoseSampleIdx].Facing = MoverPredictionSamples[PredictedSampleIdx].Transform.GetRotation();
-			InOutTrajectory.Samples[PoseSampleIdx].AccumulatedSeconds = AccumulatedSeconds;
+			InOutTrajectory.Samples[PoseSampleIdx].TimeInSeconds = AccumulatedSeconds;
 
 			AccumulatedSeconds += SecondsPerPredictionSample;
 		}
diff --git a/Engine/Plugins/Experimental/Mover/Source/Mover/Public/MoverPoseSearchTrajectoryPredictor.h b/Engine/Plugins/Experimental/Mover/Source/Mover/Public/MoverPoseSearchTrajectoryPredictor.h
--- a/Engine/Plugins/Experimental/Mover/Source/Mover/Public/MoverPoseSearchTrajectoryPredictor.h
+++ b/Engine/Plugins/Experimental/Mover/Source/Mover/Public/MoverPoseSearchTrajectoryPredictor.h
@@ -7,6 +7,7 @@
 
 
 class UMoverComponent;
+struct FTransformTrajectory;
 
 /**
  * Trajectory predictor that can query from a Mover-driven actor, for use with Pose Search
@@ -21,7 +22,7 @@ public:
 	UFUNCTION(BlueprintCallable, Category = "Animation|PoseSearch|Experimental")
 	void Setup(UMoverComponent* InMoverComponent) { MoverComponent = InMoverComponent; }
 
-	virtual void Predict(FPoseSearchQueryTrajectory& InOutTrajectory, int32 NumPredictionSamples, float SecondsPerPredictionSample, int32 NumHistorySamples) override;
+	virtual void Predict(FTransformTrajectory& InOutTrajectory, int32 NumPredictionSamples, float SecondsPerPredictionSample, int32 NumHistorySamples) override;
 
 	virtual void GetGravity(FVector& OutGravityAccel) override;
 
@@ -30,7 +31,7 @@ public:
 	virtual void GetVelocity(FVector& OutVelocity) override;
 
 	// Share code in static helper functions so we can call it from AnimNext, which doesn't support interfaces
-	static void Predict(UMoverComponent& Mover, FPoseSearchQueryTrajectory& InOutTrajectory, int32 NumPredictionSamples, float SecondsPerPredictionSample, int32 NumHistorySamples);
+	static void Predict(UMoverComponent& Mover, FTransformTrajectory& InOutTrajectory, int32 NumPredictionSamples, float SecondsPerPredictionSample, int32 NumHistorySamples);
 	static void GetCurrentState(UMoverComponent& Mover, FVector& OutPosition, FQuat& OutFacing, FVector& OutVelocity);
 
 
diff --git a/Engine/Plugins/Experimental/MoverAnimNext/Source/MoverAnimNext/Private/Graph/RigUnit_GenerateMoverTrajectory.cpp b/Engine/Plugins/Experimental/MoverAnimNext/Source/MoverAnimNext/Private/Graph/RigUnit_GenerateMoverTrajectory.cpp
--- a/Engine/Plugins/Experimental/MoverAnimNext/Source/MoverAnimNext/Private/Graph/RigUnit_GenerateMoverTrajectory.cpp
+++ b/Engine/Plugins/Experimental/MoverAnimNext/Source/MoverAnimNext/Private/Graph/RigUnit_GenerateMoverTrajectory.cpp
@@ -5,6 +5,7 @@
 #include "Component/AnimNextComponent.h"
 #include "MoverComponent.h"
 #include "MoverPoseSearchTrajectoryPredictor.h"
+#include "PoseSearch/PoseSearchTrajectoryLibrary.h"
 
 FRigUnit_GenerateMoverTrajectory_Execute()
 {
diff --git a/Engine/Plugins/Experimental/MoverAnimNext/Source/MoverAnimNext/Private/Graph/RigUnit_GenerateMoverTrajectory.h b/Engine/Plugins/Experimental/MoverAnimNext/Source/MoverAnimNext/Private/Graph/RigUnit_GenerateMoverTrajectory.h
--- a/Engine/Plugins/Experimental/MoverAnimNext/Source/MoverAnimNext/Private/Graph/RigUnit_GenerateMoverTrajectory.h
+++ b/Engine/Plugins/Experimental/MoverAnimNext/Source/MoverAnimNext/Private/Graph/RigUnit_GenerateMoverTrajectory.h
@@ -4,7 +4,7 @@
 
 #include "AnimNextExecuteContext.h"
 #include "Graph/RigUnit_AnimNextBase.h"
-#include "PoseSearch/PoseSearchTrajectoryLibrary.h"
+#include "Animation/TrajectoryTypes.h"
 
 #include "RigUnit_GenerateMoverTrajectory.generated.h"
 
@@ -38,6 +38,5 @@ struct FRigUnit_GenerateMoverTrajectory : public FRigUnit_AnimNextBase
 	int32 NumPredictionSamples = 15;
 
 	UPROPERTY(EditAnywhere, Category = "Graph", meta = (Input, Output))
-	FPoseSearchQueryTrajectory InOutTrajectory;
-
+	FTransformTrajectory InOutTrajectory;
 };
diff --git a/Engine/Plugins/Experimental/MoverAnimNext/Source/MoverAnimNext/Private/Graph/RigUnit_TrajectoryAnalysis.cpp b/Engine/Plugins/Experimental/MoverAnimNext/Source/MoverAnimNext/Private/Graph/RigUnit_TrajectoryAnalysis.cpp
--- a/Engine/Plugins/Experimental/MoverAnimNext/Source/MoverAnimNext/Private/Graph/RigUnit_TrajectoryAnalysis.cpp
+++ b/Engine/Plugins/Experimental/MoverAnimNext/Source/MoverAnimNext/Private/Graph/RigUnit_TrajectoryAnalysis.cpp
@@ -1,18 +1,19 @@
 // Copyright Epic Games, Inc. All Rights Reserved.
 
 #include "RigUnit_TrajectoryAnalysis.h"
+#include "PoseSearch/PoseSearchTrajectoryLibrary.h"
 
 FRigUnit_GetTrajectorySampleAtTime_Execute()
 {
-	UPoseSearchTrajectoryLibrary::GetTrajectorySampleAtTime(InTrajectory, Time, OutTrajectorySample, bExtrapolate);
+	UPoseSearchTrajectoryLibrary::GetTransformTrajectorySampleAtTime(InTrajectory, Time, OutTrajectorySample, bExtrapolate);
 }
 
 FRigUnit_GetTrajectoryVelocity_Execute()
 {
-	UPoseSearchTrajectoryLibrary::GetTrajectoryVelocity(InTrajectory, Time1, Time2, OutVelocity, bExtrapolate);
+	UPoseSearchTrajectoryLibrary::GetTransformTrajectoryVelocity(InTrajectory, Time1, Time2, OutVelocity, bExtrapolate);
 }
 
 FRigUnit_GetTrajectoryAngularVelocity_Execute()
 {
-	UPoseSearchTrajectoryLibrary::GetTrajectoryAngularVelocity(InTrajectory, Time1, Time2, OutAngularVelocity, bExtrapolate);
+	UPoseSearchTrajectoryLibrary::GetTransformTrajectoryAngularVelocity(InTrajectory, Time1, Time2, OutAngularVelocity, bExtrapolate);
 }
\ No newline at end of file
diff --git a/Engine/Plugins/Experimental/MoverAnimNext/Source/MoverAnimNext/Private/Graph/RigUnit_TrajectoryAnalysis.h b/Engine/Plugins/Experimental/MoverAnimNext/Source/MoverAnimNext/Private/Graph/RigUnit_TrajectoryAnalysis.h
--- a/Engine/Plugins/Experimental/MoverAnimNext/Source/MoverAnimNext/Private/Graph/RigUnit_TrajectoryAnalysis.h
+++ b/Engine/Plugins/Experimental/MoverAnimNext/Source/MoverAnimNext/Private/Graph/RigUnit_TrajectoryAnalysis.h
@@ -4,7 +4,7 @@
 
 #include "AnimNextExecuteContext.h"
 #include "Graph/RigUnit_AnimNextBase.h"
-#include "PoseSearch/PoseSearchTrajectoryLibrary.h"
+#include "Animation/TrajectoryTypes.h"
 
 #include "RigUnit_TrajectoryAnalysis.generated.h"
 
@@ -17,10 +17,10 @@ struct FRigUnit_GetTrajectorySampleAtTime : public FRigUnit_AnimNextBase
 	void Execute();
 
 	UPROPERTY(EditAnywhere, Category = "Trajectory", meta = (Output))
-	FPoseSearchQueryTrajectorySample OutTrajectorySample;
+	FTransformTrajectorySample OutTrajectorySample;
 
 	UPROPERTY(EditAnywhere, Category = "Trajectory", meta = (Input))
-	FPoseSearchQueryTrajectory InTrajectory;
+	FTransformTrajectory InTrajectory;
 
 	UPROPERTY(EditAnywhere, Category = "Trajectory", meta = (Input))
 	float Time = 0.0f;
@@ -38,7 +38,7 @@ struct FRigUnit_GetTrajectoryVelocity : public FRigUnit_AnimNextBase
 	void Execute();
 
 	UPROPERTY(EditAnywhere, Category = "Trajectory", meta = (Input))
-	FPoseSearchQueryTrajectory InTrajectory;
+	FTransformTrajectory InTrajectory;
 
 	UPROPERTY(EditAnywhere, Category = "Trajectory", meta = (Output))
 	FVector OutVelocity = FVector::ZeroVector;
@@ -62,7 +62,7 @@ struct FRigUnit_GetTrajectoryAngularVelocity : public FRigUnit_AnimNextBase
 	void Execute();
 
 	UPROPERTY(EditAnywhere, Category = "Trajectory", meta = (Input))
-	FPoseSearchQueryTrajectory InTrajectory;
+	FTransformTrajectory InTrajectory;
 
 	UPROPERTY(EditAnywhere, Category = "Trajectory", meta = (Output))
 	FVector OutAngularVelocity = FVector::ZeroVector;
diff --git a/Engine/Plugins/Experimental/MoverAnimNext/Source/MoverAnimNext/Private/MoverAnimNextModule.cpp b/Engine/Plugins/Experimental/MoverAnimNext/Source/MoverAnimNext/Private/MoverAnimNextModule.cpp
--- a/Engine/Plugins/Experimental/MoverAnimNext/Source/MoverAnimNext/Private/MoverAnimNextModule.cpp
+++ b/Engine/Plugins/Experimental/MoverAnimNext/Source/MoverAnimNext/Private/MoverAnimNextModule.cpp
@@ -6,9 +6,10 @@
 #include "MoveLibrary/MovementUtils.h"
 #include "MoverComponent.h"
 #include "MoverPoseSearchTrajectoryPredictor.h"
-#include "PoseSearch/PoseSearchTrajectoryLibrary.h"
+#include "Animation/TrajectoryTypes.h"
 
 #include "Modules/ModuleManager.h"
+#include "PoseSearch/PoseSearchTrajectoryLibrary.h"
 
 
 void FMoverAnimNextModule::StartupModule()
@@ -18,7 +19,7 @@ void FMoverAnimNextModule::StartupModule()
 		FMoverPredictTrajectoryParams::StaticStruct(),	// Todo: remove
 		FTrajectorySampleInfo::StaticStruct(),			// Todo: remove
 		FPoseSearchTrajectoryData::StaticStruct(),
-		FPoseSearchQueryTrajectory::StaticStruct(),
+		FTransformTrajectory::StaticStruct(),
 
 	};
 
diff --git a/Engine/Source/Editor/AnimGraph/Private/AnimGraphNode_Base.cpp b/Engine/Source/Editor/AnimGraph/Private/AnimGraphNode_Base.cpp
--- a/Engine/Source/Editor/AnimGraph/Private/AnimGraphNode_Base.cpp
+++ b/Engine/Source/Editor/AnimGraph/Private/AnimGraphNode_Base.cpp
@@ -934,6 +934,43 @@ bool UAnimGraphNode_Base::HasBinding(FName InBindingName) const
 	return false;
 }
 
+bool UAnimGraphNode_Base::RedirectBinding(FName OldPropertyBindingName, FName NewPropertyBindingName)
+{
+	bool bRenamed = false;
+
+	if (Binding)
+	{
+		const FString OldPropertyBindingNameStr = OldPropertyBindingName.ToString();
+		const FString NewPropertyBindingNameStr = NewPropertyBindingName.ToString();
+		
+		Binding->UpdateBindingNames([this, &bRenamed, &OldPropertyBindingNameStr, &NewPropertyBindingNameStr](const FString& InBindingName) -> FString
+		{
+			// Ignore any property binding that does not match exactly with what we requested.
+			if (InBindingName != OldPropertyBindingNameStr)
+			{
+				return InBindingName;
+			}
+
+			// We found our target binding name, replace it!
+			bRenamed = true;
+			return NewPropertyBindingNameStr;
+		});
+
+		if (bRenamed)
+		{
+			// Remove old binding name if it still exists (UpdateBindingNames() doesn't remove old binding entry).
+			if (HasBinding(OldPropertyBindingName))
+			{
+				Binding->RemoveBindings(OldPropertyBindingName);
+			}
+			
+			Binding->Modify();
+		}
+	}
+
+	return bRenamed;
+}
+
 void UAnimGraphNode_Base::RemoveBindings(FName InBindingName)
 {
 	if (Binding)
diff --git a/Engine/Source/Editor/AnimGraph/Public/AnimGraphNode_Base.h b/Engine/Source/Editor/AnimGraph/Public/AnimGraphNode_Base.h
--- a/Engine/Source/Editor/AnimGraph/Public/AnimGraphNode_Base.h
+++ b/Engine/Source/Editor/AnimGraph/Public/AnimGraphNode_Base.h
@@ -553,6 +553,9 @@ public:
 	// Check whether the specified property is bound via PropertyBindings
 	virtual bool HasBinding(FName InPropertyName) const;
 
+	// Update the property binding name
+	bool RedirectBinding(FName OldPropertyBindingName, FName NewPropertyBindingName);
+	
 	// Get the bindings for this node
 	const UAnimGraphNodeBinding* GetBinding() const { return Binding; }
 
diff --git a/Engine/Source/Runtime/Engine/Private/Animation/MotionTrajectoryTypes.cpp b/Engine/Source/Runtime/Engine/Private/Animation/MotionTrajectoryTypes.cpp
--- a/Engine/Source/Runtime/Engine/Private/Animation/MotionTrajectoryTypes.cpp
+++ b/Engine/Source/Runtime/Engine/Private/Animation/MotionTrajectoryTypes.cpp
@@ -4,6 +4,7 @@
 #include "Algo/AllOf.h"
 #include "Animation/AnimTypes.h"
 #include "DrawDebugHelpers.h"
+#include "Animation/AnimInstanceProxy.h"
 #include "Animation/AnimTypes.h"
 #include "HAL/IConsoleManager.h"
 #include "Misc/StringFormatArg.h"
@@ -62,6 +63,8 @@ bool FTrajectorySample::IsZeroSample() const
 		&& Transform.GetRotation().IsIdentity();
 }
 
+PRAGMA_DISABLE_DEPRECATION_WARNINGS
+
 FTrajectorySample FTrajectorySample::Lerp(const FTrajectorySample& Sample, float Alpha) const
 {
 	FTrajectorySample Interp;
@@ -119,169 +122,4 @@ void FTrajectorySample::TransformReferenceFrame(const FTransform DeltaTransform)
 	LinearVelocity = DeltaTransform.TransformVectorNoScale(LinearVelocity);
 }
 
-bool FTrajectorySampleRange::HasSamples() const
-{
-	return !Samples.IsEmpty();
-}
-
-bool FTrajectorySampleRange::HasOnlyZeroSamples() const
-{
-	return Algo::AllOf(Samples, [](const FTrajectorySample& Sample)
-		{
-			return Sample.IsZeroSample();
-		});
-}
-
-void FTrajectorySampleRange::RemoveHistory()
-{
-	Samples.RemoveAll([](const FTrajectorySample& Sample)
-		{
-			return Sample.AccumulatedSeconds < 0.f;
-		});
-}
-
-void FTrajectorySampleRange::Rotate(const FQuat& Rotation)
-{
-	for (auto& Sample : Samples)
-	{
-		Sample.PrependOffset(FTransform(Rotation), 0.0f);
-	}
-}
-
-void FTrajectorySampleRange::TransformReferenceFrame(const FTransform& Transform)
-{
-	for (auto& Sample : Samples)
-	{
-		Sample.TransformReferenceFrame(Transform);
-	}
-}
-
-FTrajectorySample FTrajectorySampleRange::GetSampleAtTime(float Time, bool bExtrapolate) const
-{
-	const int32 Num = Samples.Num();
-	if (Num > 1)
-	{
-		const int32 LowerBoundIdx = Algo::LowerBound(Samples, Time, [](const FTrajectorySample& TrajectorySample, float Value)
-			{
-				return Value > TrajectorySample.AccumulatedSeconds;
-			});
-
-		const int32 NextIdx = FMath::Clamp(LowerBoundIdx, 1, Samples.Num() - 1);
-		const int32 PrevIdx = NextIdx - 1;
-
-		const float Denominator = Samples[NextIdx].AccumulatedSeconds - Samples[PrevIdx].AccumulatedSeconds;
-		if (!FMath::IsNearlyZero(Denominator))
-		{
-			const float Numerator = Time - Samples[PrevIdx].AccumulatedSeconds;
-			const float LerpValue = bExtrapolate ? Numerator / Denominator : FMath::Clamp(Numerator / Denominator, 0.f, 1.f);
-			return Samples[PrevIdx].Lerp(Samples[NextIdx], LerpValue);
-		}
-
-		return Samples[PrevIdx];
-	}
-
-	if (Num > 0)
-	{
-		return Samples[0];
-	}
-
-	return FTrajectorySample();
-}
-
-void FTrajectorySampleRange::DebugDrawTrajectory(bool bEnable
-	, const UWorld* World
-	, const FTransform& WorldTransform
-	, const FLinearColor PredictionColor
-	, const FLinearColor HistoryColor
-	, float TransformScale
-	, float TransformThickness
-	, float VelArrowScale
-	, float VelArrowSize
-	, float VelArrowThickness) const
-{
-	if (bEnable
-#if ENABLE_ANIM_DEBUG
-		|| CVarMotionTrajectoryDebug.GetValueOnAnyThread()
-#endif
-		)
-	{
-		if (World)
-		{
-#if ENABLE_ANIM_DEBUG
-			const int32 DebugSampleStride = CVarMotionTrajectoryDebugStride.GetValueOnAnyThread();
-			const int32 DebugSampleOptions = CVarMotionTrajectoryDebugOptions.GetValueOnAnyThread();
-#endif
-			for (int32 Idx = 0, Num = Samples.Num(); Idx < Num; Idx++)
-			{
-				const FTransform SampleTransformWS = Samples[Idx].Transform * WorldTransform;
-				const FVector SamplePositionWS = SampleTransformWS.GetTranslation();
-
-				const FVector WorldVelocity =
-					(WorldTransform.TransformVector(Samples[Idx].LinearVelocity) * VelArrowScale) + SamplePositionWS;
-
-				// Interpolate the history and prediction color over the entire trajectory range
-				const float ColorLerp = static_cast<float>(Idx) / static_cast<float>(Num);
-				const FLinearColor Color = FLinearColor::LerpUsingHSV(PredictionColor, HistoryColor, ColorLerp);
-
-				if (VelArrowScale > 0.0f)
-				{
-					DrawDebugDirectionalArrow(
-						World, SamplePositionWS, WorldVelocity, 
-						VelArrowSize, Color.ToFColor(true), false, 0.f, 0, VelArrowThickness);
-				}
-
-				if (TransformScale > 0.0f)
-				{
-					DrawDebugCoordinateSystem(
-						World, SampleTransformWS.GetLocation(), SampleTransformWS.Rotator(),
-						TransformScale, false, 0.f, 0, TransformThickness);
-
-					if (VelArrowScale == 0.0f)
-					{
-						DrawDebugSphere(
-							World, SamplePositionWS,										
-							TransformScale * 0.5f, 4, Color.ToFColor(true), false, 0.f, 0, TransformThickness * 0.5f);
-					}
-				}
-
-#if ENABLE_ANIM_DEBUG
-				FString DebugString;
-				FString DebugSampleString;
-				switch (DebugSampleOptions)
-				{
-				case DebugTrajectorySampleCount: // Sample Index
-					DebugString = "Sample Index:";
-					DebugSampleString = DebugSampleString.Format(TEXT("{0}"), { Idx });
-;					break;
-				case DebugTrajectorySampleTime: // Sample Accumulated Time
-					DebugString = "Sample Time:";
-					DebugSampleString = DebugSampleString.Format(TEXT("{0}"), { Samples[Idx].AccumulatedSeconds });
-					break;
-				case DebugTrajectorySamplePosition: // Sample Position
-					DebugString = "Sample Position:";
-					DebugSampleString = DebugSampleString.Format(TEXT("{0}"), { Samples[Idx].Transform.GetLocation().ToCompactString() });
-					break;
-				case DebugTrajectorySampleVelocity: // Sample Velocity
-					DebugString = "Sample Velocity:";
-					DebugSampleString = DebugSampleString.Format(TEXT("{0}"), { Samples[Idx].LinearVelocity.ToCompactString() });
-					break;
-				default:
-					break;
-				}
-
-				// Conditionally display per-sample information against a specified stride
-				if (!DebugSampleString.IsEmpty() && !!DebugSampleStride && (Idx % DebugSampleStride == 0))
-				{
-					// One time debug drawing of the per-sample type description
-					if (!DebugString.IsEmpty() && Idx == 0)
-					{
-						DrawDebugString(World, WorldTransform.GetLocation() + DebugSampleTypeOffset, DebugString, nullptr, FColor::White, 0.f, false, 1.f);
-					}
-
-					DrawDebugString(World, WorldVelocity + DebugSampleOffset, DebugSampleString, nullptr, FColor::White, 0.f, false, 1.f);
-				}
-#endif
-			}
-		}
-	}
-}
+PRAGMA_ENABLE_DEPRECATION_WARNINGS
\ No newline at end of file
diff --git a/Engine/Source/Runtime/Engine/Private/Animation/TrajectoryTypes.cpp b/Engine/Source/Runtime/Engine/Private/Animation/TrajectoryTypes.cpp
new file mode 100644
--- /dev/null
+++ b/Engine/Source/Runtime/Engine/Private/Animation/TrajectoryTypes.cpp
@@ -0,0 +1,218 @@
+﻿// Copyright Epic Games, Inc. All Rights Reserved.
+
+#include "Animation/TrajectoryTypes.h"
+
+#include "Animation/AnimInstanceProxy.h"
+#include "VisualLogger/VisualLogger.h"
+#include "VisualLogger/VisualLoggerTypes.h"
+#include "DrawDebugHelpers.h"
+
+#include UE_INLINE_GENERATED_CPP_BY_NAME(TrajectoryTypes)
+
+FTransformTrajectorySample FTransformTrajectorySample::Lerp(const FTransformTrajectorySample& Other, float Alpha) const
+{
+	check(Facing.IsNormalized());
+	check(Other.Facing.IsNormalized());
+
+	FTransformTrajectorySample Result;
+	
+	Result.Facing = FQuat::FastLerp(Facing, Other.Facing, Alpha).GetNormalized();
+	Result.Position = FMath::Lerp(Position, Other.Position, Alpha);
+	Result.TimeInSeconds = FMath::Lerp(TimeInSeconds, Other.TimeInSeconds, Alpha);
+
+	return Result;
+}
+
+void FTransformTrajectorySample::SetTransform(const FTransform& Transform)
+{
+	Position = Transform.GetTranslation();
+	Facing = Transform.GetRotation();
+}
+
+FArchive& operator<<(FArchive& Ar, FTransformTrajectorySample& Sample)
+{
+	Ar << Sample.Facing;
+	Ar << Sample.Position;
+	Ar << Sample.TimeInSeconds;
+	return Ar;
+}
+
+FTransformTrajectorySample FTransformTrajectory::GetSampleAtTime(float Time, bool bExtrapolate) const
+{
+	const int32 Num = Samples.Num();
+	if (Num > 1)
+	{
+		const int32 LowerBoundIdx = Algo::LowerBound(Samples, Time, [](const FTransformTrajectorySample& TrajectorySample, float Value)
+			{
+				return Value > TrajectorySample.TimeInSeconds;
+			});
+
+		const int32 NextIdx = FMath::Clamp(LowerBoundIdx, 1, Samples.Num() - 1);
+		const int32 PrevIdx = NextIdx - 1;
+
+		const float Denominator = Samples[NextIdx].TimeInSeconds - Samples[PrevIdx].TimeInSeconds;
+		if (!FMath::IsNearlyZero(Denominator))
+		{
+			const float Numerator = Time - Samples[PrevIdx].TimeInSeconds;
+			const float LerpValue = bExtrapolate ? Numerator / Denominator : FMath::Clamp(Numerator / Denominator, 0.f, 1.f);
+			return Samples[PrevIdx].Lerp(Samples[NextIdx], LerpValue);
+		}
+
+		return Samples[PrevIdx];
+	}
+
+	if (Num > 0)
+	{
+		return Samples[0];
+	}
+
+	return FTransformTrajectorySample();
+}
+
+void UTransformTrajectoryBlueprintLibrary::DebugDrawTrajectory(UPARAM(ref) const FTransformTrajectory& Trajectory, const UWorld* World, const float DebugThickness, float HeightOffset)
+{
+#if ENABLE_DRAW_DEBUG
+	const FVector OffsetVector = FVector::UpVector * HeightOffset;
+
+	const int32 LastIndex = Trajectory.Samples.Num() - 1;
+	if (LastIndex >= 0)
+	{
+		for (int32 Index = 0; ; ++Index)
+		{
+			const FVector Pos = Trajectory.Samples[Index].Position + OffsetVector;
+
+			DrawDebugSphere(World, Pos, 1.f, 4, FColor::Black, false, -1.f, SDPG_Foreground, DebugThickness);
+
+			const FRotationMatrix R(FRotator(Trajectory.Samples[Index].Facing));
+			const FVector X = R.GetScaledAxis( EAxis::X );
+			const FVector Y = R.GetScaledAxis( EAxis::Y );
+
+			const float Scale = 12.f;
+
+			const bool IsPast = Trajectory.Samples[Index].TimeInSeconds <= 0.f;
+
+			DrawDebugLine(World, Pos, Pos + X * Scale, IsPast ? FColor::Red : FColor::Blue, false, -1.f, SDPG_Foreground, DebugThickness);
+			DrawDebugLine(World, Pos, Pos + Y * Scale, IsPast ? FColor::Orange : FColor::Turquoise, false, -1.f, SDPG_Foreground, DebugThickness);
+
+			if (Index == LastIndex)
+			{
+				break;
+			}
+			
+			const FVector NextPos = Trajectory.Samples[Index + 1].Position + OffsetVector;
+			DrawDebugLine(World, Pos, NextPos, FColor::Black, false, -1.f, SDPG_Foreground, DebugThickness);
+		}
+	}
+#endif
+}
+
+#if ENABLE_ANIM_DEBUG
+
+void UTransformTrajectoryBlueprintLibrary::DebugDrawTrajectory(const FTransformTrajectory& Trajectory, FAnimInstanceProxy& AnimInstanceProxy, const float DebugThickness, float HeightOffset, int MaxHistorySamples, int MaxPredictionSamples)
+{
+	const FVector OffsetVector = FVector::UpVector * HeightOffset;
+	
+	int AvailableHistorySamplesCount = 0;
+	
+	if (MaxHistorySamples != -1 || MaxPredictionSamples != -1)
+	{
+		for (int32 i = 0; i < Trajectory.Samples.Num(); ++i)
+		{
+			if (Trajectory.Samples[i].TimeInSeconds <= 0)
+			{
+				++AvailableHistorySamplesCount;
+			}
+			else
+			{
+				break;
+			}
+		}
+	}
+	
+	const int32 LastIndex = Trajectory.Samples.Num() - 1;
+	const int32 ZeroTimeIndex = AvailableHistorySamplesCount - 1;
+	const int32 StartIndex = MaxHistorySamples < 0 ? 0 : FMath::Max(AvailableHistorySamplesCount - MaxHistorySamples, 0); 
+	const int32 EndIndex = MaxPredictionSamples < 0 ? LastIndex : FMath::Min(ZeroTimeIndex + MaxPredictionSamples, LastIndex);
+	
+	if (LastIndex >= 0 && StartIndex <= EndIndex)
+	{
+		for (int32 Index = StartIndex; ; ++Index)
+		{
+			const FVector Pos = Trajectory.Samples[Index].Position + OffsetVector;
+
+			AnimInstanceProxy.AnimDrawDebugSphere(Trajectory.Samples[Index].Position + OffsetVector, 1.f, 4, FColor::Black, false, -1.f, DebugThickness, SDPG_Foreground);
+
+			const FRotationMatrix R(FRotator(Trajectory.Samples[Index].Facing));
+			const FVector X = R.GetScaledAxis( EAxis::X );
+			const FVector Y = R.GetScaledAxis( EAxis::Y );
+
+			const float Scale = 12.f;
+
+			const bool IsPast = Trajectory.Samples[Index].TimeInSeconds <= 0.f;
+
+			AnimInstanceProxy.AnimDrawDebugLine(Pos, Pos + X * Scale, IsPast ? FColor::Red : FColor::Blue, false, -1.f, DebugThickness, SDPG_Foreground);
+			AnimInstanceProxy.AnimDrawDebugLine(Pos, Pos + Y * Scale, IsPast ? FColor::Orange : FColor::Turquoise, false, -1.f, DebugThickness, SDPG_Foreground);
+
+			if (Index == EndIndex)
+			{
+				break;
+			}
+			
+			const FVector NextPos = Trajectory.Samples[Index + 1].Position + OffsetVector;
+			AnimInstanceProxy.AnimDrawDebugLine(Pos, NextPos, FColor::Black, false, -1.f, DebugThickness, SDPG_Foreground);
+		}
+	}
+}
+
+void UTransformTrajectoryBlueprintLibrary::DebugDrawTrajectory(const FTransformTrajectory& Trajectory, const UObject* Owner, const FLogCategoryBase& Category, ELogVerbosity::Type Verbosity, const float DebugThickness, float HeightOffset)
+{
+#if ENABLE_VISUAL_LOG
+
+	UWorld* World = nullptr;
+	FVisualLogEntry* CurrentEntry = nullptr;
+	if (FVisualLogger::CheckVisualLogInputInternal(Owner, Category.GetCategoryName(), Verbosity, &World, &CurrentEntry) == false)
+	{
+		return;
+	}
+
+	const FVector OffsetVector = FVector::UpVector * HeightOffset;
+
+	const int32 LastIndex = Trajectory.Samples.Num() - 1;
+	if (LastIndex >= 0)
+	{
+		for (int32 Index = 0; ; ++Index)
+		{
+			const FVector Pos = Trajectory.Samples[Index].Position + OffsetVector;
+
+			CurrentEntry->AddSphere(Pos, 1, Category.GetCategoryName(), Verbosity, FColor::Black, TEXT(""), true);
+
+			const FRotationMatrix R(FRotator(Trajectory.Samples[Index].Facing));
+			const FVector X = R.GetScaledAxis(EAxis::X);
+			const FVector Y = R.GetScaledAxis(EAxis::Y);
+
+			const float Scale = 12.f;
+
+			const bool IsPast = Trajectory.Samples[Index].TimeInSeconds <= 0.f;
+
+			CurrentEntry->AddSegment(Pos, Pos + X * Scale, Category.GetCategoryName(), Verbosity, IsPast ? FColor::Red : FColor::Blue, TEXT(""), DebugThickness);
+			CurrentEntry->AddSegment(Pos, Pos + Y * Scale, Category.GetCategoryName(), Verbosity, IsPast ? FColor::Orange : FColor::Turquoise, TEXT(""), DebugThickness);
+
+			if (Index == LastIndex)
+			{
+				break;
+			}
+
+			const FVector NextPos = Trajectory.Samples[Index + 1].Position + OffsetVector;
+			CurrentEntry->AddSegment(Pos, NextPos, Category.GetCategoryName(), Verbosity, FColor::Black, TEXT(""), DebugThickness);
+		}
+	}
+#endif // ENABLE_VISUAL_LOG
+}
+
+#endif // ENABLE_ANIM_DEBUG
+
+FArchive& operator<<(FArchive& Ar, FTransformTrajectory& Trajectory)
+{
+	Ar << Trajectory.Samples;
+	return Ar;
+}
diff --git a/Engine/Source/Runtime/Engine/Public/Animation/MotionTrajectoryTypes.h b/Engine/Source/Runtime/Engine/Public/Animation/MotionTrajectoryTypes.h
--- a/Engine/Source/Runtime/Engine/Public/Animation/MotionTrajectoryTypes.h
+++ b/Engine/Source/Runtime/Engine/Public/Animation/MotionTrajectoryTypes.h
@@ -7,7 +7,7 @@
 #include "MotionTrajectoryTypes.generated.h"
 
 USTRUCT(BlueprintType, Category="Motion Trajectory")
-struct FTrajectorySample
+struct UE_DEPRECATED(5.6, "FTrajectorySample has been deprecated. Instead use FTransformTrajectorySample.") FTrajectorySample
 {
 	GENERATED_BODY()
 
@@ -24,6 +24,7 @@ struct FTrajectorySample
 	UPROPERTY(EditAnywhere, BlueprintReadWrite, Category="Motion Trajectory")
 	FVector LinearVelocity = FVector::ZeroVector;
 
+	PRAGMA_DISABLE_DEPRECATION_WARNINGS
 	// Linear interpolation of all parameters of two trajectory samples
 	ENGINE_API FTrajectorySample Lerp(const FTrajectorySample& Sample, float Alpha) const;
 
@@ -32,7 +33,8 @@ struct FTrajectorySample
 		, const FTrajectorySample& Sample
 		, const FTrajectorySample& NextSample
 		, float Alpha) const;
-
+	PRAGMA_ENABLE_DEPRECATION_WARNINGS
+	
 	// Concatenates DeltaTransform before the current transform is applied and shifts the accumulated time by 
 	// DeltaSeconds
 	ENGINE_API void PrependOffset(const FTransform DeltaTransform, float DeltaSeconds);
@@ -41,47 +43,4 @@ struct FTrajectorySample
 
 	// Determines if all sample properties are zeroed
 	ENGINE_API bool IsZeroSample() const;
-};
-
-// A container of ordered trajectory samples and associated sampling rate
-USTRUCT(BlueprintType, Category="Motion Trajectory")
-struct UE_DEPRECATED(5.3, "Use FPoseSearchQueryTrajectory instead") ENGINE_API FTrajectorySampleRange
-{
-	GENERATED_BODY()
-
-	// Linearly ordered container for past, present, and future trajectory samples
-	UPROPERTY(EditAnywhere, BlueprintReadWrite, Category="Motion Trajectory")
-	TArray<FTrajectorySample> Samples;
-
-	// Removes history samples from trajectory (retains present and future)
-	void RemoveHistory();
-
-	// Rotates all samples in the trajectory
-	void Rotate(const FQuat& Rotation);
-
-	// Interpolates transform over time
-	void TransformOverTime(const FTransform& Transform, float StartTime, float DeltaTime);
-
-	// Rotates all samples in the trajectory
-	void TransformReferenceFrame(const FTransform& Transform);
-
-	// Determine if any trajectory samples are present
-	bool HasSamples() const;
-
-	// Determine if all trajectory samples are default values
-	bool HasOnlyZeroSamples() const;
-
-	FTrajectorySample GetSampleAtTime(float Time, bool bExtrapolate = true) const;
-
-	// Debug draw in-world trajectory samples and optional corresponding information
-	void DebugDrawTrajectory(bool bEnable
-		, const UWorld* World
-		, const FTransform& WorldTransform
-		, const FLinearColor PredictionColor = { 0.f, 1.f, 0.f }
-		, const FLinearColor HistoryColor = { 0.f, 0.f, 1.f }
-		, float TransformScale = 10.f
-		, float TransformThickness = 2.f
-		, float VelArrowScale = 0.025f
-		, float VelArrowSize = 40.f
-		, float VelArrowThickness = 2.f) const;
-};
+};
\ No newline at end of file
diff --git a/Engine/Source/Runtime/Engine/Public/Animation/TrajectoryTypes.h b/Engine/Source/Runtime/Engine/Public/Animation/TrajectoryTypes.h
new file mode 100644
--- /dev/null
+++ b/Engine/Source/Runtime/Engine/Public/Animation/TrajectoryTypes.h
@@ -0,0 +1,73 @@
+// Copyright Epic Games, Inc. All Rights Reserved.
+
+#pragma once
+
+#include "Animation/AnimTypes.h"
+#include "Kismet/BlueprintFunctionLibrary.h"
+
+#include "TrajectoryTypes.generated.h"
+
+struct FAnimInstanceProxy;
+
+/**
+ * A trajectory sample of conformed of a location, orientation, and time. *Scale is disregarded*.
+ */
+USTRUCT(BlueprintType, Category="Trajectory")
+struct ENGINE_API FTransformTrajectorySample
+{
+	GENERATED_BODY()
+
+	UPROPERTY(EditAnywhere, BlueprintReadWrite, Category = "Trajectory")
+	FQuat Facing = FQuat::Identity;
+
+	UPROPERTY(EditAnywhere, BlueprintReadWrite, Category = "Trajectory")
+	FVector Position = FVector::ZeroVector;
+
+	UPROPERTY(EditAnywhere, BlueprintReadWrite, Category = "Trajectory")
+	float TimeInSeconds = 0.f;
+
+	FTransformTrajectorySample Lerp(const FTransformTrajectorySample& Other, float Alpha) const;
+	void SetTransform(const FTransform& Transform);
+	FTransform GetTransform() const { return FTransform(Facing, Position); }
+};
+ENGINE_API FArchive& operator<<(FArchive& Ar, FTransformTrajectorySample& Sample);
+
+/**
+ * A trajectory defined by a range of samples.
+ *
+ * The sample range is conformed of history samples, a current sample and future predicted samples.
+ * - History samples have time < 0. Zero or more can be present.
+ * - Current sample has a time of 0. Assumes only one is ever present.
+ * - Predicted samples have a time of > 0. Zero or more can be present.
+ */
+USTRUCT(BlueprintType, Category = "Trajectory")
+struct ENGINE_API FTransformTrajectory
+{
+	GENERATED_BODY()
+	
+	// This contains zero or more history samples, a current sample, and zero or more future predicted samples.
+	UPROPERTY(EditAnywhere, BlueprintReadWrite, Category = "Trajectory")
+	TArray<FTransformTrajectorySample> Samples;
+
+	FTransformTrajectorySample GetSampleAtTime(float Time, bool bExtrapolate = false) const;
+};
+ENGINE_API FArchive& operator<<(FArchive& Ar, FTransformTrajectory& Trajectory);
+
+/**
+ * A function library of utilities for working with trajectories.
+ */
+UCLASS(MinimalAPI)
+class UTransformTrajectoryBlueprintLibrary : public UBlueprintFunctionLibrary
+{
+	GENERATED_BODY()
+
+public:
+	
+	UFUNCTION(BlueprintCallable, Category = "Animation|Trajectory|Debug", meta = (BlueprintThreadSafe))
+	static ENGINE_API void DebugDrawTrajectory(UPARAM(ref) const FTransformTrajectory& Trajectory, const UWorld* World, const float DebugThickness = 0.f, float HeightOffset = 0.f);
+
+#if ENABLE_ANIM_DEBUG
+	ENGINE_API static void DebugDrawTrajectory(const FTransformTrajectory& Trajectory, FAnimInstanceProxy& AnimInstanceProxy, const float DebugThickness = 0.f, float HeightOffset = 0.f, int MaxHistorySamples = -1, int MaxPredictionSamples = -1);
+	ENGINE_API static void DebugDrawTrajectory(const FTransformTrajectory& Trajectory, const UObject* Owner, const FLogCategoryBase& Category, ELogVerbosity::Type Verbosity, const float DebugThickness = 0.f, float HeightOffset = 0.f);
+#endif
+};
\ No newline at end of file
