Fixed mismatched PRAGMA_DISABLE_DEPRECATION_WARNINGS and other issues with the deprecation for PoseSearch's trajectory types.

#rnx

diff --git a/Engine/Plugins/Animation/PoseSearch/Source/Runtime/Private/PoseSearchTrajectoryLibrary.cpp b/Engine/Plugins/Animation/PoseSearch/Source/Runtime/Private/PoseSearchTrajectoryLibrary.cpp
--- a/Engine/Plugins/Animation/PoseSearch/Source/Runtime/Private/PoseSearchTrajectoryLibrary.cpp
+++ b/Engine/Plugins/Animation/PoseSearch/Source/Runtime/Private/PoseSearchTrajectoryLibrary.cpp
@@ -190,7 +190,6 @@ void UPoseSearchTrajectoryLibrary::InitTrajectorySamples(
 }
 
 // Deprecated
-PRAGMA_DISABLE_DEPRECATION_WARNINGS
 void UPoseSearchTrajectoryLibrary::InitTrajectorySamples(
 	FPoseSearchQueryTrajectory& Trajectory,
 	const FPoseSearchTrajectoryData& TrajectoryData,	// Unused parameter
@@ -207,7 +206,6 @@ void UPoseSearchTrajectoryLibrary::InitTrajectorySamples(
 }
 
 // Deprecated
-
 void UPoseSearchTrajectoryLibrary::InitTrajectorySamples(
 	FPoseSearchQueryTrajectory& Trajectory,
 	FVector DefaultPosition, FQuat DefaultFacing,
diff --git a/Engine/Plugins/Animation/PoseSearch/Source/Runtime/Private/PoseSearchTrajectoryTypes.cpp b/Engine/Plugins/Animation/PoseSearch/Source/Runtime/Private/PoseSearchTrajectoryTypes.cpp
--- a/Engine/Plugins/Animation/PoseSearch/Source/Runtime/Private/PoseSearchTrajectoryTypes.cpp
+++ b/Engine/Plugins/Animation/PoseSearch/Source/Runtime/Private/PoseSearchTrajectoryTypes.cpp
@@ -232,9 +232,9 @@ FPoseSearchQueryTrajectory::operator FTransformTrajectory() const
 {
 	FTransformTrajectory OutTrajectory;
 
-	for (const auto & Sample : Samples)
+	for (const FPoseSearchQueryTrajectorySample& Sample : Samples)
 	{
-		OutTrajectory.Samples.Emplace(Sample.Facing, Sample.Position, Sample.AccumulatedSeconds);
+		OutTrajectory.Samples.Add(FTransformTrajectorySample(Sample.Facing, Sample.Position, Sample.AccumulatedSeconds));
 	}
 
 	return OutTrajectory;
