posesearch
- FAnimNode_MotionMatchingInteraction warping improvements (added PoseSearch dependency to AnimationWarpingRuntime)
- introduced FPoseSearchIKWindowConstraint for future MM interaction IK constraints (removed old IK solution form FAnimNode_MotionMatchingInteraction::Evaluate_AnyThread)
- added new UPoseSearchAssetSamplerLibrary::VLogDraw
- introduced MM interaction TickPriority to control tick dependency order in interaction islands
- added recursion safe FTickFunction logging functions

diff --git a/Engine/Plugins/Animation/PoseSearch/PoseSearch.uplugin b/Engine/Plugins/Animation/PoseSearch/PoseSearch.uplugin
--- a/Engine/Plugins/Animation/PoseSearch/PoseSearch.uplugin
+++ b/Engine/Plugins/Animation/PoseSearch/PoseSearch.uplugin
@@ -27,11 +27,8 @@
 	],
 	"Plugins": [
 		{
-			"Name": "GameplayInsights",
-			"Enabled": true,
-			"TargetAllowList": [
-				"Editor"
-			]
+			"Name": "AnimationWarping",
+			"Enabled": true
 		},
 		{
 			"Name": "BlendStack",
@@ -40,6 +37,13 @@
 		{
 			"Name": "Chooser",
 			"Enabled": true
+		},
+		{
+			"Name": "GameplayInsights",
+			"Enabled": true,
+			"TargetAllowList": [
+				"Editor"
+			]
 		}
 	]
 }
\ No newline at end of file
diff --git a/Engine/Plugins/Animation/PoseSearch/Source/Runtime/PoseSearch.Build.cs b/Engine/Plugins/Animation/PoseSearch/Source/Runtime/PoseSearch.Build.cs
--- a/Engine/Plugins/Animation/PoseSearch/Source/Runtime/PoseSearch.Build.cs
+++ b/Engine/Plugins/Animation/PoseSearch/Source/Runtime/PoseSearch.Build.cs
@@ -16,21 +16,22 @@ public class PoseSearch : ModuleRules
 		PublicDependencyModuleNames.AddRange(
 			new string[]
 			{
+				"AnimationCore",
+				"AnimGraphRuntime",
+				"BlendStack",
 				"Chooser",
 				"Core",
 				"CoreUObject",
 				"Engine",
-				"AnimationCore",
-				"TraceLog",
-				"AnimGraphRuntime",
-				"BlendStack"
+				"TraceLog"
 			}
 		);
 
 		PrivateDependencyModuleNames.AddRange(
 		new string[]{
+			"AnimationWarpingRuntime",
 			"GameplayTags",
-			"RewindDebuggerRuntimeInterface",
+			"RewindDebuggerRuntimeInterface"
 		}
 		);
 
diff --git a/Engine/Plugins/Animation/PoseSearch/Source/Runtime/Private/AnimNode_MotionMatchingInteraction.cpp b/Engine/Plugins/Animation/PoseSearch/Source/Runtime/Private/AnimNode_MotionMatchingInteraction.cpp
--- a/Engine/Plugins/Animation/PoseSearch/Source/Runtime/Private/AnimNode_MotionMatchingInteraction.cpp
+++ b/Engine/Plugins/Animation/PoseSearch/Source/Runtime/Private/AnimNode_MotionMatchingInteraction.cpp
@@ -4,73 +4,28 @@
 #include "Animation/AnimInertializationSyncScope.h"
 #include "Animation/AnimInstanceProxy.h"
 #include "Animation/AnimRootMotionProvider.h"
+#include "BoneControllers/AnimNode_OffsetRootBone.h"
 #include "HAL/IConsoleManager.h"
 #include "PoseSearch/PoseHistoryProvider.h"
 #include "PoseSearch/PoseSearchAnimNotifies.h"
 #include "PoseSearch/PoseSearchAssetSamplerLibrary.h"
+#include "PoseSearch/PoseSearchContext.h"
 #include "PoseSearch/PoseSearchDatabase.h"
 #include "PoseSearch/PoseSearchInteractionAsset.h"
 #include "PoseSearch/PoseSearchInteractionIsland.h"
 #include "PoseSearch/PoseSearchInteractionLibrary.h"
+#include "PoseSearch/PoseSearchInteractionUtils.h"
 #include "PoseSearch/PoseSearchLibrary.h"
 #include "PoseSearch/PoseSearchMirrorDataCache.h"
 #include "PoseSearch/PoseSearchSchema.h"
 #include "VisualLogger/VisualLogger.h"
 
-namespace UE::PoseSearch
-{
-
-#if ENABLE_DRAW_DEBUG && ENABLE_ANIM_DEBUG
-static bool GVarAnimNodeMotionMatchingInteractionDrawFullAlignedPose = false;
-static FAutoConsoleVariableRef CVarAnimNodeMotionMatchingInteractionDrawFullAlignedPose(TEXT("a.AnimNode.MotionMatchingInteraction.DrawFullAlignedPose"), GVarAnimNodeMotionMatchingInteractionDrawFullAlignedPose, TEXT("Draw Full Aligned Pose"));
-#endif // ENABLE_DRAW_DEBUG && ENABLE_ANIM_DEBUG
-
-// Taken from https://theorangeduck.com/page/spring-roll-call#implicitspringdamper
-static float DamperImplicit(float Halflife, float DeltaTime, float Epsilon = 1e-8f)
-{
-	// Halflife values very close to 0 approach infinity, and result in big motion spikes when Halflife < DeltaTime
-	// This is a hack, and adds some degree of frame-rate dependency, but it holds up even at lower frame-rates.
-	Halflife = FMath::Max(Halflife, DeltaTime);
-	return FMath::Clamp((1.0f - FMath::InvExpApprox((0.69314718056f * DeltaTime) / (Halflife + Epsilon))), 0.0f, 1.0f);
-}
-
-static FRoleToIndex MakeRoleToIndex(const UMultiAnimAsset* MultiAnimAsset)
-{
-	check(MultiAnimAsset);
-	const int32 NumRoles = MultiAnimAsset->GetNumRoles();
-
-	FRoleToIndex RoleToIndex;
-	RoleToIndex.Reserve(NumRoles);
-	for (int32 MultiAnimAssetRoleIndex = 0; MultiAnimAssetRoleIndex < NumRoles; ++MultiAnimAssetRoleIndex)
-	{
-		RoleToIndex.Add(MultiAnimAsset->GetRole(MultiAnimAssetRoleIndex)) = MultiAnimAssetRoleIndex;
-	}
-	return RoleToIndex;
-}
-
-typedef TArray<bool, TInlineAllocator<PreallocatedRolesNum, TMemStackAllocator<>>> FRelevantRoleIndexes;
-static FRelevantRoleIndexes GetRelevantRoleIndexes(const UMultiAnimAsset* MultiAnimAsset, const UPoseSearchSchema* Schema)
-{
-	check(Schema && MultiAnimAsset);
-
-	const FRoleToIndex RoleToIndex = MakeRoleToIndex(MultiAnimAsset);
-	const int32 NumRoles = MultiAnimAsset->GetNumRoles();
-
-	FRelevantRoleIndexes RelevantRoleIndexes;
-	RelevantRoleIndexes.SetNum(NumRoles);
-	for (const FPoseSearchRoledSkeleton& RoledSkeleton : Schema->GetRoledSkeletons())
-	{
-		RelevantRoleIndexes[RoleToIndex[RoledSkeleton.Role]] = true;
-	}
-	return RelevantRoleIndexes;
-}
-
-} // UE::PoseSearch
-
 void FAnimNode_MotionMatchingInteraction::Reset()
 {
 	Super::Reset();
 	CurrentResult = FPoseSearchInteractionBlueprintResult();
+	MeshWithOffset = FTransform::Identity;
+	MeshWithoutOffset = FTransform::Identity;
 	CachedDeltaTime = 0.f;
 }
 
@@ -95,6 +50,16 @@ void FAnimNode_MotionMatchingInteraction::UpdateAssetPlayer(const FAnimationUpda
 		Reset();
 	}
 
+	MeshWithoutOffset = GetContextTransform(Context.AnimInstanceProxy->GetAnimInstanceObject());
+	if (UE::AnimationWarping::FRootOffsetProvider* RootOffsetProvider = Context.GetMessage<UE::AnimationWarping::FRootOffsetProvider>())
+	{
+		MeshWithOffset = RootOffsetProvider->GetRootTransform();
+	}
+	else
+	{
+		MeshWithOffset = MeshWithoutOffset;
+	}
+
 	UpdateCounter.SynchronizeWith(Context.AnimInstanceProxy->GetUpdateCounter());
 	CachedDeltaTime = Context.GetDeltaTime();
 
@@ -165,11 +130,6 @@ void FAnimNode_MotionMatchingInteraction::Evaluate_AnyThread(FPoseContext& Outpu
 	Super::Evaluate_AnyThread(Output);
 
 	// calculating the animation time we want the full alignment transforms from
-	const float AccumulatedTime = GetAccumulatedTime();
-				
-	// warping alignment times in seconds
-	float WarpingAlignmentExtractionTime = AccumulatedTime;
-
 	UAnimInstance* AnimInstance = Cast<UAnimInstance>(Output.AnimInstanceProxy->GetAnimInstanceObject());
 	if (!AnimInstance)
 	{
@@ -198,15 +158,7 @@ void FAnimNode_MotionMatchingInteraction::Evaluate_AnyThread(FPoseContext& Outpu
 		return;
 	}
 
-	int32 CurrentResultRoleIndex = INDEX_NONE;
-	for (int32 MultiAnimAssetRoleIndex = 0; MultiAnimAssetRoleIndex < NumRoles; ++MultiAnimAssetRoleIndex)
-	{
-		if (MultiAnimAsset->GetRole(MultiAnimAssetRoleIndex) == CurrentResult.Role)
-		{
-			CurrentResultRoleIndex = MultiAnimAssetRoleIndex;
-			break;
-		}
-	}
+	const int32 CurrentResultRoleIndex = GetRoleIndex(MultiAnimAsset, CurrentResult.Role);
 	if (CurrentResultRoleIndex == INDEX_NONE)
 	{
 		return;
@@ -219,211 +171,36 @@ void FAnimNode_MotionMatchingInteraction::Evaluate_AnyThread(FPoseContext& Outpu
 		return;
 	}
 
-	FMemMark Mark(FMemStack::Get());
-
-	// @todo: predicting ActorRootTransforms TimeToFullAlignment in the future! Maybe by using the trajectories..
-	const TArrayView<FTransform> ActorRootTransforms = CurrentResult.ActorRootTransforms;
-	const TArrayView<FTransform> ActorRootBoneTransforms = CurrentResult.ActorRootBoneTransforms;
-
-	// NoTe: applying the previous frame historical root bone orientation to the ActorRootTransforms (array of ComponentTransform(s)), instead of applying the translation as well via
-	// ActorRootBoneTransforms[i] * ActorRootTransforms[i], to get the ActorTransforms[i] where the animation wants it to be without being affected by OffsetRootBone trnaslations,
-	// so we'll be able to warp into full aligment and blend back once the aligment is no longer required, ultimately converging back where the capsule is located 
-	// (calculated below in AlignmentTransformNoWarpAdjustment, and adjusted by any RootMotionTransform)
-	TArray<FTransform, TInlineAllocator<PreallocatedRolesNum, TMemStackAllocator<>>> ActorTransforms;
-	ActorTransforms.SetNum(NumRoles);
-	for (int32 MultiAnimAssetRoleIndex = 0; MultiAnimAssetRoleIndex < NumRoles; ++MultiAnimAssetRoleIndex)
-	{
-		ActorTransforms[MultiAnimAssetRoleIndex] = FTransform(
-			ActorRootTransforms[MultiAnimAssetRoleIndex].GetRotation() * ActorRootBoneTransforms[MultiAnimAssetRoleIndex].GetRotation(),
-			ActorRootTransforms[MultiAnimAssetRoleIndex].GetTranslation(),
-			ActorRootTransforms[MultiAnimAssetRoleIndex].GetScale3D());
-	}
-
-	const FAnimationAssetSampler Sampler(AnimationAsset, FTransform::Identity, CurrentResult.BlendParameters, FAnimationAssetSampler::DefaultRootTransformSamplingRate, false);
+	TArray<FTransform, TInlineAllocator<PreallocatedRolesNum>> FullAlignedTransforms;
+	FullAlignedTransforms.SetNum(NumRoles);
+	CalculateFullAlignedTransforms(CurrentResult, bWarpUsingRootBone, FullAlignedTransforms);
 
-	if (bUsePoseSearchWarpingNotifyStates)
-	{
-		// using UAnimNotifyState_PoseSearchWarping to change the WarpingAlignmentExtractionTime (time we sample the full alignment transforms at)
-		for (const FAnimNotifyEvent& NotifyEvent : Sampler.GetAllAnimNotifyEvents())
-		{
-			if (NotifyEvent.NotifyStateClass && NotifyEvent.NotifyStateClass->GetClass()->IsChildOf<UAnimNotifyState_PoseSearchWarping>())
-			{
-				const UAnimNotifyState_PoseSearchWarping* NotifyStatePoseSearchWarping = Cast<UAnimNotifyState_PoseSearchWarping>(NotifyEvent.NotifyStateClass);
-				check(NotifyStatePoseSearchWarping);
-
-				WarpingAlignmentExtractionTime = FMath::Max(NotifyEvent.GetTime(), AccumulatedTime);
-				break;
-			}
-		}
-	}
+	// NoTe: keep in mind DeltaAlignment is relative to the previous execution frame so we still need to extract and and apply the current animation root motion transform to get to the current frame full aligned transform.
+	const FTransform DeltaAlignment = CalculateDeltaAlignment(MeshWithoutOffset, MeshWithOffset, FullAlignedTransforms[CurrentResultRoleIndex], WarpingRotationRatio, WarpingTranslationRatio);
 	
-	TSet<FName, DefaultKeyFuncs<FName>, TInlineSetAllocator<16, TMemStackSetAllocator<>>> GoalNames;
-	if (bUsePoseSearchIKWindowNotifyStates)
-	{
-		FAnimNotifyContext PreAllocatedNotifyContext;
-		Sampler.ExtractAnimNotifyStates(AccumulatedTime, PreAllocatedNotifyContext, [&GoalNames](const UAnimNotifyState* AnimNotifyState)
-			{
-				if (const UAnimNotifyState_PoseSearchIKWindow* PoseSearchIKWindow = Cast<const UAnimNotifyState_PoseSearchIKWindow>(AnimNotifyState))
-				{
-					for (FName GoalName : PoseSearchIKWindow->GoalNames)
-					{
-						GoalNames.Add(GoalName);
-					}
-				}
-				// continue for other notify states
-				return true;
-			});
-	}
-
-	const bool bNeedCalculatingWarping = WarpingTranslationRatio > UE_KINDA_SMALL_NUMBER || WarpingRotationRatio > UE_KINDA_SMALL_NUMBER;
-	if (!bNeedCalculatingWarping && GoalNames.IsEmpty())
-	{
-		// no warping nor IK necessary. we're done
-		return;
-	}
-
-	// calculating the world space warp tranform to have the character in its "original" no reference frame (where no warping associated to MultiAnimAsset is impacting its transform)
-	FTransform RootMotionTransform;
-	RootMotionProvider->ExtractRootMotion(Output.CustomAttributes, RootMotionTransform);
-
-	// NoTe that ActorRootTransforms[CurrentResultRoleIndex] is Output.AnimInstanceProxy->GetSkelMeshComponent()->GetComponentTransform();
-	const FTransform AlignmentTransformNoWarpAdjustment = RootMotionTransform * ActorRootTransforms[CurrentResultRoleIndex];
-					
-	check(CurrentResult.SelectedDatabase != nullptr && CurrentResult.SelectedDatabase->Schema != nullptr);
-	const FRelevantRoleIndexes RelevantRoleIndexes = GetRelevantRoleIndexes(MultiAnimAsset, CurrentResult.SelectedDatabase->Schema);
+	FTransform RootMotionDelta;
+	RootMotionProvider->ExtractRootMotion(Output.CustomAttributes, RootMotionDelta);
 
-	TArray<const UMirrorDataTable*, TInlineAllocator<PreallocatedRolesNum, TMemStackAllocator<>>> MirrorDataTables;
-	if (CurrentResult.bIsMirrored)
-	{
-		MirrorDataTables.SetNum(NumRoles);
-		for (int32 MultiAnimAssetRoleIndex = 0; MultiAnimAssetRoleIndex < NumRoles; ++MultiAnimAssetRoleIndex)
-		{
-			const FRole Role = MultiAnimAsset->GetRole(MultiAnimAssetRoleIndex);
-			const FPoseSearchRoledSkeleton* RoledSkeleton = CurrentResult.SelectedDatabase->Schema->GetRoledSkeleton(Role);
-			check(RoledSkeleton);
-			MirrorDataTables[MultiAnimAssetRoleIndex] = RoledSkeleton->MirrorDataTable;
-		}
-	}
-
-	// calculating the world space warp tranform to have the character fully aligned to what MultiAnimAsset wants.
-	// using ActorTransforms instead of ActorRootTransforms, since we want the previous frame root bone offset rotation to be embedded into the calculation
-	TArray<FTransform, TInlineAllocator<PreallocatedRolesNum, TMemStackAllocator<>>> FullAlignedActorRootBoneTransformsAtExtractionTime;
-	FullAlignedActorRootBoneTransformsAtExtractionTime.SetNum(NumRoles);
-	MultiAnimAsset->CalculateWarpTransforms(WarpingAlignmentExtractionTime, ActorTransforms, FullAlignedActorRootBoneTransformsAtExtractionTime, MirrorDataTables, RelevantRoleIndexes);
-	const FTransform AlignmentTransformFullWarpAdjustment = FullAlignedActorRootBoneTransformsAtExtractionTime[CurrentResultRoleIndex];
-					
-	// @todo: compensating AlignmentTransformFullWarpAdjustment by the delta time between AccumulatedTime and WarpingAlignmentExtractionTime?				
-	// @todo: optimize me! we can skip some math if (!bNeedCalculatingWarping)
-	
-	// calculating the world space blended warp aligment transform
-	const FTransform AlignmentTransform(
-		FMath::Lerp(AlignmentTransformNoWarpAdjustment.GetRotation(), AlignmentTransformFullWarpAdjustment.GetRotation(), WarpingRotationRatio),
-		FMath::Lerp(AlignmentTransformNoWarpAdjustment.GetTranslation(), AlignmentTransformFullWarpAdjustment.GetTranslation(), WarpingTranslationRatio),
-		FVector::OneVector);
-
-	// calculating WarpRootMotionTransform as AlignmentTransform in root bone space
-	const FTransform RootBoneTransform = Output.AnimInstanceProxy->GetSkelMeshComponent()->GetBoneTransform(RootBoneIndexType);
-	FTransform WarpRootMotionTransform = AlignmentTransform.GetRelativeTransform(RootBoneTransform);
-	check(WarpRootMotionTransform.IsRotationNormalized());
-
-	// apply some damping to WarpRootMotionTransform in root bone space
-	const float DampenWarpingTranslation = WarpingTranslationHalflife > 0.f ? DamperImplicit(WarpingTranslationHalflife, CachedDeltaTime) : 1.f;
-	const float DampenWarpingRotation = WarpingRotationHalfLife > 0.f ? DamperImplicit(WarpingRotationHalfLife, CachedDeltaTime) : 1.f;
-	WarpRootMotionTransform.SetTranslation(WarpRootMotionTransform.GetTranslation() * DampenWarpingTranslation);
-	WarpRootMotionTransform.SetRotation(FQuat::Slerp(FQuat::Identity, WarpRootMotionTransform.GetRotation(), DampenWarpingRotation));
+	const FTransform DeltaAlignmentWithRootMotion = DeltaAlignment * RootMotionDelta;
+	RootMotionProvider->OverrideRootMotion(DeltaAlignmentWithRootMotion, Output.CustomAttributes);
 
 #if ENABLE_VISUAL_LOG
 	if (FVisualLogger::IsRecording())
 	{
 		static const TCHAR* LogName = TEXT("MotionMatchingInteraction");
 
-		const FTransform WorldWarpRootMotionTransform = WarpRootMotionTransform * RootBoneTransform;
-		UE_VLOG_ARROW(AnimInstance, LogName, Display, WorldWarpRootMotionTransform.GetLocation(), WorldWarpRootMotionTransform.GetLocation() + WorldWarpRootMotionTransform.GetRotation().GetForwardVector() * 100, FColor::Green, TEXT(""));
-		UE_VLOG_ARROW(AnimInstance, LogName, Display, AlignmentTransform.GetLocation(), AlignmentTransform.GetLocation() + AlignmentTransform.GetRotation().GetForwardVector() * 100, FColor::Blue, TEXT(""));
-		UE_VLOG_ARROW(AnimInstance, LogName, Display, RootBoneTransform.GetLocation(), RootBoneTransform.GetLocation() + RootBoneTransform.GetRotation().GetForwardVector() * 100, FColor::Orange, TEXT(""));
-	}
-#endif
-	RootMotionProvider->OverrideRootMotion(WarpRootMotionTransform, Output.CustomAttributes);
-
-	if (GoalNames.IsEmpty())
-	{
-		return;
-	}
-
-	FTransform FullAlignedPoseRootTransform;
-	if (FMath::IsNearlyEqual(WarpingAlignmentExtractionTime, AccumulatedTime))
-	{
-		FullAlignedPoseRootTransform = FullAlignedActorRootBoneTransformsAtExtractionTime[CurrentResultRoleIndex];
-	}
-	else
-	{
-		// if the we extracted the full alignment transforms NOT at AccumulatedTime (usually to make aligment calculation more stable - closer to the interaction time)
-		// we need to reextract them. This could happen only if users want to align the IK before aligning the character, things that I think should never happen, 
-		// but just in case the code supports it
-		// using ActorTransforms instead of ActorRootTransforms, since we want the previous frame root bone offset rotation to be embedded into the calculation
-		TArray<FTransform, TInlineAllocator<PreallocatedRolesNum, TMemStackAllocator<>>> FullAlignedActorRootBoneTransformsAtAccumulatedTime;
-		FullAlignedActorRootBoneTransformsAtAccumulatedTime.SetNum(NumRoles);
-		MultiAnimAsset->CalculateWarpTransforms(AccumulatedTime, ActorTransforms, FullAlignedActorRootBoneTransformsAtAccumulatedTime, MirrorDataTables, RelevantRoleIndexes);
-		FullAlignedPoseRootTransform = FullAlignedActorRootBoneTransformsAtAccumulatedTime[CurrentResultRoleIndex];
-	}
-								
-	// solving the IK constraints for this character, knowing where its FullAlignedPose should be
-	const FBoneContainer& BoneContainer = AnimInstance->GetRequiredBonesOnAnyThread();
-	const USkeleton* Skeleton = BoneContainer.GetSkeletonAsset();
-	check(Skeleton);
-
-	FCompactPose FullAlignedPose;
-	FullAlignedPose.SetBoneContainer(&BoneContainer);
-	Sampler.ExtractPose(AccumulatedTime, FullAlignedPose);
-	FCSPose<FCompactPose> ComponentSpaceFullAlignedPose;
-	ComponentSpaceFullAlignedPose.InitPose(FullAlignedPose);
-
-	FCSPose<FCompactPose> ComponentSpacePose;
-	ComponentSpacePose.InitPose(Output.Pose);
-	const FTransform WorldWarpRootMotionTransform = WarpRootMotionTransform * RootBoneTransform;
-
-	for (FName GoalName : GoalNames)
-	{
-		FBoneReference BoneReference;
-		BoneReference.BoneName = GoalName;
-		BoneReference.Initialize(Skeleton);
-		if (!BoneReference.HasValidSetup())
+		for (int32 Index = 0; Index < NumRoles; ++Index)
 		{
-			UE_LOG(LogPoseSearch, Error, TEXT("FAnimNode_MotionMatchingInteraction::Evaluate_AnyThread invalid GoalName %s for Skeleton %s"), *GoalName.ToString(), *GetNameSafe(Skeleton));
-		}
-		else
-		{
-			const FCompactPoseBoneIndex CompactBoneIdx = BoneContainer.GetCompactPoseIndexFromSkeletonPoseIndex(FSkeletonPoseBoneIndex(BoneReference.BoneIndex));
-			if (!CompactBoneIdx.IsValid())
-			{
-				UE_LOG(LogPoseSearch, Error, TEXT("FAnimNode_MotionMatchingInteraction::Evaluate_AnyThread invalid CompactPoseBoneIndex for GoalName %s for Skeleton %s with BoneIndex %d"), *GoalName.ToString(), *GetNameSafe(Skeleton), BoneReference.BoneIndex);
-			}
-			else
-			{
-				const FTransform GoalWorldTransform = ComponentSpaceFullAlignedPose.GetComponentSpaceTransform(CompactBoneIdx) * FullAlignedPoseRootTransform;
-				const FCompactPoseBoneIndex CompactParentBoneIdx = BoneContainer.GetParentBoneIndex(CompactBoneIdx);
-				const FTransform WorldPoseParentBone = ComponentSpacePose.GetComponentSpaceTransform(CompactParentBoneIdx) * WorldWarpRootMotionTransform;
-				const FTransform GoalLocalTransform = GoalWorldTransform.GetRelativeTransform(WorldPoseParentBone);
-				Output.Pose[CompactBoneIdx] = GoalLocalTransform;
-#if ENABLE_VISUAL_LOG
-				if (FVisualLogger::IsRecording())
-				{
-					static const TCHAR* LogName = TEXT("MotionMatchingInteraction");
-					UE_VLOG_ARROW(AnimInstance, LogName, Display, FullAlignedPoseRootTransform.GetLocation(), GoalWorldTransform.GetLocation(), FColor::Black, TEXT(""));
-					UE_VLOG_ARROW(AnimInstance, LogName, Display, ActorRootTransforms[CurrentResultRoleIndex].GetLocation(), WorldPoseParentBone.GetLocation(), FColor::Red, TEXT(""));
-					const FTransform WorldPoseBone = ComponentSpacePose.GetComponentSpaceTransform(CompactBoneIdx) * WorldWarpRootMotionTransform;
-					UE_VLOG_ARROW(AnimInstance, LogName, Display, WorldPoseParentBone.GetLocation(), WorldPoseBone.GetLocation(), FColor::Green, TEXT(""));
-				}
-#endif
-			}
+			const FTransform& ActorRootTransform = CurrentResult.ActorRootTransforms[Index];
+			const FTransform& FullAlignedTransform = FullAlignedTransforms[Index];
+
+			UE_VLOG_SEGMENT_THICK(AnimInstance, LogName, Display, FullAlignedTransform.GetLocation(), ActorRootTransform.GetLocation(), FColorList::Orange, 1.f, TEXT(""));
+			UE_VLOG_SEGMENT_THICK(AnimInstance, LogName, Display, ActorRootTransform.GetLocation(), ActorRootTransform.GetLocation() + ActorRootTransform.GetRotation().GetForwardVector() * 35, FColorList::LightGrey, 3.f, TEXT(""));
+			UE_VLOG_SEGMENT_THICK(AnimInstance, LogName, Display, FullAlignedTransform.GetLocation(), FullAlignedTransform.GetLocation() + FullAlignedTransform.GetRotation().GetForwardVector() * 30, FColorList::Orange, 2.f, TEXT(""));
 		}
-	}
 
-#if ENABLE_DRAW_DEBUG && ENABLE_ANIM_DEBUG
-	if (GVarAnimNodeMotionMatchingInteractionDrawFullAlignedPose)
-	{
-		UPoseSearchAssetSamplerLibrary::Draw(AnimInstance->GetWorld(), ComponentSpaceFullAlignedPose, FullAlignedPoseRootTransform);
+		UE_VLOG_SEGMENT_THICK(AnimInstance, LogName, Display, MeshWithOffset.GetLocation(), MeshWithOffset.GetLocation() + MeshWithOffset.GetRotation().GetForwardVector() * 35, FColorList::Blue, 3.f, TEXT(""));
+		UE_VLOG_SEGMENT_THICK(AnimInstance, LogName, Display, MeshWithoutOffset.GetLocation(), MeshWithoutOffset.GetLocation() + MeshWithoutOffset.GetRotation().GetForwardVector() * 40, FColorList::Cyan, 4.f, TEXT(""));
 	}
-#endif // ENABLE_DRAW_DEBUG && ENABLE_ANIM_DEBUG
-
+#endif
 }
diff --git a/Engine/Plugins/Animation/PoseSearch/Source/Runtime/Private/PoseSearchAssetSamplerLibrary.cpp b/Engine/Plugins/Animation/PoseSearch/Source/Runtime/Private/PoseSearchAssetSamplerLibrary.cpp
--- a/Engine/Plugins/Animation/PoseSearch/Source/Runtime/Private/PoseSearchAssetSamplerLibrary.cpp
+++ b/Engine/Plugins/Animation/PoseSearch/Source/Runtime/Private/PoseSearchAssetSamplerLibrary.cpp
@@ -3,6 +3,7 @@
 #include "PoseSearch/PoseSearchAssetSamplerLibrary.h"
 #include "Animation/AnimInstance.h"
 #include "DrawDebugHelpers.h"
+#include "Engine/SkinnedAsset.h"
 #include "PoseSearch/PoseSearchDefines.h"
 #include "PoseSearch/PoseSearchMirrorDataCache.h"
 
@@ -99,4 +100,40 @@ void UPoseSearchAssetSamplerLibrary::Draw(const UAnimInstance* AnimInstance, UPA
 #endif // ENABLE_DRAW_DEBUG
 }
 
+#if ENABLE_VISUAL_LOG
+void UPoseSearchAssetSamplerLibrary::VLogDraw(const UObject* VLogContext, const USkeletalMeshComponent* Mesh, const TCHAR* VLogName, const FColor Color, float DebugDrawSamplerRootAxisLength)
+{
+	if (!Mesh)
+	{
+		return;
+	}
+		
+	const USkinnedAsset* SkinnedAsset = Mesh->GetSkinnedAsset();
+	if (!SkinnedAsset)
+	{
+		return;
+	}
+	
+	if (DebugDrawSamplerRootAxisLength > 0.f)
+	{
+		const FTransform& AxisWorldTransform = Mesh->GetComponentTransform();
+		UE_VLOG_SEGMENT(VLogContext, VLogName, Display, AxisWorldTransform.GetTranslation(), AxisWorldTransform.GetTranslation() + AxisWorldTransform.GetScaledAxis(EAxis::X) * DebugDrawSamplerRootAxisLength, FColor::Red, TEXT(""));
+		UE_VLOG_SEGMENT(VLogContext, VLogName, Display, AxisWorldTransform.GetTranslation(), AxisWorldTransform.GetTranslation() + AxisWorldTransform.GetScaledAxis(EAxis::Y) * DebugDrawSamplerRootAxisLength, FColor::Green, TEXT(""));
+		UE_VLOG_SEGMENT(VLogContext, VLogName, Display, AxisWorldTransform.GetTranslation(), AxisWorldTransform.GetTranslation() + AxisWorldTransform.GetScaledAxis(EAxis::Z) * DebugDrawSamplerRootAxisLength, FColor::Blue, TEXT(""));
+	}
+
+	const int NumBones = Mesh->GetNumBones();
+	for (int32 BoneIndex = 0; BoneIndex < NumBones; ++BoneIndex)
+	{
+		const int32 ParentBoneIndex = SkinnedAsset->GetRefSkeleton().GetParentIndex(BoneIndex);
+		if (ParentBoneIndex != INDEX_NONE)
+		{
+			const FTransform BoneWorldTransform = Mesh->GetBoneTransform(BoneIndex);
+			const FTransform ParentBoneWorldTransform = Mesh->GetBoneTransform(ParentBoneIndex);
+			UE_VLOG_SEGMENT(VLogContext, VLogName, Display, BoneWorldTransform.GetTranslation(), ParentBoneWorldTransform.GetTranslation(), Color, TEXT(""));
+		}
+	}
+}
+#endif // ENABLE_VISUAL_LOG
+
 
diff --git a/Engine/Plugins/Animation/PoseSearch/Source/Runtime/Private/PoseSearchContext.cpp b/Engine/Plugins/Animation/PoseSearch/Source/Runtime/Private/PoseSearchContext.cpp
--- a/Engine/Plugins/Animation/PoseSearch/Source/Runtime/Private/PoseSearchContext.cpp
+++ b/Engine/Plugins/Animation/PoseSearch/Source/Runtime/Private/PoseSearchContext.cpp
@@ -11,6 +11,7 @@
 #include "PoseSearch/PoseSearchFeatureChannel_Heading.h"
 #include "PoseSearch/PoseSearchFeatureChannel_PermutationTime.h"
 #include "PoseSearch/PoseSearchFeatureChannel_Position.h"
+#include "BoneContainer.h"
 
 namespace UE::PoseSearch
 {
@@ -65,6 +66,50 @@ const USkeleton* GetContextSkeleton(const UObject* AnimContext)
 	return nullptr;
 }
 
+const USkeletalMeshComponent* GetContextSkeletalMeshComponent(const UObject* AnimContext)
+{
+	if (const UAnimInstance* AnimInstance = Cast<UAnimInstance>(AnimContext))
+	{
+		return AnimInstance->GetSkelMeshComponent();
+	}
+	
+	if (const UActorComponent* ActorComponent = Cast<UActorComponent>(AnimContext))
+	{
+		const AActor* Actor = ActorComponent->GetOwner();
+		check(Actor);
+		// @todo: this code depends on how AnimNext gather its context object, and will likely change in the future
+		return Actor->GetComponentByClass<USkeletalMeshComponent>();
+	}
+
+	return nullptr;
+}
+
+const FBoneContainer GetBoneContainer(const UObject* AnimContext)
+{
+	if (const UAnimInstance* AnimInstance = Cast<UAnimInstance>(AnimContext))
+	{
+		return AnimInstance->GetRequiredBonesOnAnyThread();
+	}
+
+	if (const UActorComponent* ActorComponent = Cast<UActorComponent>(AnimContext))
+	{
+		const AActor* Actor = ActorComponent->GetOwner();
+		check(Actor);
+		// @todo: this code depends on how AnimNext gather its context object, and will likely change in the future
+		if (const USkeletalMeshComponent* SkeletalMeshComponent = Actor->GetComponentByClass<USkeletalMeshComponent>())
+		{
+			if (USkeletalMesh* SkeletalMesh = SkeletalMeshComponent->GetSkeletalMeshAsset())
+			{
+				return FBoneContainer(SkeletalMeshComponent->RequiredBones, UE::Anim::FCurveFilterSettings(UE::Anim::ECurveFilterMode::DisallowAll), *SkeletalMesh->GetSkeleton());
+			}
+		}
+		return FBoneContainer();
+	}
+
+	unimplemented();
+	return FBoneContainer();
+}
+
 const AActor* GetContextOwningActor(const UObject* AnimContext)
 {
 	if (const UAnimInstance* AnimInstance = Cast<UAnimInstance>(AnimContext))
diff --git a/Engine/Plugins/Animation/PoseSearch/Source/Runtime/Private/PoseSearchInteractionIsland.cpp b/Engine/Plugins/Animation/PoseSearch/Source/Runtime/Private/PoseSearchInteractionIsland.cpp
--- a/Engine/Plugins/Animation/PoseSearch/Source/Runtime/Private/PoseSearchInteractionIsland.cpp
+++ b/Engine/Plugins/Animation/PoseSearch/Source/Runtime/Private/PoseSearchInteractionIsland.cpp
@@ -738,6 +738,52 @@ bool FInteractionIsland::CanPerformSearch(const UObject* AnimContext) const
 }
 
 #if ENABLE_ANIM_DEBUG
+// recursion safe FTickFunction logging functions
+static void ShowPrerequistes(const FTickFunction& NestedTick, int32 Indent, int32 MaxIndent)
+{
+	if (Indent >= MaxIndent)
+	{
+		UE_LOG(LogPoseSearch, Log, TEXT("%s==== REACHED MAX INDENT ===="), FCString::Spc(Indent * 2));
+	}
+	else
+	{
+		for (const FTickPrerequisite& Prereq : NestedTick.GetPrerequisites())
+		{
+			if (Prereq.PrerequisiteTickFunction)
+			{
+				UE_LOG(LogPoseSearch, Log, TEXT("%s prereq %s"), FCString::Spc(Indent * 2), *Prereq.PrerequisiteTickFunction->DiagnosticMessage());
+				ShowPrerequistes(*Prereq.PrerequisiteTickFunction, Indent + 1, MaxIndent);
+			}
+		}
+	}
+}
+
+static void LogTickFunction(FTickFunction& Tick, ENamedThreads::Type CurrentThread, bool bLogPrerequisites, int32 Indent, int32 MaxIndent = 50)
+{
+	if (Indent >= MaxIndent)
+	{
+		UE_LOG(LogPoseSearch, Log, TEXT("%s==== REACHED MAX INDENT ===="), FCString::Spc(Indent * 2));
+	}
+	else
+	{
+		// scoping brackets to save some heap for the recursion
+		{
+			UE_LOG(LogPoseSearch, Log, TEXT("%stick %s [%1d, %1d] %6llu %2d %s"), FCString::Spc(Indent * 2), Tick.bHighPriority ? TEXT("*") : TEXT(" "), (int32)Tick.GetActualTickGroup(), (int32)Tick.GetActualEndTickGroup(), (uint64)GFrameCounter, (int32)CurrentThread, *Tick.DiagnosticMessage());
+			if (bLogPrerequisites)
+			{
+				ShowPrerequistes(Tick, Indent, MaxIndent);
+			}
+		}
+
+		// Handle nested ticks
+		Tick.ForEachNestedTick([CurrentThread, bLogPrerequisites, Indent, MaxIndent](FTickFunction& NestedTick)
+			{
+				LogTickFunction(NestedTick, CurrentThread, bLogPrerequisites, Indent + 1, MaxIndent);
+			});
+	}
+}
+
+
 void FInteractionIsland::LogTickDependencies() const
 {
 	const ENamedThreads::Type Type = IsInGameThread() ? ENamedThreads::GameThread : ENamedThreads::AnyThread;
@@ -754,7 +800,7 @@ void FInteractionIsland::LogTickDependencies() const
 				// AnimNextComponent case
 				if (const FTickFunction* TickActorComponentTickFunction = IslandDependency->FindTickFunction(TickActorComponent))
 				{
-					const_cast<FTickFunction*>(TickActorComponentTickFunction)->LogTickFunction(Type, true, 1);
+					LogTickFunction(*const_cast<FTickFunction*>(TickActorComponentTickFunction), Type, true, 1);
 				}
 				else
 				{
@@ -764,7 +810,7 @@ void FInteractionIsland::LogTickDependencies() const
 			else
 			{
 				// SkeletalMeshComponent / AnimInstance case
-				TickActorComponent->PrimaryComponentTick.LogTickFunction(Type, true, 1);
+				LogTickFunction(TickActorComponent->PrimaryComponentTick, Type, true, 1);
 			}
 		}
 	}
@@ -925,7 +971,7 @@ bool FInteractionIsland::DoSearch_AnyThread(const UObject* AnimContext, const FP
 	return false;
 }
 
-bool FInteractionIsland::GetResult_AnyThread(const UObject* AnimContext, FPoseSearchInteractionBlueprintResult& Result)
+bool FInteractionIsland::GetResult_AnyThread(const UObject* AnimContext, FPoseSearchInteractionBlueprintResult& Result, bool bCompareOwningActors)
 {
 	check(AnimContext);
 
@@ -934,13 +980,25 @@ bool FInteractionIsland::GetResult_AnyThread(const UObject* AnimContext, FPoseSe
 	FScopeLock Lock(&SearchResultsMutex);
 #endif // UE_POSE_SEARCH_ENABLE_ISLAND_CRITICAL_SECTION
 
+	const AActor* Actor = bCompareOwningActors ? GetContextOwningActor(AnimContext) : nullptr;
+
 	// looking for AnimContext in SearchResults to fill up Result
 	for (const FInteractionSearchResult& SearchResult : SearchResults)
 	{
 		const FInteractionSearchContext& SearchContext = SearchContexts[SearchResult.SearchIndex];
 		for (int32 AnimContextIndex = 0; AnimContextIndex < SearchContext.AnimContexts.Num(); ++AnimContextIndex)
 		{
-			if (SearchContext.AnimContexts[AnimContextIndex].Get() == AnimContext)
+			bool bAnimContextFound = false;
+			if (bCompareOwningActors)
+			{
+				bAnimContextFound = GetContextOwningActor(SearchContext.AnimContexts[AnimContextIndex].Get()) == Actor;
+			}
+			else
+			{
+				bAnimContextFound = SearchContext.AnimContexts[AnimContextIndex].Get() == AnimContext;
+			}		
+
+			if (bAnimContextFound)
 			{
 				const UPoseSearchDatabase* Database = SearchResult.Database.Get();
 				check(Database);
@@ -963,8 +1021,6 @@ bool FInteractionIsland::GetResult_AnyThread(const UObject* AnimContext, FPoseSe
 				Result.SearchCost = SearchResult.PoseCost;
 				Result.Role = Role;
 
-				//Result.bIsFromContinuingPlaying = SearchResult.bIsContinuingPoseSearch;
-
 				// figuring out the WantedPlayRate
 				Result.WantedPlayRate = 1.f;
 				//if (Future.Animation && Future.IntervalTime > 0.f)
@@ -986,6 +1042,7 @@ bool FInteractionIsland::GetResult_AnyThread(const UObject* AnimContext, FPoseSe
 					const int32 NumRoles = MultiAnimAsset->GetNumRoles();
 					Result.ActorRootTransforms.SetNum(NumRoles);
 					Result.ActorRootBoneTransforms.SetNum(NumRoles);
+					Result.AnimContexts.SetNum(NumRoles);
 
 					const FRoleToIndex InteractionSearchContextRoleToIndex = MakeRoleToIndex(SearchContext.Roles);
 
@@ -995,9 +1052,25 @@ bool FInteractionIsland::GetResult_AnyThread(const UObject* AnimContext, FPoseSe
 						{
 							Result.ActorRootTransforms[MultiAnimAssetRoleIndex] = SearchResult.ActorRootTransforms[*InteractionSearchContextRoleIndex];
 							Result.ActorRootBoneTransforms[MultiAnimAssetRoleIndex] = SearchResult.ActorRootBoneTransforms[*InteractionSearchContextRoleIndex];
+							Result.AnimContexts[MultiAnimAssetRoleIndex] = SearchContext.AnimContexts[*InteractionSearchContextRoleIndex];
+						}
+						else
+						{
+							Result.ActorRootTransforms[MultiAnimAssetRoleIndex] = FTransform::Identity;
+							Result.ActorRootBoneTransforms[MultiAnimAssetRoleIndex] = FTransform::Identity;
+							Result.AnimContexts[MultiAnimAssetRoleIndex] = nullptr;
 						}
 					}
 				}
+				else
+				{
+					// @todo: should we support trivial "interactions" with only a character defined using some other assets rather then UMultiAnimAsset?
+					check(Result.AnimContexts.Num() == 1);
+
+					Result.ActorRootTransforms = SearchResult.ActorRootTransforms;
+					Result.ActorRootBoneTransforms = SearchResult.ActorRootBoneTransforms;
+					Result.AnimContexts = SearchContext.AnimContexts;
+				}
 
 				// we found our AnimContext: we can stop searching
 				return true;
diff --git a/Engine/Plugins/Animation/PoseSearch/Source/Runtime/Private/PoseSearchInteractionLibrary.cpp b/Engine/Plugins/Animation/PoseSearch/Source/Runtime/Private/PoseSearchInteractionLibrary.cpp
--- a/Engine/Plugins/Animation/PoseSearch/Source/Runtime/Private/PoseSearchInteractionLibrary.cpp
+++ b/Engine/Plugins/Animation/PoseSearch/Source/Runtime/Private/PoseSearchInteractionLibrary.cpp
@@ -31,6 +31,16 @@ FPoseSearchInteractionBlueprintResult UPoseSearchInteractionLibrary::MotionMatch
 	return Result;
 }
 
+FPoseSearchInteractionBlueprintResult UPoseSearchInteractionLibrary::GetCachedInteraction(const UObject* AnimContext, bool bCompareOwningActors)
+{
+	FPoseSearchInteractionBlueprintResult Result;
+	if (UPoseSearchInteractionSubsystem* InteractionSubsystem = UPoseSearchInteractionSubsystem::GetSubsystem_AnyThread(AnimContext))
+	{
+		InteractionSubsystem->GetResult_AnyThread(AnimContext, Result, bCompareOwningActors);
+	}
+	return Result;
+}
+
 FPoseSearchContinuingProperties UPoseSearchInteractionLibrary::GetMontageContinuingProperties(UAnimInstance* AnimInstance)
 {
 	FPoseSearchContinuingProperties ContinuingProperties;
diff --git a/Engine/Plugins/Animation/PoseSearch/Source/Runtime/Private/PoseSearchInteractionSubsystem.cpp b/Engine/Plugins/Animation/PoseSearch/Source/Runtime/Private/PoseSearchInteractionSubsystem.cpp
--- a/Engine/Plugins/Animation/PoseSearch/Source/Runtime/Private/PoseSearchInteractionSubsystem.cpp
+++ b/Engine/Plugins/Animation/PoseSearch/Source/Runtime/Private/PoseSearchInteractionSubsystem.cpp
@@ -11,6 +11,7 @@
 
 namespace UE::PoseSearch
 {
+
 #if ENABLE_DRAW_DEBUG && ENABLE_ANIM_DEBUG
 static bool GVarPoseSearchInteractionShowIslands = false;
 static FAutoConsoleVariableRef CVarPoseSearchInteractionShowIslands(TEXT("a.PoseSearchInteraction.ShowIslands"), GVarPoseSearchInteractionShowIslands, TEXT("Show Pose Search Interaction Islands"));
@@ -96,6 +97,7 @@ struct FAnimContextInfo
 	}
 
 	const UObject* AnimContext = nullptr;
+	int32 TickPriority = DefaultTickPriority;
 	const FInteractionAvailabilities* Availabilities = nullptr;
 	FVector Location = FVector::ZeroVector;
 	float MaxDistanceSquared = 0.f;
@@ -199,7 +201,7 @@ void UPoseSearchInteractionSubsystem::RegenerateAllIslands(TConstArrayView<UE::P
 		{
 			for (const TWeakObjectPtr<const UObject>& AnimContext : SearchContext.AnimContexts)
 			{
-				if (AnimContexts.Find(AnimContext.Get()))
+				if (AnimContextToTickPriority.Find(AnimContext.Get()))
 				{
 					return true;
 				}
@@ -209,11 +211,18 @@ void UPoseSearchInteractionSubsystem::RegenerateAllIslands(TConstArrayView<UE::P
 
 		void Add(const FInteractionSearchContext& SearchContext, int32 SearchContextIndex)
 		{
-			for (const TWeakObjectPtr<const UObject>& AnimContext : SearchContext.AnimContexts)
+			for (int32 AnimContextIndex = 0; AnimContextIndex < SearchContext.AnimContexts.Num(); ++AnimContextIndex)
 			{
-				if (AnimContext.Get())
+				if (const UObject* AnimContext = SearchContext.AnimContexts[AnimContextIndex].Get())
 				{
-					AnimContexts.Add(AnimContext.Get());
+					if (int32* TickPriority = AnimContextToTickPriority.Find(AnimContext))
+					{
+						*TickPriority = FMath::Max(*TickPriority, SearchContext.TickPriorities[AnimContextIndex]);
+					}
+					else
+					{
+						AnimContextToTickPriority.Add(AnimContext) = SearchContext.TickPriorities[AnimContextIndex];
+					}
 				}
 			}
 
@@ -222,10 +231,16 @@ void UPoseSearchInteractionSubsystem::RegenerateAllIslands(TConstArrayView<UE::P
 
 		void Merge(const FInteractionSearchContextGroup& SearchContextGroup)
 		{
-			for (const UObject* AnimContext : SearchContextGroup.AnimContexts)
+			for (const FAnimContextToTickPriorityPair& AnimContextToTickPriorityPair : SearchContextGroup.AnimContextToTickPriority)
 			{
-				check(AnimContext);
-				AnimContexts.Add(AnimContext);
+				if (int32* TickPriority = AnimContextToTickPriority.Find(AnimContextToTickPriorityPair.Key))
+				{
+					*TickPriority = FMath::Max(*TickPriority, AnimContextToTickPriorityPair.Value);
+				}
+				else
+				{
+					AnimContextToTickPriority.Add(AnimContextToTickPriorityPair.Key) = AnimContextToTickPriorityPair.Value;
+				}
 			}
 
 			for (int32 SearchContextsIndex : SearchContextGroup.SearchContextsIndices)
@@ -234,7 +249,12 @@ void UPoseSearchInteractionSubsystem::RegenerateAllIslands(TConstArrayView<UE::P
 			}
 		}
 
-		TSet<const UObject*, DefaultKeyFuncs<const UObject*>, TInlineSetAllocator<16, TMemStackSetAllocator<>>> AnimContexts;
+		// all the AnimContexts in this group with their TickPriority
+		typedef TPair<const UObject*, int32> FAnimContextToTickPriorityPair;
+		typedef TMap<const UObject*, int32, TInlineSetAllocator<16, TMemStackSetAllocator<>>> FAnimContextToTickPriority;
+		FAnimContextToTickPriority AnimContextToTickPriority;
+
+		// indexes to the searchcontexts assigned to this group
 		TArray<int32, TInlineAllocator<16, TMemStackAllocator<>>> SearchContextsIndices;
 	};
 
@@ -272,6 +292,7 @@ void UPoseSearchInteractionSubsystem::RegenerateAllIslands(TConstArrayView<UE::P
 		}
 	}
 
+	TArray<FInteractionSearchContextGroup::FAnimContextToTickPriorityPair, TInlineAllocator<16, TMemStackAllocator<>>> SortedByTickPriorityAnimContexts;
 	for (FInteractionSearchContextGroup& SearchContextGroup : SearchContextGroups)
 	{
 		// @todo: search for the most suitable island to reuse to avoid having to Uninitialize/RemoveTickDependencies and InjectToActor right away
@@ -290,13 +311,23 @@ void UPoseSearchInteractionSubsystem::RegenerateAllIslands(TConstArrayView<UE::P
 			Island.AddSearchContext(SearchContext);
 		}
 
-		// @todo: we sort SearchContextGroup.AnimContexts depending on how much any of them want 
-		//        to be the MainAnimContext, to avoid external conflicting tick dependencies
-		
+		// sorting SearchContextGroup.AnimContextToTickPriority by TickPriority
+		SortedByTickPriorityAnimContexts.Reset();
+		SortedByTickPriorityAnimContexts.Reserve(SearchContextGroup.AnimContextToTickPriority.Num());
+		for (const FInteractionSearchContextGroup::FAnimContextToTickPriorityPair& AnimContextToTickPriorityPair : SearchContextGroup.AnimContextToTickPriority)
+		{
+			SortedByTickPriorityAnimContexts.Add(AnimContextToTickPriorityPair);
+		}
+		SortedByTickPriorityAnimContexts.Sort(
+			[](const FInteractionSearchContextGroup::FAnimContextToTickPriorityPair& A, const FInteractionSearchContextGroup::FAnimContextToTickPriorityPair& B)
+			{
+				return B.Value < A.Value;
+			});
+
 		// injecting tick dependencies between island AnimContext
-		for (const UObject* AnimContext : SearchContextGroup.AnimContexts)
+		for (const FInteractionSearchContextGroup::FAnimContextToTickPriorityPair& SortedByTickPriorityAnimContext : SortedByTickPriorityAnimContexts)
 		{
-			Island.InjectToActor(AnimContext, bAreTickDependenciesRequired);
+			Island.InjectToActor(SortedByTickPriorityAnimContext.Key, bAreTickDependenciesRequired);
 		}
 	}
 }
@@ -395,17 +426,35 @@ void UPoseSearchInteractionSubsystem::PopulateContinuingProperties(float DeltaSe
 	}
 }
 
-UE::PoseSearch::FInteractionIsland* UPoseSearchInteractionSubsystem::FindIsland(const UObject* AnimContext)
+UE::PoseSearch::FInteractionIsland* UPoseSearchInteractionSubsystem::FindIsland(const UObject* AnimContext, bool bCompareOwningActors)
 {
 	using namespace UE::PoseSearch;
 
 	if (AnimContext)
 	{
-		for (FInteractionIsland* Island : Islands)
+		if (bCompareOwningActors)
+		{
+			const AActor* Actor = GetContextOwningActor(AnimContext);
+
+			for (FInteractionIsland* Island : Islands)
+			{
+				for (const TWeakObjectPtr<const UObject>& IslandAnimContext : Island->GetIslandAnimContexts())
+				{ 
+					if (GetContextOwningActor(IslandAnimContext.Get()) == Actor)
+					{
+						return Island;
+					}
+				}
+			}
+		}
+		else
 		{
-			if (Island->GetIslandAnimContexts().Contains(AnimContext))
+			for (FInteractionIsland* Island : Islands)
 			{
-				return Island;
+				if (Island->GetIslandAnimContexts().Contains(AnimContext))
+				{
+					return Island;
+				}
 			}
 		}
 	}
@@ -567,13 +616,15 @@ void UPoseSearchInteractionSubsystem::GenerateAnimContextInfosAndTagToDatabases(
 
 		AnimContextInfo.AnimContext = AnimContextsWithAvailabilities[AnimContextInfoIndex].Key;
 		AnimContextInfo.Availabilities = AnimContextsWithAvailabilities[AnimContextInfoIndex].Value;
-		check(AnimContextInfo.AnimContext && AnimContextInfo.Availabilities);
+		check(AnimContextInfo.AnimContext && AnimContextInfo.Availabilities && !AnimContextInfo.Availabilities->IsEmpty());
 
 		AnimContextInfo.Location = GetContextLocation(AnimContextInfo.AnimContext);
 
+		AnimContextInfo.TickPriority = -MAX_int32;
 		for (const FInteractionAvailability& Availability : *AnimContextInfo.Availabilities)
 		{
 			AnimContextInfo.MaxDistanceSquared = FMath::Max(AnimContextInfo.MaxDistanceSquared, Availability.BroadPhaseRadius);
+			AnimContextInfo.TickPriority = FMath::Max(AnimContextInfo.TickPriority, Availability.TickPriority);
 		}
 		AnimContextInfo.MaxDistanceSquared *= AnimContextInfo.MaxDistanceSquared;
 	}
@@ -859,6 +910,7 @@ void UPoseSearchInteractionSubsystem::GenerateSearchContexts(UE::PoseSearch::FIn
 							SearchContext.Roles.Add(RoledAnimContextInfo.Role);
 							SearchContext.PoseHistories.Add(RoledAnimContextInfo.PoseHistory);
 							SearchContext.bDisableCollisions |= RoledAnimContextInfo.bDisableCollisions;
+							SearchContext.TickPriorities.Add(RoledAnimContextInfo.AnimContextInfo->TickPriority);
 						}
 
 						// does SearchContext cover all the roles required by this interaction?
@@ -1165,6 +1217,20 @@ void UPoseSearchInteractionSubsystem::Query_AnyThread(const TArrayView<const FPo
 	}
 }
 
+void UPoseSearchInteractionSubsystem::GetResult_AnyThread(const UObject* AnimContext, FPoseSearchInteractionBlueprintResult& Result, bool bCompareOwningActors)
+{
+	using namespace UE::PoseSearch;
+
+	if (FInteractionIsland* Island = FindIsland(AnimContext, bCompareOwningActors))
+	{
+		Island->GetResult_AnyThread(AnimContext, Result, bCompareOwningActors);
+	}
+	else
+	{
+		Result = FPoseSearchInteractionBlueprintResult();
+	}
+}
+
 #if ENABLE_DRAW_DEBUG && ENABLE_ANIM_DEBUG
 void UPoseSearchInteractionSubsystem::DebugDrawIslands() const
 {
diff --git a/Engine/Plugins/Animation/PoseSearch/Source/Runtime/Private/PoseSearchInteractionUtils.cpp b/Engine/Plugins/Animation/PoseSearch/Source/Runtime/Private/PoseSearchInteractionUtils.cpp
new file mode 100644
--- /dev/null
+++ b/Engine/Plugins/Animation/PoseSearch/Source/Runtime/Private/PoseSearchInteractionUtils.cpp
@@ -0,0 +1,104 @@
+// Copyright Epic Games, Inc. All Rights Reserved.
+
+#include "PoseSearch/PoseSearchInteractionUtils.h"
+#include "PoseSearch/MultiAnimAsset.h"
+#include "PoseSearch/PoseSearchDatabase.h"
+#include "PoseSearch/PoseSearchInteractionLibrary.h"
+#include "PoseSearch/PoseSearchSchema.h"
+
+namespace UE::PoseSearch
+{
+	
+int32 GetRoleIndex(const UMultiAnimAsset* MultiAnimAsset, const FRole& Role)
+{
+	check(MultiAnimAsset);
+	const int32 NumRoles = MultiAnimAsset->GetNumRoles();
+	for (int32 MultiAnimAssetRoleIndex = 0; MultiAnimAssetRoleIndex < NumRoles; ++MultiAnimAssetRoleIndex)
+	{
+		if (MultiAnimAsset->GetRole(MultiAnimAssetRoleIndex) == Role)
+		{
+			return MultiAnimAssetRoleIndex;
+		}
+	}
+	return INDEX_NONE;
+}
+
+FRoleToIndex MakeRoleToIndex(const UMultiAnimAsset* MultiAnimAsset)
+{
+	check(MultiAnimAsset);
+	const int32 NumRoles = MultiAnimAsset->GetNumRoles();
+
+	FRoleToIndex RoleToIndex;
+	RoleToIndex.Reserve(NumRoles);
+	for (int32 MultiAnimAssetRoleIndex = 0; MultiAnimAssetRoleIndex < NumRoles; ++MultiAnimAssetRoleIndex)
+	{
+		RoleToIndex.Add(MultiAnimAsset->GetRole(MultiAnimAssetRoleIndex)) = MultiAnimAssetRoleIndex;
+	}
+	return RoleToIndex;
+}
+
+void CalculateFullAlignedTransforms(const FPoseSearchInteractionBlueprintResult& CurrentResult, bool bWarpUsingRootBone, TArrayView<FTransform> OutFullAlignedTransforms)
+{
+	const UMultiAnimAsset* MultiAnimAsset = CastChecked<UMultiAnimAsset>(CurrentResult.SelectedAnimation);
+	const int32 NumRoles = MultiAnimAsset->GetNumRoles();
+	
+	check(OutFullAlignedTransforms.Num() == NumRoles);
+	check(CurrentResult.ActorRootTransforms.Num() == CurrentResult.ActorRootBoneTransforms.Num());
+	check(CurrentResult.ActorRootTransforms.Num() == CurrentResult.AnimContexts.Num());
+	check(CurrentResult.SelectedDatabase != nullptr && CurrentResult.SelectedDatabase->Schema != nullptr);
+
+	TConstArrayView<FTransform> ActorTransforms = CurrentResult.ActorRootTransforms;
+	
+	// if bWarpUsingRootBone we set ActorTransforms as the actors root bone world transforms instead of the root transforms
+	TArray<FTransform, TInlineAllocator<PreallocatedRolesNum, TMemStackAllocator<>>> WorldActorRootBoneTransforms;
+	if (bWarpUsingRootBone)
+	{
+		WorldActorRootBoneTransforms.SetNum(NumRoles);
+		for (int32 Index = 0; Index < NumRoles; ++Index)
+		{
+			WorldActorRootBoneTransforms[Index] = CurrentResult.ActorRootBoneTransforms[Index] * CurrentResult.ActorRootTransforms[Index];
+		}
+
+		ActorTransforms = WorldActorRootBoneTransforms;
+	}
+
+	TArray<const UMirrorDataTable*, TInlineAllocator<PreallocatedRolesNum, TMemStackAllocator<>>> MirrorDataTables;
+	if (CurrentResult.bIsMirrored)
+	{
+		MirrorDataTables.SetNum(NumRoles);
+		for (int32 MultiAnimAssetRoleIndex = 0; MultiAnimAssetRoleIndex < NumRoles; ++MultiAnimAssetRoleIndex)
+		{
+			const FRole Role = MultiAnimAsset->GetRole(MultiAnimAssetRoleIndex);
+			const FPoseSearchRoledSkeleton* RoledSkeleton = CurrentResult.SelectedDatabase->Schema->GetRoledSkeleton(Role);
+			check(RoledSkeleton);
+			MirrorDataTables[MultiAnimAssetRoleIndex] = RoledSkeleton->MirrorDataTable;
+		}
+	}
+
+	TArray<bool, TInlineAllocator<PreallocatedRolesNum, TMemStackAllocator<>>> RelevantRoleIndexes;
+	RelevantRoleIndexes.SetNum(NumRoles);
+	for (int32 Index = 0; Index < NumRoles; ++Index)
+	{
+		RelevantRoleIndexes[Index] = CurrentResult.AnimContexts[Index] != nullptr;
+	}
+	MultiAnimAsset->CalculateWarpTransforms(CurrentResult.SelectedTime, ActorTransforms, OutFullAlignedTransforms, MirrorDataTables, RelevantRoleIndexes);
+}
+
+FTransform CalculateDeltaAlignment(const FTransform& MeshWithoutOffset, const FTransform& MeshWithOffset, const FTransform& FullAlignedTransform, float WarpingRotationRatio, float WarpingTranslationRatio)
+{
+	// calculating the NoDeltaAlignment as the delta transform that brings the actor to original mesh transform.
+	const FTransform NoDeltaAlignment = MeshWithoutOffset.GetRelativeTransform(MeshWithOffset);
+
+	// calculating the FullDeltaAlignment as the delta transform that brings the actor to its full aligned transform.
+	const FTransform FullDeltaAlignment = FullAlignedTransform.GetRelativeTransform(MeshWithOffset);
+
+	// calculating the DeltaAlignment as blend between the NoDeltaAlignment and the FullDeltaAlignment: how much the character need to move to get to the desired aligment
+	const FTransform DeltaAlignment(FMath::Lerp(NoDeltaAlignment.GetRotation(), FullDeltaAlignment.GetRotation(), FMath::Clamp(WarpingRotationRatio, 0.f, 1.f)),
+		FMath::Lerp(NoDeltaAlignment.GetTranslation(), FullDeltaAlignment.GetTranslation(), FMath::Clamp(WarpingTranslationRatio, 0.f, 1.f)), FVector::OneVector);
+
+	// NoTe: keep in mind MeshWithoutOffset, MeshWithOffset, and FullAlignedTransform are relative to the previous execution frame so we still need to 
+	//		 extract and and apply the current animation root motion transform to get to the current frame full aligned transform.
+	return DeltaAlignment;
+}
+
+} // UE::PoseSearch
diff --git a/Engine/Plugins/Animation/PoseSearch/Source/Runtime/Public/PoseSearch/AnimNode_MotionMatchingInteraction.h b/Engine/Plugins/Animation/PoseSearch/Source/Runtime/Public/PoseSearch/AnimNode_MotionMatchingInteraction.h
--- a/Engine/Plugins/Animation/PoseSearch/Source/Runtime/Public/PoseSearch/AnimNode_MotionMatchingInteraction.h
+++ b/Engine/Plugins/Animation/PoseSearch/Source/Runtime/Public/PoseSearch/AnimNode_MotionMatchingInteraction.h
@@ -23,27 +23,14 @@ struct POSESEARCH_API FAnimNode_MotionMatchingInteraction : public FAnimNode_Ble
 	UPROPERTY(EditAnywhere, Category = Warping, meta = (PinHiddenByDefault, ClampMin = "0", ClampMax = "1"))
 	float WarpingTranslationRatio = 1.f;
 
-	// Controls how fast the warp translation offset is blended in
-	// Values closer to 0 make it faster. if 0 the warp translation is applied in full
-	UPROPERTY(EditAnywhere, Category = Warping, meta = (PinHiddenByDefault))
-	float WarpingTranslationHalflife = 0.05f;
-
 	// amount or rotation warping to apply
 	UPROPERTY(EditAnywhere, Category = Warping, meta = (PinHiddenByDefault, ClampMin = "0", ClampMax = "1"))
 	float WarpingRotationRatio = 1.f;
 
-	// Controls how fast the warp rotation offset is blended in
-	// Values closer to 0 make it faster. if 0 the warp rotation is applied in full
-	UPROPERTY(EditAnywhere, Category = Warping, meta = (PinHiddenByDefault))
-	float WarpingRotationHalfLife = 0.05f;
-
-	// enable / disable the usage of UAnimNotifyState_PoseSearchWarping(s) to change the WarpingAlignmentExtractionTime (time we sample the full alignment transforms at)
-	UPROPERTY(EditAnywhere, Category = Warping, meta = (PinHiddenByDefault))
-	bool bUsePoseSearchWarpingNotifyStates = true;
-
-	// enable / disable the usage of UAnimNotifyState_PoseSearchIKWindow(s) to calculate IK constraints and feed them to the related Goals
-	UPROPERTY(EditAnywhere, Category = Warping, meta = (PinHiddenByDefault))
-	bool bUsePoseSearchIKWindowNotifyStates = true;
+	// if bWarpUsingRootBone is true, warping will be calculated using the interacting actors previous frame root bone transforms (effective for setups with OffsetRootBone node allowing root bone drifting from capsule)
+	// if bWarpUsingRootBone is true, warping will be calculated using the previous frame root transforms (effective root motion driven for setups)
+	UPROPERTY(EditAnywhere, Category = Settings)
+	bool bWarpUsingRootBone = true;
 
 	// Reset the blend stack if it has become relevant to the graph after not being updated on previous frames.
 	UPROPERTY(EditAnywhere, Category = Settings)
@@ -81,6 +68,8 @@ protected:
 	bool NeedsReset(const FAnimationUpdateContext& Context) const;
 
 	FPoseSearchInteractionBlueprintResult CurrentResult;
+	FTransform MeshWithOffset = FTransform::Identity;
+	FTransform MeshWithoutOffset = FTransform::Identity;
 
 	// Update Counter for detecting being relevant
 	FGraphTraversalCounter UpdateCounter;
diff --git a/Engine/Plugins/Animation/PoseSearch/Source/Runtime/Public/PoseSearch/PoseSearchAnimNotifies.h b/Engine/Plugins/Animation/PoseSearch/Source/Runtime/Public/PoseSearch/PoseSearchAnimNotifies.h
--- a/Engine/Plugins/Animation/PoseSearch/Source/Runtime/Public/PoseSearch/PoseSearchAnimNotifies.h
+++ b/Engine/Plugins/Animation/PoseSearch/Source/Runtime/Public/PoseSearch/PoseSearchAnimNotifies.h
@@ -121,10 +121,34 @@ public:
 };
 
 // multi character interaction notifies
-UCLASS(Experimental, Blueprintable, meta = (DisplayName = "Pose Search: Warping"))
-class POSESEARCH_API UAnimNotifyState_PoseSearchWarping : public UAnimNotifyState_PoseSearchBase
+USTRUCT(Experimental, Blueprintable)
+struct POSESEARCH_API FPoseSearchIKWindowConstraint
 {
 	GENERATED_BODY()
+	
+	UPROPERTY(EditAnywhere, Category = Config)
+	FName FromGoalName;
+
+	UPROPERTY(EditAnywhere, Category = Config)
+	FName FromGoalRole;
+	
+	UPROPERTY(EditAnywhere, Category = Config)
+	FName ToGoalName;
+	
+	UPROPERTY(EditAnywhere, Category = Config)
+	FName ToGoalRole;
+
+	UPROPERTY(EditAnywhere, Category = Config, meta = (ClampMin = "0", ClampMax = "1", UIMin = "0", UIMax = "1"))
+	float TranslationWeight = 0.f;
+	
+	UPROPERTY(EditAnywhere, Category = Config, meta = (ClampMin = "0", ClampMax = "1", UIMin = "0", UIMax = "1"))
+	float RotationWeight = 0.f;
+	
+	UPROPERTY(EditAnywhere, Category = Config, meta = (ClampMin = "0", UIMin = "0"))
+	float ActivationTime = 0.f;
+	
+	UPROPERTY(EditAnywhere, Category = Config, meta = (ClampMin = "0", UIMin = "0"))
+	float DeactivationTime = 0.f;
 };
 
 UCLASS(Experimental, Blueprintable, meta = (DisplayName = "Pose Search: IKWindow"))
@@ -133,9 +157,8 @@ class POSESEARCH_API UAnimNotifyState_PoseSearchIKWindow : public UAnimNotifySta
 	GENERATED_BODY()
 
 public:
-	// list of IK goals to update with transforms from the fully aligned pose
 	UPROPERTY(EditAnywhere, Category = Config)
-	TArray<FName> GoalNames;
+	TArray<FPoseSearchIKWindowConstraint> Constraints;
 };
 
 UCLASS(Experimental, Blueprintable, meta = (DisplayName = "Pose Search: Event"))
diff --git a/Engine/Plugins/Animation/PoseSearch/Source/Runtime/Public/PoseSearch/PoseSearchAssetSamplerLibrary.h b/Engine/Plugins/Animation/PoseSearch/Source/Runtime/Public/PoseSearch/PoseSearchAssetSamplerLibrary.h
--- a/Engine/Plugins/Animation/PoseSearch/Source/Runtime/Public/PoseSearch/PoseSearchAssetSamplerLibrary.h
+++ b/Engine/Plugins/Animation/PoseSearch/Source/Runtime/Public/PoseSearch/PoseSearchAssetSamplerLibrary.h
@@ -3,10 +3,11 @@
 #pragma once
 
 #include "Animation/AnimCurveTypes.h"
+#include "DrawDebugHelpers.h"
 #include "Kismet/BlueprintFunctionLibrary.h"
 #include "PoseSearch/PoseSearchAssetSampler.h"
 #include "PoseSearch/PoseSearchDefines.h"
-#include "DrawDebugHelpers.h"
+#include "VisualLogger/VisualLogger.h"
 #include "PoseSearchAssetSamplerLibrary.generated.h"
 
 USTRUCT(Experimental, BlueprintType, Category="Animation|Pose Search")
@@ -123,7 +124,7 @@ public:
 
 #if ENABLE_DRAW_DEBUG
 	template<typename FComponentSpacePoseType>
-	static void Draw(const UWorld* World, FComponentSpacePoseType& ComponentSpacePose, const FTransform& RootTransform, float DebugDrawSamplerRootAxisLength = 20.f, float DebugDrawSamplerSize = 6.f)
+	static void Draw(const UWorld* World, FComponentSpacePoseType& ComponentSpacePose, const FTransform& RootTransform, const FColor Color = FColor::Red, float DebugDrawSamplerRootAxisLength = 20.f, float DebugDrawSamplerSize = 6.f)
 	{
 		if (World)
 		{
@@ -138,10 +139,41 @@ public:
 			for (int32 BoneIndex = 0; BoneIndex < ComponentSpacePose.GetPose().GetNumBones(); ++BoneIndex)
 			{
 				const FTransform BoneWorldTransform = GetTransform(ComponentSpacePose, RootTransform, FCompactPoseBoneIndex(BoneIndex), EPoseSearchAssetSamplerSpace::World);
-				DrawDebugPoint(World, BoneWorldTransform.GetTranslation(), DebugDrawSamplerSize, FColor::Red, false, 0.f, SDPG_Foreground);
+				DrawDebugPoint(World, BoneWorldTransform.GetTranslation(), DebugDrawSamplerSize, Color, false, 0.f, SDPG_Foreground);
 			}
 		}
 	}
 #endif // ENABLE_DRAW_DEBUG
+
+#if ENABLE_VISUAL_LOG
+	template<typename FComponentSpacePoseType>
+	static void VLogDraw(const UObject* VLogContext, FComponentSpacePoseType& ComponentSpacePose, const FTransform& RootTransform, const TCHAR* VLogName, const FColor Color = FColor::Red, float DebugDrawSamplerRootAxisLength = 20.f)
+	{
+		check(VLogName);
+		if (VLogContext)
+		{
+			if (DebugDrawSamplerRootAxisLength > 0.f)
+			{
+				const FTransform AxisWorldTransform = GetTransform(ComponentSpacePose, RootTransform);
+				UE_VLOG_SEGMENT(VLogContext, VLogName, Display, AxisWorldTransform.GetTranslation(), AxisWorldTransform.GetTranslation() + AxisWorldTransform.GetScaledAxis(EAxis::X) * DebugDrawSamplerRootAxisLength, FColor::Red, TEXT(""));
+				UE_VLOG_SEGMENT(VLogContext, VLogName, Display, AxisWorldTransform.GetTranslation(), AxisWorldTransform.GetTranslation() + AxisWorldTransform.GetScaledAxis(EAxis::Y) * DebugDrawSamplerRootAxisLength, FColor::Green, TEXT(""));
+				UE_VLOG_SEGMENT(VLogContext, VLogName, Display, AxisWorldTransform.GetTranslation(), AxisWorldTransform.GetTranslation() + AxisWorldTransform.GetScaledAxis(EAxis::Z) * DebugDrawSamplerRootAxisLength, FColor::Blue, TEXT(""));
+			}
+
+			for (int32 BoneIndex = 0; BoneIndex < ComponentSpacePose.GetPose().GetNumBones(); ++BoneIndex)
+			{
+				const FCompactPoseBoneIndex CompactPoseBoneIndex(BoneIndex);
+				const FCompactPoseBoneIndex ParentCompactPoseBoneIndex = ComponentSpacePose.GetPose().GetParentBoneIndex(CompactPoseBoneIndex);
+				
+				const FTransform BoneWorldTransform = GetTransform(ComponentSpacePose, RootTransform, CompactPoseBoneIndex, EPoseSearchAssetSamplerSpace::World);
+				const FTransform ParentBoneWorldTransform = GetTransform(ComponentSpacePose, RootTransform, ParentCompactPoseBoneIndex, EPoseSearchAssetSamplerSpace::World);
+				UE_VLOG_SEGMENT(VLogContext, VLogName, Display, BoneWorldTransform.GetTranslation(), ParentBoneWorldTransform.GetTranslation(), Color, TEXT(""));
+			}
+		}
+	}
+
+	static void VLogDraw(const UObject* VLogContext, const USkeletalMeshComponent* Mesh, const TCHAR* VLogName, const FColor Color = FColor::Red, float DebugDrawSamplerRootAxisLength = 20.f);
+
+#endif // ENABLE_VISUAL_LOG
 };
 
diff --git a/Engine/Plugins/Animation/PoseSearch/Source/Runtime/Public/PoseSearch/PoseSearchContext.h b/Engine/Plugins/Animation/PoseSearch/Source/Runtime/Public/PoseSearch/PoseSearchContext.h
--- a/Engine/Plugins/Animation/PoseSearch/Source/Runtime/Public/PoseSearch/PoseSearchContext.h
+++ b/Engine/Plugins/Animation/PoseSearch/Source/Runtime/Public/PoseSearch/PoseSearchContext.h
@@ -23,6 +23,11 @@ const USkeleton* GetContextSkeleton(const UObject* AnimContext);
 const AActor* GetContextOwningActor(const UObject* AnimContext);
 FVector GetContextLocation(const UObject* AnimContext);
 
+// Experimental, this feature might be removed without warning, not for production use
+const USkeletalMeshComponent* GetContextSkeletalMeshComponent(const UObject* AnimContext);
+// Experimental, this feature might be removed without warning, not for production use
+const FBoneContainer GetBoneContainer(const UObject* AnimContext);
+
 struct FPoseIndicesHistory;
 struct IPoseHistory;
 
diff --git a/Engine/Plugins/Animation/PoseSearch/Source/Runtime/Public/PoseSearch/PoseSearchDefines.h b/Engine/Plugins/Animation/PoseSearch/Source/Runtime/Public/PoseSearch/PoseSearchDefines.h
--- a/Engine/Plugins/Animation/PoseSearch/Source/Runtime/Public/PoseSearch/PoseSearchDefines.h
+++ b/Engine/Plugins/Animation/PoseSearch/Source/Runtime/Public/PoseSearch/PoseSearchDefines.h
@@ -53,6 +53,7 @@ static constexpr float FiniteDelta = 1 / 60.0f; // Time delta used for computing
 static constexpr int32 MaxNumberOfCollectedPoseCandidatesPerDatabase = 256;
 static constexpr int32 InvalidSearchId = INDEX_NONE; // @todo: it should be 0, since the search id is an hash!
 static constexpr int32 DatabaseIndexDerivedDataCacheKeyVersion = 2677; // bump this value to force reindex ALL the UPoseSearchDatabases, without having to change the global FDevSystemGuids::POSESEARCHDB_DERIVEDDATA_VER
+static constexpr int32 DefaultTickPriority = 100;
 
 template<uint32 Alignment = DEFAULT_ALIGNMENT> class TMemStackSetAllocator : public TSetAllocator<TSparseArrayAllocator<TMemStackAllocator<Alignment>, TMemStackAllocator<Alignment>>, TMemStackAllocator<Alignment>> {};
 
diff --git a/Engine/Plugins/Animation/PoseSearch/Source/Runtime/Public/PoseSearch/PoseSearchInteractionIsland.h b/Engine/Plugins/Animation/PoseSearch/Source/Runtime/Public/PoseSearch/PoseSearchInteractionIsland.h
--- a/Engine/Plugins/Animation/PoseSearch/Source/Runtime/Public/PoseSearch/PoseSearchInteractionIsland.h
+++ b/Engine/Plugins/Animation/PoseSearch/Source/Runtime/Public/PoseSearch/PoseSearchInteractionIsland.h
@@ -46,6 +46,7 @@ struct FInteractionSearchContext : FInteractionSearchContextBase
 	TWeakObjectPtr<const UObject> PlayingAsset = nullptr;
 	float PlayingAssetAccumulatedTime = 0.f;
 	EPoseSearchInterruptMode InterruptMode = EPoseSearchInterruptMode::DoNotInterrupt;
+	TArray<int32, TInlineAllocator<PreallocatedRolesNum>> TickPriorities;
 
 	FPoseSearchContinuingProperties GetContinuingProperties() const
 	{
@@ -88,6 +89,7 @@ struct FInteractionIsland
 	~FInteractionIsland();
 
 	bool DoSearch_AnyThread(const UObject* AnimContext, const FPoseSearchContinuingProperties& ContinuingProperties, FPoseSearchInteractionBlueprintResult& Result);
+	bool GetResult_AnyThread(const UObject* AnimContext, FPoseSearchInteractionBlueprintResult& Result, bool bCompareOwningActors = false);
 
 	const TArray<TWeakObjectPtr<UActorComponent>>& GetTickActorComponents() const { return TickActorComponents; }
 	const TArray<TWeakObjectPtr<const UObject>>& GetIslandAnimContexts() const { return IslandAnimContexts; }
@@ -112,7 +114,6 @@ private:
 	void RemoveTickDependencies(bool bValidateTickDependencies);
 	IInteractionIslandDependency* FindCustomDependency(UActorComponent* InTickComponent) const;
 
-	bool GetResult_AnyThread(const UObject* AnimContext, FPoseSearchInteractionBlueprintResult& Result);
 	const UObject* GetMainAnimContext() const;
 	const AActor* GetMainActor() const;
 	bool CanPerformSearch(const UObject* AnimContext) const;
diff --git a/Engine/Plugins/Animation/PoseSearch/Source/Runtime/Public/PoseSearch/PoseSearchInteractionLibrary.h b/Engine/Plugins/Animation/PoseSearch/Source/Runtime/Public/PoseSearch/PoseSearchInteractionLibrary.h
--- a/Engine/Plugins/Animation/PoseSearch/Source/Runtime/Public/PoseSearch/PoseSearchInteractionLibrary.h
+++ b/Engine/Plugins/Animation/PoseSearch/Source/Runtime/Public/PoseSearch/PoseSearchInteractionLibrary.h
@@ -4,7 +4,6 @@
 
 #include "Kismet/BlueprintFunctionLibrary.h"
 #include "PoseSearch/PoseSearchInteractionIsland.h"
-#include "PoseSearch/PoseSearchLibrary.h"
 #include "PoseSearchInteractionLibrary.generated.h"
 
 // input for MotionMatchInteraction_Pure: it declares that the associated character ("AnimContext" that could be an AnimInstance or an AnimNextCharacterComponent)
@@ -43,6 +42,12 @@ struct POSESEARCH_API FPoseSearchInteractionAvailability
 	UPROPERTY(EditAnywhere, BlueprintReadWrite, Category=Settings)
 	bool bDisableCollisions = false;
 
+	// the Actor with the higher TickPriority of any Availability request will be elected as the MainActor of the interaction island (containing all the actors that could interact with each other)
+	// the main Actor will tick first and all the other interacting actors will tick after in a concurrently from each other. TickPriority is useful if your setup is already enforcing tick dependency 
+	// between actors, and the motion matching interaction system needs to play nicely with them.
+	UPROPERTY(EditAnywhere, BlueprintReadWrite, Category=Settings)
+	int32 TickPriority = UE::PoseSearch::DefaultTickPriority;
+
 	bool IsTagValid() const { return !Tag.IsNone(); }
 };
 
@@ -93,13 +98,17 @@ public:
 	UPROPERTY(Transient, VisibleAnywhere, BlueprintReadOnly, Category=State)
 	FName Role;
 
-	// cached actors root transforms for all the roles in SelectedAnimation (as UMultiAnimAsset)
+	// cached actors root transforms for all the roles in SelectedAnimation (as UMultiAnimAsset). If the role is missing because optional the ActorRootTransform will be identity
 	UPROPERTY(Transient, VisibleAnywhere, BlueprintReadOnly, Category=State)
 	TArray<FTransform> ActorRootTransforms;
 
-	// cached actors root bone transforms for all the roles in SelectedAnimation (as UMultiAnimAsset)
+	// cached actors root bone transforms for all the roles in SelectedAnimation (as UMultiAnimAsset). If the role is missing because optional the ActorRootBoneTransform will be identity
 	UPROPERTY(Transient, VisibleAnywhere, BlueprintReadOnly, Category=State)
 	TArray<FTransform> ActorRootBoneTransforms;
+
+	// AnimContexts involved in this interaction for all the roles in SelectedAnimation (as UMultiAnimAsset). If the role is missing because optional the AnimContext will be null
+	// BE CAREFUL HOW YOU USE THEM! since accessing data from other actors without proper tick dependencies leads to threading problems!
+	TArray<TWeakObjectPtr<const UObject>> AnimContexts;
 };
 
 // Experimental, this feature might be removed without warning, not for production use
@@ -130,6 +139,8 @@ public:
 	// version of MotionMatchInteraction_Pure referencing directly the HistoryCollector rather than looking for it by name
 	static FPoseSearchInteractionBlueprintResult MotionMatchInteraction(const TArrayView<const FPoseSearchInteractionAvailability> Availabilities, const UObject* AnimContext, const FPoseSearchContinuingProperties& ContinuingProperties, const UE::PoseSearch::IPoseHistory* PoseHistory, bool bValidateResultAgainstAvailabilities = true);
 
+	static FPoseSearchInteractionBlueprintResult GetCachedInteraction(const UObject* AnimContext, bool bCompareOwningActors = false);
+	
 	UFUNCTION(Experimental, BlueprintPure, Category = "Animation|Pose Search", meta = (BlueprintThreadSafe))
 	static FPoseSearchContinuingProperties GetMontageContinuingProperties(UAnimInstance* AnimInstance);
 };
\ No newline at end of file
diff --git a/Engine/Plugins/Animation/PoseSearch/Source/Runtime/Public/PoseSearch/PoseSearchInteractionSubsystem.h b/Engine/Plugins/Animation/PoseSearch/Source/Runtime/Public/PoseSearch/PoseSearchInteractionSubsystem.h
--- a/Engine/Plugins/Animation/PoseSearch/Source/Runtime/Public/PoseSearch/PoseSearchInteractionSubsystem.h
+++ b/Engine/Plugins/Animation/PoseSearch/Source/Runtime/Public/PoseSearch/PoseSearchInteractionSubsystem.h
@@ -145,6 +145,8 @@ public:
 		const FPoseSearchContinuingProperties& ContinuingProperties, FPoseSearchInteractionBlueprintResult& Result,
 		FName PoseHistoryName, const UE::PoseSearch::IPoseHistory* PoseHistory, bool bValidateResultAgainstAvailabilities);
 
+	void GetResult_AnyThread(const UObject* AnimContext, FPoseSearchInteractionBlueprintResult& Result, bool bCompareOwningActors = false);
+
 	// returns and array containing all the other AnimContext(s) that the input AnimContext is currently interacting with
 	TArray<const UObject*> GetInteractingAnimContextsWith(const UObject* AnimContext) const;
 
@@ -167,7 +169,7 @@ private:
 #endif // DO_CHECK
 
 	void PopulateContinuingProperties(float DeltaSeconds, TArrayView<UE::PoseSearch::FInteractionSearchContext> SearchContexts) const;
-	UE::PoseSearch::FInteractionIsland* FindIsland(const UObject* AnimContext);
+	UE::PoseSearch::FInteractionIsland* FindIsland(const UObject* AnimContext, bool bCompareOwningActors = false);
 	
 #if ENABLE_DRAW_DEBUG && ENABLE_ANIM_DEBUG
 	void DebugDrawIslands() const;
diff --git a/Engine/Plugins/Animation/PoseSearch/Source/Runtime/Public/PoseSearch/PoseSearchInteractionUtils.h b/Engine/Plugins/Animation/PoseSearch/Source/Runtime/Public/PoseSearch/PoseSearchInteractionUtils.h
new file mode 100644
--- /dev/null
+++ b/Engine/Plugins/Animation/PoseSearch/Source/Runtime/Public/PoseSearch/PoseSearchInteractionUtils.h
@@ -0,0 +1,25 @@
+// Copyright Epic Games, Inc. All Rights Reserved.
+
+#pragma once
+
+#include "PoseSearch/PoseSearchRole.h"
+
+struct FPoseSearchInteractionBlueprintResult;
+class UMultiAnimAsset;
+
+namespace UE::PoseSearch
+{
+
+// Experimental, this feature might be removed without warning, not for production use
+POSESEARCH_API int32 GetRoleIndex(const UMultiAnimAsset* MultiAnimAsset, const FRole& Role);
+
+// Experimental, this feature might be removed without warning, not for production use
+FRoleToIndex MakeRoleToIndex(const UMultiAnimAsset* MultiAnimAsset);
+
+// Experimental, this feature might be removed without warning, not for production use
+POSESEARCH_API void CalculateFullAlignedTransforms(const FPoseSearchInteractionBlueprintResult& CurrentResult, bool bWarpUsingRootBone, TArrayView<FTransform> OutFullAlignedTransforms);
+
+// Experimental, this feature might be removed without warning, not for production use
+POSESEARCH_API FTransform CalculateDeltaAlignment(const FTransform& MeshWithoutOffset, const FTransform& MeshWithOffset, const FTransform& FullAlignedTransform, float WarpingRotationRatio, float WarpingTranslationRatio);
+
+} // namespace UE::PoseSearch
\ No newline at end of file
