Addressing compiling issues

#rnx
#pf skip

diff --git a/Engine/Plugins/Animation/PoseSearch/Source/Runtime/Private/PoseSearchTrajectoryTypes.cpp b/Engine/Plugins/Animation/PoseSearch/Source/Runtime/Private/PoseSearchTrajectoryTypes.cpp
--- a/Engine/Plugins/Animation/PoseSearch/Source/Runtime/Private/PoseSearchTrajectoryTypes.cpp
+++ b/Engine/Plugins/Animation/PoseSearch/Source/Runtime/Private/PoseSearchTrajectoryTypes.cpp
@@ -234,7 +234,7 @@ FPoseSearchQueryTrajectory::operator FTransformTrajectory() const
 
 	for (const FPoseSearchQueryTrajectorySample& Sample : Samples)
 	{
-		OutTrajectory.Samples.Add(FTransformTrajectorySample(Sample.Facing, Sample.Position, Sample.AccumulatedSeconds));
+		OutTrajectory.Samples.Add({ Sample.Facing, Sample.Position, Sample.AccumulatedSeconds });
 	}
 
 	return OutTrajectory;
