Pose Search: Fix for bad blendspace detection for using normalized time

diff --git a/Engine/Plugins/Animation/PoseSearch/Source/Runtime/Private/PoseSearchDatabase.cpp b/Engine/Plugins/Animation/PoseSearch/Source/Runtime/Private/PoseSearchDatabase.cpp
--- a/Engine/Plugins/Animation/PoseSearch/Source/Runtime/Private/PoseSearchDatabase.cpp
+++ b/Engine/Plugins/Animation/PoseSearch/Source/Runtime/Private/PoseSearchDatabase.cpp
@@ -306,8 +306,7 @@ static bool IsDatabaseAssetIndexBlendSpace(const UPoseSearchDatabase& Database,
 	}
 	else if (const FPoseSearchDatabaseMultiAnimAsset* MultiAnimAsset =  Database.GetDatabaseAnimationAsset<FPoseSearchDatabaseMultiAnimAsset>(AnimationAssetIndex))
 	{
-		MultiAnimAsset->GetAnimationAssetForRole(Database.Schema->GetDefaultRole());
-		return true;
+		return MultiAnimAsset->GetAnimationAssetForRole(Database.Schema->GetDefaultRole())->IsA<UBlendSpace>();
 	}
 
 	return false;
