PoseSearch - derived data key generation improvements

diff --git a/Engine/Plugins/Animation/PoseSearch/Source/Runtime/Private/PoseSearchDerivedDataKey.cpp b/Engine/Plugins/Animation/PoseSearch/Source/Runtime/Private/PoseSearchDerivedDataKey.cpp
--- a/Engine/Plugins/Animation/PoseSearch/Source/Runtime/Private/PoseSearchDerivedDataKey.cpp
+++ b/Engine/Plugins/Animation/PoseSearch/Source/Runtime/Private/PoseSearchDerivedDataKey.cpp
@@ -4,18 +4,19 @@
 
 #include "PoseSearch/PoseSearchDerivedDataKey.h"
 #include "Animation/AnimBoneCompressionSettings.h"
+#include "Animation/AnimCompress.h"
 #include "Animation/AnimCurveCompressionSettings.h"
 #include "Animation/AnimData/IAnimationDataModel.h"
+#include "Animation/IAnimationSequenceCompiler.h"
 #include "AnimationModifier.h"
-#include "Animation/AnimCompress.h"
+#include "Components/SkeletalMeshComponent.h"
 #include "Engine/SkeletalMesh.h"
 #include "Engine/SkinnedAsset.h"
 #include "Engine/StreamableRenderAsset.h"
 #include "Factories/FbxAnimSequenceImportData.h"
-#include "Components/SkeletalMeshComponent.h"
+#include "PoseSearch/PoseSearchAnimNotifies.h"
 #include "PoseSearch/PoseSearchDefines.h"
 #include "UObject/DevObjectVersion.h"
-#include "Animation/IAnimationSequenceCompiler.h"
 
 namespace UE::PoseSearch
 {
@@ -183,8 +184,10 @@ FArchive& FKeyBuilder::TryAddDependency(UObject* Object, bool bAddToPartialKeyHa
 		}
 	}
 
-	if (UAnimSequence* AnimSequence = Cast<UAnimSequence>(Object))
+	if (Object->IsA<UAnimSequence>())
 	{
+		UAnimSequence* AnimSequence = Cast<UAnimSequence>(Object);
+		check(AnimSequence);
 		if (!AnimSequence->CanBeCompressed())
 		{
 			bAnyAssetNotFullyLoaded = true;
@@ -340,33 +343,40 @@ bool FKeyBuilder::AnyAssetNotFullyLoaded() const
 
 bool FKeyBuilder::AnyAssetNotReady() const
 {
-	bool bAnyAssetNotReady = false;
 	TArray<UAnimSequence*, TInlineAllocator<64>> SequencesToWaitFor;	
-	const ITargetPlatform* TargetPlatform = GetTargetPlatformManager()->GetRunningTargetPlatform();
+	const ITargetPlatform* TargetPlatform = nullptr;
 	for (const UObject* Dependency : Dependencies)
 	{
-		if (UAnimSequence* AnimSequence = Cast<UAnimSequence>(const_cast<UObject*>(Dependency)))
+		if (Dependency->IsA<UAnimSequence>())
 		{
-			constexpr bool bForceResidency = true;
-			if (bForceResidency)
+			UAnimSequence* AnimSequence = Cast<UAnimSequence>(const_cast<UObject*>(Dependency));
+			check(AnimSequence);
+
+			// initializing TargetPlatform lazily when the first UAnimSequence has been found
+			if (SequencesToWaitFor.IsEmpty())
 			{
-				AnimSequence->RequestResidency(TargetPlatform, GetTypeHash(KeyOwner));
-				SequencesToWaitFor.Add(AnimSequence);
+				TargetPlatform = GetTargetPlatformManager()->GetRunningTargetPlatform();
 			}
+
+			AnimSequence->RequestResidency(TargetPlatform, GetTypeHash(KeyOwner));
+			SequencesToWaitFor.Add(AnimSequence);
 		}
 	}
 
-	Anim::IAnimSequenceCompilingManager::FinishCompilation(MakeArrayView(SequencesToWaitFor));
-
-	for (UAnimSequence* AnimSequence : SequencesToWaitFor)
+	if (!SequencesToWaitFor.IsEmpty())
 	{
-		if(!AnimSequence->HasCompressedDataForPlatform(TargetPlatform))
-		{						
-			bAnyAssetNotReady = true;
+		Anim::IAnimSequenceCompilingManager::FinishCompilation(MakeArrayView(SequencesToWaitFor));
+
+		for (UAnimSequence* AnimSequence : SequencesToWaitFor)
+		{
+			if (!AnimSequence->HasCompressedDataForPlatform(TargetPlatform))
+			{
+				return true;
+			}
 		}
 	}
 	
-	return bAnyAssetNotReady;
+	return false;
 }
 
 FIoHash FKeyBuilder::Finalize() const
@@ -393,22 +403,40 @@ const TSet<const UObject*>& FKeyBuilder::GetDependencies() const
 // to keep the key generation lightweight, we don't hash these types
 bool FKeyBuilder::IsExcludedType(class UObject* Object)
 {
-	return
-		nullptr != Cast<UAnimationModifier>(Object);
+	if (Object->IsA<UAnimationModifier>())
+	{
+		return true;
+	}
+		
+	// excluding ALL the UAnimNotifyState(s) except the PoseSearch ones
+	if (Object->IsA<UAnimNotifyState>() && !Object->IsA<UAnimNotifyState_PoseSearchBase>())
+	{
+		return true;
+	}
+
+	// excluding ALL the UAnimNotify(s) except the PoseSearch ones
+	if (Object->IsA<UAnimNotify>() && !Object->IsA<UAnimNotify_PoseSearchBase>())
+	{
+		return true;
+	}
+
+	return false;
 }
 
 // to keep the key generation lightweight, we hash only the full names for these types. Object(s) will be added to Dependencies
 bool FKeyBuilder::IsAddNameOnlyType(class UObject* Object)
 {
+	check(Object);
+
 	return
-		nullptr != Cast<IAnimationDataModel>(Object) ||
-		nullptr != Cast<UActorComponent>(Object) ||
-		nullptr != Cast<UAnimBoneCompressionSettings>(Object) ||
-		nullptr != Cast<UAnimCurveCompressionSettings>(Object) ||
-		nullptr != Cast<UAssetImportData>(Object) ||
-		nullptr != Cast<UFunction>(Object) ||
-		nullptr != Cast<USkeletalMesh>(Object) ||
-		nullptr != Cast<UStreamableRenderAsset>(Object);
+		Object->IsA<UActorComponent>() ||
+		Object->IsA<UAnimBoneCompressionSettings>() ||
+		Object->IsA<UAnimCurveCompressionSettings>() ||
+		Object->IsA<UAssetImportData>() ||
+		Object->IsA<UFunction>() ||
+		Object->IsA<USkeletalMesh>() ||
+		Object->IsA<UStreamableRenderAsset>() ||
+		nullptr != Cast<IAnimationDataModel>(Object);
 }
 
 bool FKeyBuilder::ValidateAgainst(const FKeyBuilder& Other) const
diff --git a/Engine/Plugins/Animation/PoseSearch/Source/Runtime/Public/PoseSearch/PoseSearchAnimNotifies.h b/Engine/Plugins/Animation/PoseSearch/Source/Runtime/Public/PoseSearch/PoseSearchAnimNotifies.h
--- a/Engine/Plugins/Animation/PoseSearch/Source/Runtime/Public/PoseSearch/PoseSearchAnimNotifies.h
+++ b/Engine/Plugins/Animation/PoseSearch/Source/Runtime/Public/PoseSearch/PoseSearchAnimNotifies.h
@@ -10,7 +10,7 @@
 
 class UPoseSearchDatabase;
 
-// Base class for pose search anim notify states
+// Base class for ALL pose search anim notify states
 UCLASS(Abstract)
 class POSESEARCH_API UAnimNotifyState_PoseSearchBase : public UAnimNotifyState
 {
@@ -161,8 +161,15 @@ public:
 	TArray<FPoseSearchIKWindowConstraint> Constraints;
 };
 
+// Base class for ALL pose search anim notifies
+UCLASS(Abstract)
+class POSESEARCH_API UAnimNotify_PoseSearchBase : public UAnimNotify
+{
+	GENERATED_BODY()
+};
+
 UCLASS(Experimental, Blueprintable, meta = (DisplayName = "Pose Search: Event"))
-class POSESEARCH_API UAnimNotify_PoseSearchEvent : public UAnimNotify
+class POSESEARCH_API UAnimNotify_PoseSearchEvent : public UAnimNotify_PoseSearchBase
 {
 	GENERATED_BODY()
 
