posesearch
- MM trait warping support
- thread safety checks

diff --git a/Engine/Plugins/Animation/PoseSearch/Source/Runtime/Private/AnimNode_MotionMatchingInteraction.cpp b/Engine/Plugins/Animation/PoseSearch/Source/Runtime/Private/AnimNode_MotionMatchingInteraction.cpp
--- a/Engine/Plugins/Animation/PoseSearch/Source/Runtime/Private/AnimNode_MotionMatchingInteraction.cpp
+++ b/Engine/Plugins/Animation/PoseSearch/Source/Runtime/Private/AnimNode_MotionMatchingInteraction.cpp
@@ -50,15 +50,9 @@ void FAnimNode_MotionMatchingInteraction::UpdateAssetPlayer(const FAnimationUpda
 		Reset();
 	}
 
-	MeshWithoutOffset = GetContextTransform(Context.AnimInstanceProxy->GetAnimInstanceObject());
-	if (UE::AnimationWarping::FRootOffsetProvider* RootOffsetProvider = Context.GetMessage<UE::AnimationWarping::FRootOffsetProvider>())
-	{
-		MeshWithOffset = RootOffsetProvider->GetRootTransform();
-	}
-	else
-	{
-		MeshWithOffset = MeshWithoutOffset;
-	}
+	const UE::AnimationWarping::FRootOffsetProvider* RootOffsetProvider = Context.GetMessage<UE::AnimationWarping::FRootOffsetProvider>();
+	MeshWithoutOffset = Context.AnimInstanceProxy->GetComponentTransform();
+	MeshWithOffset = RootOffsetProvider ? RootOffsetProvider->GetRootTransform() : MeshWithoutOffset;
 
 	UpdateCounter.SynchronizeWith(Context.AnimInstanceProxy->GetUpdateCounter());
 	CachedDeltaTime = Context.GetDeltaTime();
diff --git a/Engine/Plugins/Animation/PoseSearch/Source/Runtime/Private/PoseSearchAssetSamplerLibrary.cpp b/Engine/Plugins/Animation/PoseSearch/Source/Runtime/Private/PoseSearchAssetSamplerLibrary.cpp
--- a/Engine/Plugins/Animation/PoseSearch/Source/Runtime/Private/PoseSearchAssetSamplerLibrary.cpp
+++ b/Engine/Plugins/Animation/PoseSearch/Source/Runtime/Private/PoseSearchAssetSamplerLibrary.cpp
@@ -89,6 +89,8 @@ FTransform UPoseSearchAssetSamplerLibrary::GetTransformByName(UPARAM(ref) FPoseS
 
 void UPoseSearchAssetSamplerLibrary::Draw(const UAnimInstance* AnimInstance, UPARAM(ref) FPoseSearchAssetSamplerPose& AssetSamplerPose)
 {
+	check(IsInGameThread());
+
 #if ENABLE_DRAW_DEBUG
 	static float DebugDrawSamplerRootAxisLength = 20.f;
 	static float DebugDrawSamplerSize = 6.f;
@@ -103,6 +105,8 @@ void UPoseSearchAssetSamplerLibrary::Draw(const UAnimInstance* AnimInstance, UPA
 #if ENABLE_VISUAL_LOG
 void UPoseSearchAssetSamplerLibrary::VLogDraw(const UObject* VLogContext, const USkeletalMeshComponent* Mesh, const TCHAR* VLogName, const FColor Color, float DebugDrawSamplerRootAxisLength)
 {
+	check(IsInGameThread());
+
 	if (!Mesh)
 	{
 		return;
diff --git a/Engine/Plugins/Animation/PoseSearch/Source/Runtime/Private/PoseSearchContext.cpp b/Engine/Plugins/Animation/PoseSearch/Source/Runtime/Private/PoseSearchContext.cpp
--- a/Engine/Plugins/Animation/PoseSearch/Source/Runtime/Private/PoseSearchContext.cpp
+++ b/Engine/Plugins/Animation/PoseSearch/Source/Runtime/Private/PoseSearchContext.cpp
@@ -18,6 +18,9 @@ namespace UE::PoseSearch
 	
 const FTransform& GetContextTransform(const UObject* AnimContext)
 {
+	// thread unsafe code. making sure we're on game thread!
+	check(IsInGameThread());
+
 	if (const UAnimInstance* AnimInstance = Cast<UAnimInstance>(AnimContext))
 	{
 		return AnimInstance->GetSkelMeshComponent()->GetComponentTransform();
@@ -42,6 +45,9 @@ const FTransform& GetContextTransform(const UObject* AnimContext)
 
 const USkeleton* GetContextSkeleton(const UObject* AnimContext)
 {
+	// @todo: make sure this is not called from worker threads!
+	// check(IsInGameThread());
+
 	if (const UAnimInstance* AnimInstance = Cast<UAnimInstance>(AnimContext))
 	{
 		return AnimInstance->GetRequiredBonesOnAnyThread().GetSkeletonAsset();
@@ -68,6 +74,9 @@ const USkeleton* GetContextSkeleton(const UObject* AnimContext)
 
 const USkeletalMeshComponent* GetContextSkeletalMeshComponent(const UObject* AnimContext)
 {
+	// thread unsafe code. making sure we're on game thread!
+	check(IsInGameThread());
+
 	if (const UAnimInstance* AnimInstance = Cast<UAnimInstance>(AnimContext))
 	{
 		return AnimInstance->GetSkelMeshComponent();
@@ -86,6 +95,9 @@ const USkeletalMeshComponent* GetContextSkeletalMeshComponent(const UObject* Ani
 
 const FBoneContainer GetBoneContainer(const UObject* AnimContext)
 {
+	// thread unsafe code. making sure we're on game thread!
+	check(IsInGameThread());
+
 	if (const UAnimInstance* AnimInstance = Cast<UAnimInstance>(AnimContext))
 	{
 		return AnimInstance->GetRequiredBonesOnAnyThread();
@@ -847,35 +859,39 @@ FTransform FSearchContext::GetWorldRootBoneTransformAtTime(float SampleTime, con
 			return PoseHistoryTransformAtTime;
 		}
 
+		// @ todo: should we use the component transform here in case we lack the root bone transform?
 		// NoTe: when GetTransformAtTime returns false PoseHistoryTransformAtTime doesn't contain a valid root bone transform in world space, BUT a valid ComponentToWorldTransform!
 		// So we use PoseHistoryTransformAtTime as such below (in the form of "const FTransform& ComponentToWorldTransform = PoseHistory ? PoseHistoryTransformAtTime : ...)
 	}
 
-	const UObject* AnimContext = GetAnimContext(SampleRole);
-	if (const UAnimInstance* AnimInstance = Cast<UAnimInstance>(AnimContext))
-	{
-		if (AnimInstance && AnimInstance->CurrentSkeleton)
-		{
-			const FTransform& RootBoneTransform = AnimInstance->CurrentSkeleton->GetReferenceSkeleton().GetRefBonePose()[RootSchemaBoneIdx];
-			const FTransform& ComponentToWorldTransform = PoseHistory ? PoseHistoryTransformAtTime : AnimInstance->GetSkelMeshComponent()->GetComponentTransform();
-			return RootBoneTransform * ComponentToWorldTransform;
-		}
-	}
-	else if (const UActorComponent* ActorComponent = Cast<UActorComponent>(AnimContext))
-	{
-		const AActor* Actor = ActorComponent->GetOwner();
-		check(Actor);
-		// @todo: this code depends on how AnimNext gather its context object, and will likely change in the future
-		if (const USkeletalMeshComponent* SkeletalMeshComponent = Actor->GetComponentByClass<USkeletalMeshComponent>())
-		{
-			if (const USkeletalMesh* SkeletalMesh = SkeletalMeshComponent->GetSkeletalMeshAsset())
-			{
-				const FTransform& RootBoneTransform = SkeletalMesh->GetSkeleton()->GetReferenceSkeleton().GetRefBonePose()[RootSchemaBoneIdx];
-				const FTransform& ComponentToWorldTransform = PoseHistory ? PoseHistoryTransformAtTime : SkeletalMeshComponent->GetComponentTransform();
-				return RootBoneTransform * ComponentToWorldTransform;
-			}
-		}
-	}
+	// accessing AnimInstance or SkeletalMesh properties is not thread unsafe from a worker thread! Logging an error and returning identity transform
+	UE_LOG(LogPoseSearch, Error, TEXT("FSearchContext::GetWorldRootBoneTransformAtTime - Couldn't find root bone transform at time %f for role %s for schema %s!"), SampleTime, *SampleRole.ToString(), *Schema->GetName());
+	
+	//const UObject* AnimContext = GetAnimContext(SampleRole);
+	//if (const UAnimInstance* AnimInstance = Cast<UAnimInstance>(AnimContext))
+	//{
+	//	if (AnimInstance && AnimInstance->CurrentSkeleton)
+	//	{
+	//		const FTransform& RootBoneTransform = AnimInstance->CurrentSkeleton->GetReferenceSkeleton().GetRefBonePose()[RootSchemaBoneIdx];
+	//		const FTransform& ComponentToWorldTransform = PoseHistory ? PoseHistoryTransformAtTime : AnimInstance->GetSkelMeshComponent()->GetComponentTransform();
+	//		return RootBoneTransform * ComponentToWorldTransform;
+	//	}
+	//}
+	//else if (const UActorComponent* ActorComponent = Cast<UActorComponent>(AnimContext))
+	//{
+	//	const AActor* Actor = ActorComponent->GetOwner();
+	//	check(Actor);
+	//	// @todo: this code depends on how AnimNext gather its context object, and will likely change in the future
+	//	if (const USkeletalMeshComponent* SkeletalMeshComponent = Actor->GetComponentByClass<USkeletalMeshComponent>())
+	//	{
+	//		if (const USkeletalMesh* SkeletalMesh = SkeletalMeshComponent->GetSkeletalMeshAsset())
+	//		{
+	//			const FTransform& RootBoneTransform = SkeletalMesh->GetSkeleton()->GetReferenceSkeleton().GetRefBonePose()[RootSchemaBoneIdx];
+	//			const FTransform& ComponentToWorldTransform = PoseHistory ? PoseHistoryTransformAtTime : SkeletalMeshComponent->GetComponentTransform();
+	//			return RootBoneTransform * ComponentToWorldTransform;
+	//		}
+	//	}
+	//}
 
 	return FTransform::Identity;
 }
diff --git a/Engine/Plugins/Animation/PoseSearch/Source/Runtime/Private/PoseSearchInteractionIsland.cpp b/Engine/Plugins/Animation/PoseSearch/Source/Runtime/Private/PoseSearchInteractionIsland.cpp
--- a/Engine/Plugins/Animation/PoseSearch/Source/Runtime/Private/PoseSearchInteractionIsland.cpp
+++ b/Engine/Plugins/Animation/PoseSearch/Source/Runtime/Private/PoseSearchInteractionIsland.cpp
@@ -33,17 +33,13 @@ static FInteractionSearchResult InitSearchResult(const FSearchResult& SearchResu
 
 	for (int32 AnimContextIndex = 0; AnimContextIndex < AnimContextsNum; ++AnimContextIndex)
 	{
-		if (const UObject* AnimContext = SearchContext.AnimContexts[AnimContextIndex].Get())
+		const UObject* AnimContext = SearchContext.AnimContexts[AnimContextIndex].Get();
+		const USkeleton* Skeleton = AnimContext ? GetContextSkeleton(AnimContext) : nullptr;
+		const IPoseHistory* PoseHistory = SearchContext.PoseHistories[AnimContextIndex];
+		if (Skeleton && PoseHistory)
 		{
-			InteractionSearchResult.ActorRootTransforms[AnimContextIndex] = GetContextTransform(AnimContext);
-			if (const IPoseHistory* PoseHistory = SearchContext.PoseHistories[AnimContextIndex])
-			{
-				PoseHistory->GetTransformAtTime(0.f, InteractionSearchResult.ActorRootBoneTransforms[AnimContextIndex], GetContextSkeleton(AnimContext), RootBoneIndexType, ComponentSpaceIndexType);
-			}
-			else
-			{
-				InteractionSearchResult.ActorRootBoneTransforms[AnimContextIndex] = FTransform::Identity;
-			}
+			PoseHistory->GetTransformAtTime(0.f, InteractionSearchResult.ActorRootTransforms[AnimContextIndex], Skeleton, ComponentSpaceIndexType, WorldSpaceIndexType);
+			PoseHistory->GetTransformAtTime(0.f, InteractionSearchResult.ActorRootBoneTransforms[AnimContextIndex], Skeleton, RootBoneIndexType, ComponentSpaceIndexType);
 		}
 		else
 		{
@@ -533,7 +529,7 @@ void FInteractionIsland::AddTickDependencies(UActorComponent* TickActorComponent
 
 void FInteractionIsland::RemoveTickDependencies(bool bValidateTickDependencies)
 {
-	// Called by UPoseSearchInteractionSubsystem::Tick when there aren't animation jobs flying. No need to FScopeLock Lock(&Mutex);
+	// Called by UPoseSearchInteractionSubsystem::Tick when there aren't animation jobs flying.
 	check(IsInGameThread());
 
 	check(TickActorComponents.Num() == IslandAnimContexts.Num());
@@ -730,13 +726,6 @@ const AActor* FInteractionIsland::GetMainActor() const
 	return nullptr;
 }
 
-bool FInteractionIsland::CanPerformSearch(const UObject* AnimContext) const
-{
-	// searches are performed only on the MainAnimContext / MainActor
-	check(AnimContext);
-	return AnimContext == GetMainAnimContext();
-}
-
 #if ENABLE_ANIM_DEBUG
 // recursion safe FTickFunction logging functions
 static void ShowPrerequistes(const FTickFunction& NestedTick, int32 Indent, int32 MaxIndent)
@@ -822,164 +811,110 @@ bool FInteractionIsland::DoSearch_AnyThread(const UObject* AnimContext, const FP
 {
 	check(AnimContext);
 
-	bool bDoPerfomSearch = false;
+	if (bSearchPerfomed)
 	{
-		// thread safety note!
-		// goal:	avoiding deadlock between SearchResultsMutex lock and waiting for UAnimInstance::HandleExistingParallelEvaluationTask.
-		// why:		UPoseSearchLibrary::MotionMatch could call via AnimContext GetProxyOnAnyThread<FAnimInstanceProxy>() that, if on GameThread, 
-		//			could call UAnimInstance::HandleExistingParallelEvaluationTask 
-		// fix:		avoid UPoseSearchLibrary::MotionMatch calls wrapped by any lock, at the cost of eventually (by design should be NEVER) performing the searches twice
-		//			By design we should inject ticks dependencies (added by FInteractionIsland::InjectToActor via AddTickPrerequisiteComponent), so concurrently 
-		//			fly of FInteractionIsland within the same island that requires searches is forbidden
-#if UE_POSE_SEARCH_ENABLE_ISLAND_CRITICAL_SECTION
-		FScopeLock Lock(&SearchResultsMutex);
-#endif // UE_POSE_SEARCH_ENABLE_ISLAND_CRITICAL_SECTION
-		bDoPerfomSearch = !bSearchPerfomed;
-
-		if (bDoPerfomSearch)
-		{
-			if (!CanPerformSearch(AnimContext))
-			{
-				// search failed, because AnimContext is not the context from the main actor. continuing the search would lead to threading issues,
-				// since the threading model expect the main actor to perform the search! (this happens when recompiling animation blue prints..)
-				bSearchPerfomed = true;
-				InitSearchResults(SearchResults, TConstArrayView<FSearchResult>(), TConstArrayView<FInteractionSearchContext>());
-
-				Result = FPoseSearchInteractionBlueprintResult();
-				return false;
-			}
-		}
-		else
-		{
-			return GetResult_AnyThread(AnimContext, Result);
-		}
+		check(AnimContext != GetMainAnimContext());
+		return GetResult_AnyThread(AnimContext, Result);
 	}
 
-	if (bDoPerfomSearch)
+	// searches are performed only on the MainAnimContext / MainActor
+	if (AnimContext != GetMainAnimContext())
 	{
-		QUICK_SCOPE_CYCLE_COUNTER(STAT_UPoseSearchInteractionInteractionIsland_Search);
+		// search failed, because AnimContext is not the context from the main actor. continuing the search would lead to threading issues,
+		// since the threading model expect the main actor to perform the search! (this happens when recompiling animation blue prints..)
+		Result = FPoseSearchInteractionBlueprintResult();
+		return false;
+	}
 
-		FMemMark Mark(FMemStack::Get());
+	QUICK_SCOPE_CYCLE_COUNTER(STAT_UPoseSearchInteractionInteractionIsland_Search);
 
-		TArray<const UObject*, TInlineAllocator<PreallocatedRolesNum, TMemStackAllocator<>>> AnimContexts;
-		TArray<FSearchResult, TMemStackAllocator<>> PoseSearchResults;
+	FMemMark Mark(FMemStack::Get());
 
-		// SearchContexts are modified only by UPoseSearchInteractionSubsystem::Tick and constant otherwise, so it's safe to access them in a threaded enviroment without locks
-		PoseSearchResults.SetNum(SearchContexts.Num());
+	TArray<const UObject*, TInlineAllocator<PreallocatedRolesNum, TMemStackAllocator<>>> AnimContexts;
+	TArray<FSearchResult, TMemStackAllocator<>> PoseSearchResults;
 
-		for (int32 SearchIndex = 0; SearchIndex < SearchContexts.Num(); ++SearchIndex)
+	// SearchContexts are modified only by UPoseSearchInteractionSubsystem::Tick and constant otherwise, so it's safe to access them in a threaded enviroment without locks
+	PoseSearchResults.SetNum(SearchContexts.Num());
+
+	for (int32 SearchIndex = 0; SearchIndex < SearchContexts.Num(); ++SearchIndex)
+	{
+		FInteractionSearchContext& SearchContext = SearchContexts[SearchIndex];
+		const UPoseSearchDatabase* Database = SearchContext.Database.Get();
+		if (!Database)
 		{
-			FInteractionSearchContext& SearchContext = SearchContexts[SearchIndex];
-			const UPoseSearchDatabase* Database = SearchContext.Database.Get();
-			if (!Database)
-			{
-				UE_LOG(LogPoseSearch, Error, TEXT("FInteractionIsland::DoSearch_AnyThread invalid context database"));
-				return false;
-			}
+			UE_LOG(LogPoseSearch, Error, TEXT("FInteractionIsland::DoSearch_AnyThread invalid context database"));
+			return false;
+		}
+
+		if (!Database->Schema)
+		{
+			UE_LOG(LogPoseSearch, Error, TEXT("FInteractionIsland::DoSearch_AnyThread invalid schema for context database %s"), *Database->GetName());
+			return false;
+		}
 
-			if (!Database->Schema)
+		AnimContexts.Reset();
+		for (const TWeakObjectPtr<const UObject>& AnimContextPtr : SearchContext.AnimContexts)
+		{
+			const UObject* SearchContextAnimContext = AnimContextPtr.Get();
+			if (!SearchContextAnimContext)
 			{
-				UE_LOG(LogPoseSearch, Error, TEXT("FInteractionIsland::DoSearch_AnyThread invalid schema for context database %s"), *Database->GetName());
+				UE_LOG(LogPoseSearch, Error, TEXT("FInteractionIsland::DoSearch_AnyThread null anim context"));
 				return false;
 			}
 
-			AnimContexts.Reset();
-			for (const TWeakObjectPtr<const UObject>& AnimContextPtr : SearchContext.AnimContexts)
-			{
-				const UObject* SearchContextAnimContext = AnimContextPtr.Get();
-				if (!SearchContextAnimContext)
-				{
-					UE_LOG(LogPoseSearch, Error, TEXT("FInteractionIsland::DoSearch_AnyThread null anim context"));
-					return false;
-				}
-
-				AnimContexts.Add(SearchContextAnimContext);
-			}
+			AnimContexts.Add(SearchContextAnimContext);
+		}
 
-			const UObject* AssetsToSearch[] = { Database };
-			FPoseSearchFutureProperties PoseSearchFutureProperties;
+		const UObject* AssetsToSearch[] = { Database };
+		FPoseSearchFutureProperties PoseSearchFutureProperties;
 
-			FPoseSearchContinuingProperties ContinuingPropertiesToUse = SearchContext.GetContinuingProperties();
+		FPoseSearchContinuingProperties ContinuingPropertiesToUse = SearchContext.GetContinuingProperties();
 
-			// @todo: implement this carefully, since it'll cause thread safety concerns depending on which actors is calling the query first 
-			// and if their animations get integrated by different play rates
+		// @todo: implement this carefully, since it'll cause thread safety concerns depending on which actors is calling the query first 
+		// and if their animations get integrated by different play rates
 			
-			//if (ContinuingProperties.PlayingAsset)
-			//{
-			//	if (const UMultiAnimAsset* MultiAnimAsset = Cast<UMultiAnimAsset>(ContinuingProperties.PlayingAsset))
-			//	{
-			//		// @todo: if they are compatible... etcetc MultiAnimAsset->GetNumRoles
-			//		ContinuingPropertiesToUse = ContinuingProperties;
-			//	}
-			//	else if (const UMultiAnimAsset* SearchContextMultiAnimAsset = Cast<UMultiAnimAsset>(SearchContext.ContinuingProperties.PlayingAsset))
-			//	{
-			//		for (int32 RoleIndex = 0; RoleIndex < SearchContextMultiAnimAsset->GetNumRoles(); ++RoleIndex)
-			//		{
-			//			if (SearchContextMultiAnimAsset->GetAnimationAsset(SearchContextMultiAnimAsset->GetRole(RoleIndex)) == ContinuingProperties.PlayingAsset)
-			//			{
-			//				ContinuingPropertiesToUse.PlayingAssetAccumulatedTime = ContinuingProperties.PlayingAssetAccumulatedTime;
-			//				break;
-			//			}
-			//		}
-			//	}
-			//	else if (!SearchContext.ContinuingProperties.PlayingAsset)
-			//	{
-			//		// @todo: should we scan all the database iterating over the UMultiAnimAsset to look for ContinuingProperties.PlayingAsset?
-			//	}
-			//	else
-			//	{
-			//	}
-			//}
-
-			// @todo: we could perform multiple UPoseSearchLibrary::MotionMatch in parallel!
-			PoseSearchResults[SearchIndex] = UPoseSearchLibrary::MotionMatch(AnimContexts, SearchContext.Roles,
-				SearchContext.PoseHistories, AssetsToSearch, ContinuingPropertiesToUse, PoseSearchFutureProperties, FPoseSearchEvent());
-		}
+		//if (ContinuingProperties.PlayingAsset)
+		//{
+		//	if (const UMultiAnimAsset* MultiAnimAsset = Cast<UMultiAnimAsset>(ContinuingProperties.PlayingAsset))
+		//	{
+		//		// @todo: if they are compatible... etcetc MultiAnimAsset->GetNumRoles
+		//		ContinuingPropertiesToUse = ContinuingProperties;
+		//	}
+		//	else if (const UMultiAnimAsset* SearchContextMultiAnimAsset = Cast<UMultiAnimAsset>(SearchContext.ContinuingProperties.PlayingAsset))
+		//	{
+		//		for (int32 RoleIndex = 0; RoleIndex < SearchContextMultiAnimAsset->GetNumRoles(); ++RoleIndex)
+		//		{
+		//			if (SearchContextMultiAnimAsset->GetAnimationAsset(SearchContextMultiAnimAsset->GetRole(RoleIndex)) == ContinuingProperties.PlayingAsset)
+		//			{
+		//				ContinuingPropertiesToUse.PlayingAssetAccumulatedTime = ContinuingProperties.PlayingAssetAccumulatedTime;
+		//				break;
+		//			}
+		//		}
+		//	}
+		//	else if (!SearchContext.ContinuingProperties.PlayingAsset)
+		//	{
+		//		// @todo: should we scan all the database iterating over the UMultiAnimAsset to look for ContinuingProperties.PlayingAsset?
+		//	}
+		//	else
+		//	{
+		//	}
+		//}
+
+		// @todo: we could perform multiple UPoseSearchLibrary::MotionMatch in parallel!
+		PoseSearchResults[SearchIndex] = UPoseSearchLibrary::MotionMatch(AnimContexts, SearchContext.Roles,
+			SearchContext.PoseHistories, AssetsToSearch, ContinuingPropertiesToUse, PoseSearchFutureProperties, FPoseSearchEvent());
+	}
 		
-#if UE_POSE_SEARCH_ENABLE_ISLAND_CRITICAL_SECTION
-		// locking to update SearchResults and bSearchPerfomed
-		FScopeLock Lock(&SearchResultsMutex);
-#endif // UE_POSE_SEARCH_ENABLE_ISLAND_CRITICAL_SECTION
-
-		if (!bSearchPerfomed)
-		{
-			InitSearchResults(SearchResults, PoseSearchResults, SearchContexts);
-
-			bSearchPerfomed = true;
-		}
-		else
-		{
-			UE_LOG(LogPoseSearch, Warning, TEXT("FInteractionIsland::DoSearch_AnyThread performance warning: performed duplicated search :("));
-				
-#if DO_CHECK
-			TArray<FInteractionSearchResult> CompareSearchResults;
-			InitSearchResults(CompareSearchResults, PoseSearchResults, SearchContexts);
-
-			if (CompareSearchResults != SearchResults)
-			{
-				UE_LOG(LogPoseSearch, Error, TEXT("FInteractionIsland::DoSearch_AnyThread duplicated search differs from original one. Searches are NOT deterministic!"));
-			}
-#endif // DO_CHECK
-		}
+	InitSearchResults(SearchResults, PoseSearchResults, SearchContexts);
+	bSearchPerfomed = true;
 	
-		// calling this funtion within this scope since we already locked SearchResultsMutex
-		return GetResult_AnyThread(AnimContext, Result);
-	}
-
-	checkNoEntry();
-	Result = FPoseSearchInteractionBlueprintResult();
-	return false;
+	return GetResult_AnyThread(AnimContext, Result);
 }
 
 bool FInteractionIsland::GetResult_AnyThread(const UObject* AnimContext, FPoseSearchInteractionBlueprintResult& Result, bool bCompareOwningActors)
 {
 	check(AnimContext);
 
-#if UE_POSE_SEARCH_ENABLE_ISLAND_CRITICAL_SECTION
-	// locking to read SearchResults
-	FScopeLock Lock(&SearchResultsMutex);
-#endif // UE_POSE_SEARCH_ENABLE_ISLAND_CRITICAL_SECTION
-
 	const AActor* Actor = bCompareOwningActors ? GetContextOwningActor(AnimContext) : nullptr;
 
 	// looking for AnimContext in SearchResults to fill up Result
diff --git a/Engine/Plugins/Animation/PoseSearch/Source/Runtime/Private/PoseSearchInteractionSubsystem.cpp b/Engine/Plugins/Animation/PoseSearch/Source/Runtime/Private/PoseSearchInteractionSubsystem.cpp
--- a/Engine/Plugins/Animation/PoseSearch/Source/Runtime/Private/PoseSearchInteractionSubsystem.cpp
+++ b/Engine/Plugins/Animation/PoseSearch/Source/Runtime/Private/PoseSearchInteractionSubsystem.cpp
@@ -133,10 +133,6 @@ static void GeneratePermutations(int32 DataCardinality, int32 CombinationCardina
 
 } // namespace UE::PoseSearch
 
-#if UE_POSE_SEARCH_ENABLE_INTERACTION_SUBSYSTEM_CRITICAL_SECTION
-FCriticalSection UPoseSearchInteractionSubsystem::RetrieveSubsystemMutex;
-#endif // UE_POSE_SEARCH_ENABLE_INTERACTION_SUBSYSTEM_CRITICAL_SECTION
-
 UE::PoseSearch::FInteractionIsland& UPoseSearchInteractionSubsystem::CreateIsland()
 {
 	return *Islands.Add_GetRef(new UE::PoseSearch::FInteractionIsland(ToRawPtr(GetWorld()->PersistentLevel), this));
@@ -467,10 +463,6 @@ UPoseSearchInteractionSubsystem* UPoseSearchInteractionSubsystem::GetSubsystem_A
 	{
 		if (UWorld* World = AnimContext->GetWorld())
 		{
-#if UE_POSE_SEARCH_ENABLE_INTERACTION_SUBSYSTEM_CRITICAL_SECTION
-			FScopeLock Lock(&RetrieveSubsystemMutex);
-#endif // UE_POSE_SEARCH_ENABLE_INTERACTION_SUBSYSTEM_CRITICAL_SECTION
-
 			// We expect the subsystem to be already created from the GameThread.
 			// We don't create the subsystem from any thread
 			if (World->HasSubsystem<UPoseSearchInteractionSubsystem>())
diff --git a/Engine/Plugins/Animation/PoseSearch/Source/Runtime/Public/PoseSearch/PoseSearchAssetSamplerLibrary.h b/Engine/Plugins/Animation/PoseSearch/Source/Runtime/Public/PoseSearch/PoseSearchAssetSamplerLibrary.h
--- a/Engine/Plugins/Animation/PoseSearch/Source/Runtime/Public/PoseSearch/PoseSearchAssetSamplerLibrary.h
+++ b/Engine/Plugins/Animation/PoseSearch/Source/Runtime/Public/PoseSearch/PoseSearchAssetSamplerLibrary.h
@@ -149,6 +149,8 @@ public:
 	template<typename FComponentSpacePoseType>
 	static void VLogDraw(const UObject* VLogContext, FComponentSpacePoseType& ComponentSpacePose, const FTransform& RootTransform, const TCHAR* VLogName, const FColor Color = FColor::Red, float DebugDrawSamplerRootAxisLength = 20.f)
 	{
+		check(IsInGameThread());
+
 		check(VLogName);
 		if (VLogContext)
 		{
diff --git a/Engine/Plugins/Animation/PoseSearch/Source/Runtime/Public/PoseSearch/PoseSearchDefines.h b/Engine/Plugins/Animation/PoseSearch/Source/Runtime/Public/PoseSearch/PoseSearchDefines.h
--- a/Engine/Plugins/Animation/PoseSearch/Source/Runtime/Public/PoseSearch/PoseSearchDefines.h
+++ b/Engine/Plugins/Animation/PoseSearch/Source/Runtime/Public/PoseSearch/PoseSearchDefines.h
@@ -11,16 +11,6 @@
 #define UE_POSE_SEARCH_TRACE_ENABLED ANIM_TRACE_ENABLED
 #endif
 
-// NoTe: the UE::PoseSearch::FInteractionIsland mutex is now unnecessary, but I'll leave this macro to reenable it in case I missed something =):>
-#ifndef UE_POSE_SEARCH_ENABLE_ISLAND_CRITICAL_SECTION
-#define UE_POSE_SEARCH_ENABLE_ISLAND_CRITICAL_SECTION 0
-#endif
-
-// NoTe: the UPoseSearchInteractionSubsystem mutex to retrieve the subsystem is now unnecessary, but I'll leave this macro to reenable it in case I missed something =):>
-#ifndef UE_POSE_SEARCH_ENABLE_INTERACTION_SUBSYSTEM_CRITICAL_SECTION
-#define UE_POSE_SEARCH_ENABLE_INTERACTION_SUBSYSTEM_CRITICAL_SECTION 0
-#endif
-
 #ifndef UE_POSE_SEARCH_FORCE_SINGLE_THREAD
 #define UE_POSE_SEARCH_FORCE_SINGLE_THREAD 0
 #endif
diff --git a/Engine/Plugins/Animation/PoseSearch/Source/Runtime/Public/PoseSearch/PoseSearchInteractionIsland.h b/Engine/Plugins/Animation/PoseSearch/Source/Runtime/Public/PoseSearch/PoseSearchInteractionIsland.h
--- a/Engine/Plugins/Animation/PoseSearch/Source/Runtime/Public/PoseSearch/PoseSearchInteractionIsland.h
+++ b/Engine/Plugins/Animation/PoseSearch/Source/Runtime/Public/PoseSearch/PoseSearchInteractionIsland.h
@@ -116,7 +116,6 @@ private:
 
 	const UObject* GetMainAnimContext() const;
 	const AActor* GetMainActor() const;
-	bool CanPerformSearch(const UObject* AnimContext) const;
 
 	struct FPreTickFunction : public FTickFunction
 	{
@@ -150,10 +149,6 @@ private:
 
 	UPoseSearchInteractionSubsystem* InteractionSubsystem = nullptr;
 
-#if UE_POSE_SEARCH_ENABLE_ISLAND_CRITICAL_SECTION
-	FCriticalSection SearchResultsMutex;
-#endif // UE_POSE_SEARCH_ENABLE_ISLAND_CRITICAL_SECTION
-
 #if ENABLE_ANIM_DEBUG
 	// used to analyze thread safety
 	friend struct UE::PoseSearch::FInteractionValidator;
diff --git a/Engine/Plugins/Animation/PoseSearch/Source/Runtime/Public/PoseSearch/PoseSearchInteractionSubsystem.h b/Engine/Plugins/Animation/PoseSearch/Source/Runtime/Public/PoseSearch/PoseSearchInteractionSubsystem.h
--- a/Engine/Plugins/Animation/PoseSearch/Source/Runtime/Public/PoseSearch/PoseSearchInteractionSubsystem.h
+++ b/Engine/Plugins/Animation/PoseSearch/Source/Runtime/Public/PoseSearch/PoseSearchInteractionSubsystem.h
@@ -199,9 +199,4 @@ private:
 	TArray<UE::PoseSearch::FInteractionIsland*> Islands;
 
 	TArray<UE::PoseSearch::FValidInteractionSearch> ValidInteractionSearches;
-
-#if UE_POSE_SEARCH_ENABLE_INTERACTION_SUBSYSTEM_CRITICAL_SECTION
-	// critical section to retrieve the subsystem in a thread safe manner
-	static FCriticalSection RetrieveSubsystemMutex;
-#endif // UE_POSE_SEARCH_ENABLE_INTERACTION_SUBSYSTEM_CRITICAL_SECTION
 };
diff --git a/Engine/Plugins/Experimental/AnimNextPoseSearch/Source/AnimNextPoseSearch/Private/HistoryCollectorTrait.cpp b/Engine/Plugins/Experimental/AnimNextPoseSearch/Source/AnimNextPoseSearch/Private/HistoryCollectorTrait.cpp
--- a/Engine/Plugins/Experimental/AnimNextPoseSearch/Source/AnimNextPoseSearch/Private/HistoryCollectorTrait.cpp
+++ b/Engine/Plugins/Experimental/AnimNextPoseSearch/Source/AnimNextPoseSearch/Private/HistoryCollectorTrait.cpp
@@ -7,17 +7,6 @@
 #include "EvaluationVM/EvaluationVM.h"
 #include "GenerationTools.h"
 
-FAnimNextHistoryCollectorTask FAnimNextHistoryCollectorTask::Make(UE::PoseSearch::FPoseHistory* PoseHistory, const FAnimNextHistoryCollectorTraitSharedData* SharedData, bool bStoreScales, float DeltaTime)
-{
-	check(PoseHistory);
-	FAnimNextHistoryCollectorTask Task;
-	Task.PoseHistory = PoseHistory;
-	Task.SharedData = SharedData;
-	Task.bStoreScales = bStoreScales;
-	Task.DeltaTime = DeltaTime;
-	return Task;
-}
-
 namespace UE::AnimNext
 {
 	// @todo: revisit this code once a FCSPose<FLODPoseStack> is implemented!
@@ -194,7 +183,12 @@ namespace UE::AnimNext
 		FInstanceData* InstanceData = Binding.GetInstanceData<FInstanceData>();
 		const bool bStoreScales = SharedData->GetbStoreScales(Binding);
 
-		Context.AppendTask(FAnimNextHistoryCollectorTask::Make(&InstanceData->PoseHistory, SharedData, bStoreScales, InstanceData->DeltaTime));
+		FAnimNextHistoryCollectorTask Task;
+		Task.PoseHistory = &InstanceData->PoseHistory;
+		Task.SharedData = SharedData;
+		Task.bStoreScales = bStoreScales;
+		Task.DeltaTime = InstanceData->DeltaTime;
+		Context.AppendTask(Task);
 	}
 
 	const UE::PoseSearch::IPoseHistory* FHistoryCollectorTrait::GetPoseHistory(FExecutionContext& Context, const TTraitBinding<IPoseHistory>& Binding) const
diff --git a/Engine/Plugins/Experimental/AnimNextPoseSearch/Source/AnimNextPoseSearch/Private/HistoryCollectorTrait.h b/Engine/Plugins/Experimental/AnimNextPoseSearch/Source/AnimNextPoseSearch/Private/HistoryCollectorTrait.h
--- a/Engine/Plugins/Experimental/AnimNextPoseSearch/Source/AnimNextPoseSearch/Private/HistoryCollectorTrait.h
+++ b/Engine/Plugins/Experimental/AnimNextPoseSearch/Source/AnimNextPoseSearch/Private/HistoryCollectorTrait.h
@@ -115,9 +115,6 @@ struct FAnimNextHistoryCollectorTask : public FAnimNextEvaluationTask
 
 	DECLARE_ANIM_EVALUATION_TASK(FAnimNextHistoryCollectorTask)
 
-	static FAnimNextHistoryCollectorTask Make(UE::PoseSearch::FPoseHistory* PoseHistory, const FAnimNextHistoryCollectorTraitSharedData* SharedData, bool bStoreScales, float DeltaTime);
-
-	// Task entry point
 	virtual void Execute(UE::AnimNext::FEvaluationVM& VM) const override;
 
 	UE::PoseSearch::FPoseHistory* PoseHistory = nullptr;
diff --git a/Engine/Plugins/Experimental/AnimNextPoseSearch/Source/AnimNextPoseSearch/Private/MotionMatchingTrait.cpp b/Engine/Plugins/Experimental/AnimNextPoseSearch/Source/AnimNextPoseSearch/Private/MotionMatchingTrait.cpp
--- a/Engine/Plugins/Experimental/AnimNextPoseSearch/Source/AnimNextPoseSearch/Private/MotionMatchingTrait.cpp
+++ b/Engine/Plugins/Experimental/AnimNextPoseSearch/Source/AnimNextPoseSearch/Private/MotionMatchingTrait.cpp
@@ -1,11 +1,13 @@
 // Copyright Epic Games, Inc. All Rights Reserved.
 
 #include "MotionMatchingTrait.h"
-
+#include "Animation/AnimRootMotionProvider.h"
 #include "Animation/AnimSequence.h"
 #include "Animation/BlendSpace.h"
 #include "AnimNextAnimGraphSettings.h"
+#include "Component/AnimNextComponent.h"
 #include "DataInterface/DataInterfaceStructAdapter.h"
+#include "EvaluationVM/EvaluationVM.h"
 #include "Graph/AnimNextGraphInstance.h"
 #include "GraphInterfaces/AnimNextNativeDataInterface_AnimSequencePlayer.h"
 #include "GraphInterfaces/AnimNextNativeDataInterface_BlendSpacePlayer.h"
@@ -14,10 +16,12 @@
 #include "PoseSearch/MultiAnimAsset.h"
 #include "PoseSearch/PoseSearchDatabase.h"
 #include "PoseSearch/PoseSearchDerivedData.h"
+#include "PoseSearch/PoseSearchInteractionUtils.h"
 #include "PoseSearch/PoseSearchSchema.h"
 #include "TraitInterfaces/IBlendStack.h"
 #include "TraitInterfaces/IGraphFactory.h"
 #include "TraitInterfaces/ITimeline.h"
+#include "VisualLogger/VisualLogger.h"
 
 #if ENABLE_ANIM_DEBUG
 namespace UE::Private
@@ -45,14 +49,15 @@ namespace UE::AnimNext
 {
 	AUTO_REGISTER_ANIM_TRAIT(FMotionMatchingTrait)
 
-	#define TRAIT_INTERFACE_ENUMERATOR(GeneratorMacro) \
+#define TRAIT_INTERFACE_ENUMERATOR(GeneratorMacro) \
 		GeneratorMacro(IUpdate) \
+		GeneratorMacro(IEvaluate) \
 
-	// Trait implementation boilerplate
-	GENERATE_ANIM_TRAIT_IMPLEMENTATION(FMotionMatchingTrait, TRAIT_INTERFACE_ENUMERATOR, NULL_ANIM_TRAIT_INTERFACE_ENUMERATOR, NULL_ANIM_TRAIT_EVENT_ENUMERATOR)
-	#undef TRAIT_INTERFACE_ENUMERATOR
+		// Trait implementation boilerplate
+		GENERATE_ANIM_TRAIT_IMPLEMENTATION(FMotionMatchingTrait, TRAIT_INTERFACE_ENUMERATOR, NULL_ANIM_TRAIT_INTERFACE_ENUMERATOR, NULL_ANIM_TRAIT_EVENT_ENUMERATOR)
+#undef TRAIT_INTERFACE_ENUMERATOR
 
-	void FMotionMatchingTrait::OnBecomeRelevant(FUpdateTraversalContext& Context, const TTraitBinding<IUpdate>& Binding, const FTraitUpdateState& TraitState) const
+		void FMotionMatchingTrait::OnBecomeRelevant(FUpdateTraversalContext& Context, const TTraitBinding<IUpdate>& Binding, const FTraitUpdateState& TraitState) const
 	{
 		IUpdate::OnBecomeRelevant(Context, Binding, TraitState);
 
@@ -113,7 +118,7 @@ namespace UE::AnimNext
 		{
 			bNeedsReset = true;
 		}
-	#endif // WITH_EDITOR
+#endif // WITH_EDITOR
 
 		// If we just became relevant and haven't been initialized yet, then reset motion matching state, otherwise update the asset time using the player node.
 		if (bNeedsReset)
@@ -144,7 +149,7 @@ namespace UE::AnimNext
 						const float TimelineDuration = TimelineState.GetDuration();
 						if (TimelineDuration > 0.0f)
 						{
-                            // NOTE: This doesn't work because AnimNext normalized time does not much database indexing normalized time.
+							// NOTE: This doesn't work because AnimNext normalized time does not much database indexing normalized time.
 							TimelinePosition /= TimelineState.GetDuration();
 						}
 					}
@@ -278,7 +283,6 @@ namespace UE::AnimNext
 				FAnimNextNativeDataInterface_BlendSpacePlayer& NewNativeInterface = NewPayload.InitializeAs<FAnimNextNativeDataInterface_BlendSpacePlayer>();
 				NewNativeInterface.BlendSpace = Cast<UBlendSpace>(AnimationAsset);
 
-
 				NewNativeInterface.XAxisSamplePoint = BlendSpaceParameters.X;
 				NewNativeInterface.YAxisSamplePoint = BlendSpaceParameters.Y;
 
@@ -299,7 +303,7 @@ namespace UE::AnimNext
 			GraphPayload.AddNative(MoveTemp(NewPayload));
 
 			const UAnimNextAnimationGraph* AnimationGraph = IGraphFactory::GetGraphFromObjectWithFallback(Context, Binding, AnimationAsset, GraphPayload);
-			if(AnimationGraph != nullptr)
+			if (AnimationGraph != nullptr)
 			{
 				NewGraphRequest.GraphPayload = MoveTemp(GraphPayload);
 				NewGraphRequest.AnimationGraph = AnimationGraph;
@@ -336,4 +340,134 @@ namespace UE::AnimNext
 			}
 		}
 	}
+
+	void FMotionMatchingTrait::PostEvaluate(FEvaluateTraversalContext& Context, const TTraitBinding<IEvaluate>& Binding) const
+	{
+		using namespace UE::PoseSearch;
+
+		IEvaluate::PostEvaluate(Context, Binding);
+
+		FInstanceData* InstanceData = Binding.GetInstanceData<FInstanceData>();
+		check(InstanceData);
+
+		// if it's not a multi character interaction we can skip the warping logic entirely
+		if (!InstanceData->MotionMatchingInteractionState.SelectedAnimation)
+		{
+			return;
+		}
+
+		const FPoseSearchInteractionBlueprintResult& CurrentResult = InstanceData->MotionMatchingInteractionState;
+		if (const UMultiAnimAsset* MultiAnimAsset = Cast<UMultiAnimAsset>(CurrentResult.SelectedAnimation))
+		{
+			const int32 CurrentResultRoleIndex = GetRoleIndex(MultiAnimAsset, CurrentResult.Role);
+			if (CurrentResultRoleIndex != INDEX_NONE)
+			{
+				// @todo: WIP hacky, non thread safe (unless proper tick dependencies are in place) way to retrieve the mesh transform until we find a better way
+				if (const FAnimNextModuleInstance* ModuleInstance = Context.GetRootGraphInstance().GetModuleInstance())
+				{
+					if (const UAnimNextComponent* AnimNextComponent = Cast<UAnimNextComponent>(ModuleInstance->GetObject()))
+					{
+						const AActor* Actor = AnimNextComponent->GetOwner();
+						check(Actor);
+						if (const USkeletalMeshComponent* SkeletalMeshComponent = Actor->GetComponentByClass<USkeletalMeshComponent>())
+						{
+							const FSharedData* SharedData = Binding.GetSharedData<FSharedData>();
+							check(SharedData);
+		
+							FAnimNextMotionMatchingTask Task;
+							Task.ComponentTransform = SkeletalMeshComponent->GetComponentTransform();
+							Task.InstanceData = InstanceData;
+							Task.CurrentResultRoleIndex = CurrentResultRoleIndex;
+							Task.bWarpUsingRootBone = SharedData->GetbWarpUsingRootBone(Binding);
+							Task.WarpingRotationRatio = SharedData->GetWarpingRotationRatio(Binding);
+							Task.WarpingTranslationRatio = SharedData->GetWarpingTranslationRatio(Binding);
+#if ENABLE_ANIM_DEBUG
+							// Debug Object for VisualLogger
+							Task.HostObject = Context.GetHostObject();
+#endif // ENABLE_ANIM_DEBUG 
+
+							Context.AppendTask(Task);
+						}
+					}
+				}
+			}
+		}
+	}
+} // UE::AnimNext
+
+void FAnimNextMotionMatchingTask::Execute(UE::AnimNext::FEvaluationVM& VM) const
+{
+	using namespace UE::AnimNext;
+	using namespace UE::PoseSearch;
+
+	check(CurrentResultRoleIndex != INDEX_NONE);
+
+	const UE::Anim::IAnimRootMotionProvider* RootMotionProvider = UE::Anim::IAnimRootMotionProvider::Get();
+	if (!RootMotionProvider)
+	{
+		UE_LOG(LogPoseSearch, Error, TEXT("FAnimNextMotionMatchingTask::Execute, missing RootMotionProvider"));
+		return;
+	}
+	
+	const TUniquePtr<FKeyframeState>* Keyframe = VM.PeekValue<TUniquePtr<FKeyframeState>>(KEYFRAME_STACK_NAME, 0);
+	if (!Keyframe)
+	{
+		return;
+	}
+
+	const FPoseSearchInteractionBlueprintResult& CurrentResult = InstanceData->MotionMatchingInteractionState;
+	const UMultiAnimAsset* MultiAnimAsset = Cast<UMultiAnimAsset>(CurrentResult.SelectedAnimation);
+	if (!MultiAnimAsset)
+	{
+		return;
+	}
+
+	const int32 NumRoles = MultiAnimAsset->GetNumRoles();
+	if (CurrentResult.ActorRootTransforms.Num() != NumRoles)
+	{
+		// warping is supported only for UMultiAnimAsset(s)
+		return;
+	}
+
+	FTransform RootMotionDelta = FTransform::Identity;
+	if (!RootMotionProvider->ExtractRootMotion(Keyframe->Get()->Attributes, RootMotionDelta))
+	{
+		return;
+	}
+			
+	TArray<FTransform, TInlineAllocator<PreallocatedRolesNum>> FullAlignedTransforms;
+	FullAlignedTransforms.SetNum(NumRoles);
+
+	CalculateFullAlignedTransforms(CurrentResult, bWarpUsingRootBone, FullAlignedTransforms);
+
+	// @todo: implement this properly once we have UE::AnimationWarping::FRootOffsetProvider
+	const FTransform& MeshWithoutOffset = ComponentTransform;
+	const FTransform& MeshWithOffset = MeshWithoutOffset;
+
+	// NoTe: keep in mind DeltaAlignment is relative to the previous execution frame so we still need to extract and and apply the current animation root motion transform to get to the current frame full aligned transform.
+	const FTransform DeltaAlignment = CalculateDeltaAlignment(MeshWithoutOffset, MeshWithOffset, FullAlignedTransforms[CurrentResultRoleIndex], WarpingRotationRatio, WarpingTranslationRatio);
+	
+	const FTransform DeltaAlignmentWithRootMotion = DeltaAlignment * RootMotionDelta;
+
+	RootMotionProvider->OverrideRootMotion(DeltaAlignmentWithRootMotion, Keyframe->Get()->Attributes);
+
+#if ENABLE_VISUAL_LOG
+	if (FVisualLogger::IsRecording())
+	{
+		static const TCHAR* LogName = TEXT("FAnimNextMotionMatchingTask");
+
+		for (int32 Index = 0; Index < NumRoles; ++Index)
+		{
+			const FTransform& ActorRootTransform = CurrentResult.ActorRootTransforms[Index];
+			const FTransform& FullAlignedTransform = FullAlignedTransforms[Index];
+
+			UE_VLOG_SEGMENT_THICK(HostObject, LogName, Display, FullAlignedTransform.GetLocation(), ActorRootTransform.GetLocation(), FColorList::Orange, 1.f, TEXT(""));
+			UE_VLOG_SEGMENT_THICK(HostObject, LogName, Display, ActorRootTransform.GetLocation(), ActorRootTransform.GetLocation() + ActorRootTransform.GetRotation().GetForwardVector() * 35, FColorList::LightGrey, 3.f, TEXT(""));
+			UE_VLOG_SEGMENT_THICK(HostObject, LogName, Display, FullAlignedTransform.GetLocation(), FullAlignedTransform.GetLocation() + FullAlignedTransform.GetRotation().GetForwardVector() * 30, FColorList::Orange, 2.f, TEXT(""));
+		}
+
+		UE_VLOG_SEGMENT_THICK(HostObject, LogName, Display, MeshWithOffset.GetLocation(), MeshWithOffset.GetLocation() + MeshWithOffset.GetRotation().GetForwardVector() * 35, FColorList::Blue, 3.f, TEXT(""));
+		UE_VLOG_SEGMENT_THICK(HostObject, LogName, Display, MeshWithoutOffset.GetLocation(), MeshWithoutOffset.GetLocation() + MeshWithoutOffset.GetRotation().GetForwardVector() * 40, FColorList::Cyan, 4.f, TEXT(""));
+	}
+#endif
 }
diff --git a/Engine/Plugins/Experimental/AnimNextPoseSearch/Source/AnimNextPoseSearch/Private/MotionMatchingTrait.h b/Engine/Plugins/Experimental/AnimNextPoseSearch/Source/AnimNextPoseSearch/Private/MotionMatchingTrait.h
--- a/Engine/Plugins/Experimental/AnimNextPoseSearch/Source/AnimNextPoseSearch/Private/MotionMatchingTrait.h
+++ b/Engine/Plugins/Experimental/AnimNextPoseSearch/Source/AnimNextPoseSearch/Private/MotionMatchingTrait.h
@@ -7,6 +7,7 @@
 #include "PoseSearch/PoseSearchLibrary.h"
 #include "TraitCore/Trait.h"
 #include "TraitCore/TraitSharedData.h"
+#include "TraitInterfaces/IEvaluate.h"
 #include "TraitInterfaces/IUpdate.h"
 #include "MotionMatchingTrait.generated.h"
 
@@ -70,6 +71,15 @@ struct FMotionMatchingTraitSharedData : public FAnimNextTraitSharedData
 	UPROPERTY(EditAnywhere, Category = Interaction, meta = (PinHiddenByDefault))
 	bool bValidateResultAgainstAvailabilities = true;
 
+	UPROPERTY(EditAnywhere, Category = Interaction, meta = (PinHiddenByDefault))
+	bool bWarpUsingRootBone = true;
+
+	UPROPERTY(EditAnywhere, Category = Interaction, meta = (PinHiddenByDefault))
+	float WarpingRotationRatio = 1.f;
+
+	UPROPERTY(EditAnywhere, Category = Interaction, meta = (PinHiddenByDefault))
+	float WarpingTranslationRatio = 1.f;
+
 	// Latent pin support boilerplate
 	#define TRAIT_LATENT_PROPERTIES_ENUMERATOR(GeneratorMacro) \
 		GeneratorMacro(Databases) \
@@ -86,6 +96,9 @@ struct FMotionMatchingTraitSharedData : public FAnimNextTraitSharedData
 		GeneratorMacro(NextUpdateInterruptMode) \
 		GeneratorMacro(Availabilities) \
 		GeneratorMacro(bValidateResultAgainstAvailabilities) \
+		GeneratorMacro(bWarpUsingRootBone) \
+		GeneratorMacro(WarpingRotationRatio) \
+		GeneratorMacro(WarpingTranslationRatio) \
 
 	GENERATE_TRAIT_LATENT_PROPERTIES(FMotionMatchingTraitSharedData, TRAIT_LATENT_PROPERTIES_ENUMERATOR)
 	#undef TRAIT_LATENT_PROPERTIES_ENUMERATOR
@@ -94,7 +107,7 @@ struct FMotionMatchingTraitSharedData : public FAnimNextTraitSharedData
 namespace UE::AnimNext
 {
 
-struct FMotionMatchingTrait : FAdditiveTrait, IUpdate
+struct FMotionMatchingTrait : FAdditiveTrait, IUpdate, IEvaluate
 {
 	DECLARE_ANIM_TRAIT(FMotionMatchingTrait, FAdditiveTrait)
 
@@ -111,6 +124,31 @@ struct FMotionMatchingTrait : FAdditiveTrait, IUpdate
 	// IUpdate impl
 	virtual void PreUpdate(FUpdateTraversalContext& Context, const TTraitBinding<IUpdate>& Binding, const FTraitUpdateState& TraitState) const override;
 	virtual void OnBecomeRelevant(FUpdateTraversalContext& Context, const TTraitBinding<IUpdate>& Binding, const FTraitUpdateState& TraitState) const override;
+
+	// IEvaluate impl
+	virtual void PostEvaluate(FEvaluateTraversalContext& Context, const TTraitBinding<IEvaluate>& Binding) const override;
 };
 
 } // namespace UE::AnimNext
+
+USTRUCT(Experimental)
+struct FAnimNextMotionMatchingTask : public FAnimNextEvaluationTask
+{
+	GENERATED_BODY()
+
+	DECLARE_ANIM_EVALUATION_TASK(FAnimNextMotionMatchingTask)
+
+	virtual void Execute(UE::AnimNext::FEvaluationVM& VM) const override;
+
+	FTransform ComponentTransform = FTransform::Identity;
+	UE::AnimNext::FMotionMatchingTrait::FInstanceData* InstanceData = nullptr;
+	int32 CurrentResultRoleIndex = INDEX_NONE;
+	bool bWarpUsingRootBone = true;
+	float WarpingRotationRatio = 1.f;
+	float WarpingTranslationRatio = 1.f;
+
+#if ENABLE_ANIM_DEBUG
+	// Debug Object for VisualLogger
+	const UObject* HostObject = nullptr;
+#endif // ENABLE_ANIM_DEBUG 
+};
\ No newline at end of file
diff --git a/Engine/Plugins/Experimental/AnimNextWarping/Source/AnimNextWarping/Private/WarpTestTrait.cpp b/Engine/Plugins/Experimental/AnimNextWarping/Source/AnimNextWarping/Private/WarpTestTrait.cpp
--- a/Engine/Plugins/Experimental/AnimNextWarping/Source/AnimNextWarping/Private/WarpTestTrait.cpp
+++ b/Engine/Plugins/Experimental/AnimNextWarping/Source/AnimNextWarping/Private/WarpTestTrait.cpp
@@ -23,6 +23,8 @@ namespace UE::AnimNext
 
 	void FWarpTestTrait::PreUpdate(FUpdateTraversalContext& Context, const TTraitBinding<IUpdate>& Binding, const FTraitUpdateState& TraitState) const
 	{
+		IUpdate::PreUpdate(Context, Binding, TraitState);
+
 		const FSharedData* SharedData = Binding.GetSharedData<FSharedData>();
 		check(SharedData);
 
@@ -51,14 +53,7 @@ namespace UE::AnimNext
 				InstanceData->CurrentTime -= SecondsToWait;
 				InstanceData->CurrentTransformIndex = (InstanceData->CurrentTransformIndex + 1) % Transforms.Num();
 			}
-
-			if (InstanceData->CurrentTransformIndex < Transforms.Num())
-			{
-				InstanceData->WarpTo = Transforms[InstanceData->CurrentTransformIndex];
-			}
 		}
-
-		IUpdate::PreUpdate(Context, Binding, TraitState);
 	}
 
 	void FWarpTestTrait::PostEvaluate(FEvaluateTraversalContext& Context, const TTraitBinding<IEvaluate>& Binding) const
@@ -71,7 +66,7 @@ namespace UE::AnimNext
 		FInstanceData* InstanceData = Binding.GetInstanceData<FInstanceData>();
 		check(InstanceData);
 
-		// hacky way to retrieve the mesh transform until we find a better way
+		// @todo: WIP hacky, non thread safe (unless proper tick dependencies are in place) way to retrieve the mesh transform until we find a better way
 		if (const FAnimNextModuleInstance* ModuleInstance = Context.GetRootGraphInstance().GetModuleInstance())
 		{
 			if (const UAnimNextComponent* AnimNextComponent = Cast<UAnimNextComponent>(ModuleInstance->GetObject()))
@@ -80,27 +75,22 @@ namespace UE::AnimNext
 				check(Actor);
 				if (const USkeletalMeshComponent* SkeletalMeshComponent = Actor->GetComponentByClass<USkeletalMeshComponent>())
 				{
-					InstanceData->ComponentTransform = SkeletalMeshComponent->GetComponentTransform();
+					const TArray<FTransform>& Transforms = SharedData->GetTransforms(Binding);
+					
+					FAnimNextWarpTestTask Task;
+					Task.ComponentTransform = SkeletalMeshComponent->GetComponentTransform();
+					Task.WarpToTransform = Transforms[InstanceData->CurrentTransformIndex];
+		
+#if ENABLE_ANIM_DEBUG 
+					Task.HostObject = Context.GetHostObject();
+#endif //ENABLE_ANIM_DEBUG 
+					Context.AppendTask(Task);
 				}
 			}
-		}
-
-#if ENABLE_ANIM_DEBUG 
-		InstanceData->HostObject = Context.GetHostObject();
-#endif // ENABLE_ANIM_DEBUG 
-
-		Context.AppendTask(FAnimNextWarpTestTask::Make(InstanceData));
+		}		
 	}
 } // namespace UE::AnimNext
 
-
-FAnimNextWarpTestTask FAnimNextWarpTestTask::Make(UE::AnimNext::FWarpTestTrait::FInstanceData* InstanceData)
-{
-	FAnimNextWarpTestTask Task;
-	Task.InstanceData = InstanceData;
-	return Task;
-}
-
 void FAnimNextWarpTestTask::Execute(UE::AnimNext::FEvaluationVM& VM) const
 {
 	using namespace UE::AnimNext;
@@ -116,19 +106,16 @@ void FAnimNextWarpTestTask::Execute(UE::AnimNext::FEvaluationVM& VM) const
 		// FTransform ThisFrameRootMotionTransform = FTransform::Identity;
 		// if (RootMotionProvider->ExtractRootMotion(Keyframe->Get()->Attributes, ThisFrameRootMotionTransform))
 		
-		const FTransform& WarpTo = InstanceData->WarpTo;
-		const FTransform& ComponentTransform = InstanceData->ComponentTransform;
-
 #if ENABLE_ANIM_DEBUG 
 		static const TCHAR* LogName = TEXT("WarpTest");
-		UE_VLOG_SEGMENT_THICK(InstanceData->HostObject, LogName, Display, WarpTo.GetLocation(), WarpTo.GetRotation().GetAxisX() * 100.f + WarpTo.GetLocation(), FColor::Red, 1, TEXT(""));
-		UE_VLOG_SEGMENT_THICK(InstanceData->HostObject, LogName, Display, WarpTo.GetLocation(), WarpTo.GetRotation().GetAxisY() * 100.f + WarpTo.GetLocation(), FColor::Blue, 1, TEXT(""));
+		UE_VLOG_SEGMENT_THICK(HostObject, LogName, Display, WarpToTransform.GetLocation(), WarpToTransform.GetRotation().GetAxisX() * 100.f + WarpToTransform.GetLocation(), FColor::Red, 1, TEXT(""));
+		UE_VLOG_SEGMENT_THICK(HostObject, LogName, Display, WarpToTransform.GetLocation(), WarpToTransform.GetRotation().GetAxisY() * 100.f + WarpToTransform.GetLocation(), FColor::Blue, 1, TEXT(""));
 
-		UE_VLOG_SEGMENT_THICK(InstanceData->HostObject, LogName, Display, ComponentTransform.GetLocation(), ComponentTransform.GetRotation().GetAxisX() * 80.f + ComponentTransform.GetLocation(), FColor::Black, 1, TEXT(""));
-		UE_VLOG_SEGMENT_THICK(InstanceData->HostObject, LogName, Display, ComponentTransform.GetLocation(), ComponentTransform.GetRotation().GetAxisY() * 80.f + ComponentTransform.GetLocation(), FColor::Green, 1, TEXT(""));
+		UE_VLOG_SEGMENT_THICK(HostObject, LogName, Display, ComponentTransform.GetLocation(), ComponentTransform.GetRotation().GetAxisX() * 80.f + ComponentTransform.GetLocation(), FColor::Black, 1, TEXT(""));
+		UE_VLOG_SEGMENT_THICK(HostObject, LogName, Display, ComponentTransform.GetLocation(), ComponentTransform.GetRotation().GetAxisY() * 80.f + ComponentTransform.GetLocation(), FColor::Green, 1, TEXT(""));
 #endif // ENABLE_ANIM_DEBUG 
 
-		const FTransform RootMotion = WarpTo.GetRelativeTransform(ComponentTransform);
+		const FTransform RootMotion = WarpToTransform.GetRelativeTransform(ComponentTransform);
 
 		RootMotionProvider->OverrideRootMotion(RootMotion, Keyframe->Get()->Attributes);
 	}
diff --git a/Engine/Plugins/Experimental/AnimNextWarping/Source/AnimNextWarping/Private/WarpTestTrait.h b/Engine/Plugins/Experimental/AnimNextWarping/Source/AnimNextWarping/Private/WarpTestTrait.h
--- a/Engine/Plugins/Experimental/AnimNextWarping/Source/AnimNextWarping/Private/WarpTestTrait.h
+++ b/Engine/Plugins/Experimental/AnimNextWarping/Source/AnimNextWarping/Private/WarpTestTrait.h
@@ -43,12 +43,6 @@ struct FWarpTestTrait : FAdditiveTrait, IUpdate, IEvaluate
 	{
 		int32 CurrentTransformIndex = 0;
 		float CurrentTime = 0.f;
-		FTransform WarpTo = FTransform::Identity;
-		FTransform ComponentTransform = FTransform::Identity;
-#if ENABLE_ANIM_DEBUG 
-		/** Debug Object for VisualLogger */
-		TObjectPtr<const UObject> HostObject = nullptr;
-#endif // ENABLE_ANIM_DEBUG 
 	};
 
 	// IUpdate impl 
@@ -67,9 +61,13 @@ struct FAnimNextWarpTestTask : public FAnimNextEvaluationTask
 
 	DECLARE_ANIM_EVALUATION_TASK(FAnimNextWarpTestTask)
 
-	static FAnimNextWarpTestTask Make(UE::AnimNext::FWarpTestTrait::FInstanceData* InstanceData);
-
 	virtual void Execute(UE::AnimNext::FEvaluationVM& VM) const override;
 
-	UE::AnimNext::FWarpTestTrait::FInstanceData* InstanceData = nullptr;
+	FTransform ComponentTransform = FTransform::Identity;
+	FTransform WarpToTransform = FTransform::Identity;
+
+#if ENABLE_ANIM_DEBUG
+	// Debug Object for VisualLogger
+	const UObject* HostObject = nullptr;
+#endif // ENABLE_ANIM_DEBUG
 };
\ No newline at end of file
