PoseSearch
- partial back out of CL 39422514 (initial MM interaction support for AnimNext)
- added ref counting pose histories

diff --git a/Engine/Plugins/Animation/PoseSearch/Source/Runtime/Private/AnimNode_PoseSearchHistoryCollector.cpp b/Engine/Plugins/Animation/PoseSearch/Source/Runtime/Private/AnimNode_PoseSearchHistoryCollector.cpp
--- a/Engine/Plugins/Animation/PoseSearch/Source/Runtime/Private/AnimNode_PoseSearchHistoryCollector.cpp
+++ b/Engine/Plugins/Animation/PoseSearch/Source/Runtime/Private/AnimNode_PoseSearchHistoryCollector.cpp
@@ -14,86 +14,29 @@
 /////////////////////////////////////////////////////
 // FAnimNode_PoseSearchHistoryCollector_Base
 
-// UE::PoseSearch::IPoseHistory interface
-bool FAnimNode_PoseSearchHistoryCollector_Base::GetTransformAtTime(float Time, FTransform& OutBoneTransform, const USkeleton* BoneIndexSkeleton, FBoneIndexType BoneIndexType, FBoneIndexType ReferenceBoneIndexType, bool bExtrapolate) const
-{
-	return PoseHistory.GetTransformAtTime(Time, OutBoneTransform, BoneIndexSkeleton, BoneIndexType, ReferenceBoneIndexType, bExtrapolate);
-}
-
-bool FAnimNode_PoseSearchHistoryCollector_Base::GetCurveValueAtTime(float Time, const FName& CurveName, float& OutCurveValue, bool bExtrapolate) const
-{
-	return PoseHistory.GetCurveValueAtTime(Time, CurveName, OutCurveValue, bExtrapolate);
-}
-
-const FTransformTrajectory& FAnimNode_PoseSearchHistoryCollector_Base::GetTrajectory() const
-{
-	return PoseHistory.GetTrajectory();
-}
-
-void FAnimNode_PoseSearchHistoryCollector_Base::SetTrajectory(const FTransformTrajectory& InTrajectory, float InTrajectorySpeedMultiplier)
-{
-	PoseHistory.SetTrajectory(InTrajectory, InTrajectorySpeedMultiplier);
-}
-
-float FAnimNode_PoseSearchHistoryCollector_Base::GetTrajectorySpeedMultiplier() const
-{
-	return PoseHistory.GetTrajectorySpeedMultiplier();
-}
-
-bool FAnimNode_PoseSearchHistoryCollector_Base::IsEmpty() const
-{
-	return PoseHistory.IsEmpty();
-}
-
-const UE::PoseSearch::FBoneToTransformMap& FAnimNode_PoseSearchHistoryCollector_Base::GetBoneToTransformMap() const
-{
-	return PoseHistory.GetBoneToTransformMap();
-}
-
-const TConstArrayView<FName> FAnimNode_PoseSearchHistoryCollector_Base::GetCollectedCurves() const
-{
-	return PoseHistory.GetCollectedCurves();
-}
-
-int32 FAnimNode_PoseSearchHistoryCollector_Base::GetNumEntries() const
-{
-	return PoseHistory.GetNumEntries();
-}
-
-const UE::PoseSearch::FPoseHistoryEntry& FAnimNode_PoseSearchHistoryCollector_Base::GetEntry(int32 EntryIndex) const
-{
-	return PoseHistory.GetEntry(EntryIndex);
-}
-
-#if ENABLE_DRAW_DEBUG && ENABLE_ANIM_DEBUG
-void FAnimNode_PoseSearchHistoryCollector_Base::DebugDraw(const UWorld* World, FColor Color) const
-{
-	PoseHistory.DebugDraw(World, Color);
-}
-
-void FAnimNode_PoseSearchHistoryCollector_Base::DebugDraw(FAnimInstanceProxy& AnimInstanceProxy, FColor Color) const
+void FAnimNode_PoseSearchHistoryCollector_Base::GenerateTrajectory(const UAnimInstance* InAnimInstance)
 {
-	PoseHistory.DebugDraw(AnimInstanceProxy, Color);
+	check(PoseHistoryPtr.IsValid());
+	PoseHistoryPtr->GenerateTrajectory(InAnimInstance, InAnimInstance->GetDeltaSeconds());
 }
-#endif
-// End of UE::PoseSearch::IPoseHistory interface
 
-void FAnimNode_PoseSearchHistoryCollector_Base::GenerateTrajectory(const UAnimInstance* InAnimInstance)
+void FAnimNode_PoseSearchHistoryCollector_Base::OnInitializeAnimInstance(const FAnimInstanceProxy* InProxy, const UAnimInstance* InAnimInstance)
 {
-	GenerateTrajectory(InAnimInstance, InAnimInstance->GetDeltaSeconds());
+	PoseHistoryPtr = MakeShareable(new UE::PoseSearch::FGenerateTrajectoryPoseHistory());
 }
 
 void FAnimNode_PoseSearchHistoryCollector_Base::Initialize_AnyThread(const FAnimationInitializeContext& Context)
 {
 	Super::Initialize_AnyThread(Context);
 
-	PoseHistory.Initialize_AnyThread(PoseCount, SamplingInterval);
+	check(PoseHistoryPtr.IsValid());
+	PoseHistoryPtr->Initialize_AnyThread(PoseCount, SamplingInterval);
 
 	FMemMark Mark(FMemStack::Get());
 	UE::PoseSearch::FAIPComponentSpacePoseProvider ComponentSpacePoseProvider(Context.AnimInstanceProxy);
 	if (ComponentSpacePoseProvider.GetSkeletonAsset())
 	{
-		PoseHistory.EvaluateComponentSpace_AnyThread(0.f, ComponentSpacePoseProvider, bStoreScales,
+		PoseHistoryPtr->EvaluateComponentSpace_AnyThread(0.f, ComponentSpacePoseProvider, bStoreScales,
 			RootBoneRecoveryTime, RootBoneTranslationRecoveryRatio, RootBoneRotationRecoveryRatio, true, true,
 			GetRequiredBones(Context.AnimInstanceProxy), FBlendedCurve(), MakeConstArrayView(CollectedCurves));
 	}
@@ -147,35 +90,28 @@ void FAnimNode_PoseSearchHistoryCollector_Base::Update_AnyThread(const FAnimatio
 	CheckInteractionThreadSafety(Context.AnimInstanceProxy->GetAnimInstanceObject());
 
 	GetEvaluateGraphExposedInputs().Execute(Context);
+	check(PoseHistoryPtr.IsValid());
 	
-	if (bGenerateTrajectory)
+	PoseHistoryPtr->bGenerateTrajectory = bGenerateTrajectory;
+	if (PoseHistoryPtr->bGenerateTrajectory)
 	{
-		GenerateTrajectory(Context.AnimInstanceProxy->GetAnimInstanceObject(), Context.GetDeltaTime());
+		PoseHistoryPtr->TrajectoryDataSampling.NumHistorySamples = FMath::Max(PoseCount, TrajectoryHistoryCount);
+		PoseHistoryPtr->TrajectoryDataSampling.SecondsPerHistorySample = SamplingInterval;
+		PoseHistoryPtr->TrajectoryDataSampling.NumPredictionSamples = TrajectoryPredictionCount;
+		PoseHistoryPtr->TrajectoryDataSampling.SecondsPerPredictionSample = PredictionSamplingInterval;
+		PoseHistoryPtr->TrajectoryData = TrajectoryData;
+
+		PoseHistoryPtr->GenerateTrajectory(Context.AnimInstanceProxy->GetAnimInstanceObject(), Context.GetDeltaTime());
 	}
 	else
 	{
-		PoseHistory.SetTrajectory(TransformTrajectory, TrajectorySpeedMultiplier);
+		PoseHistoryPtr->SetTrajectory(TransformTrajectory, TrajectorySpeedMultiplier);
 	}
 
-	bIsTrajectoryGeneratedBeforePreUpdate = false;
+	PoseHistoryPtr->bIsTrajectoryGeneratedBeforePreUpdate = false;
 	UpdateCounter.SynchronizeWith(Context.AnimInstanceProxy->GetUpdateCounter());
 }
 
-void FAnimNode_PoseSearchHistoryCollector_Base::GenerateTrajectory(const UObject* AnimContext, float DeltaTime)
-{
-	if (bGenerateTrajectory && !bIsTrajectoryGeneratedBeforePreUpdate)
-	{
-		FPoseSearchTrajectoryData::FSampling TrajectoryDataSampling;
-		TrajectoryDataSampling.NumHistorySamples = FMath::Max(PoseCount, TrajectoryHistoryCount);
-		TrajectoryDataSampling.SecondsPerHistorySample = SamplingInterval;
-		TrajectoryDataSampling.NumPredictionSamples = TrajectoryPredictionCount;
-		TrajectoryDataSampling.SecondsPerPredictionSample = PredictionSamplingInterval;
-
-		PoseHistory.GenerateTrajectory(AnimContext, DeltaTime, TrajectoryData, TrajectoryDataSampling);
-
-		bIsTrajectoryGeneratedBeforePreUpdate = true;
-	}
-}
 /////////////////////////////////////////////////////
 // FAnimNode_PoseSearchHistoryCollector
 
@@ -219,7 +155,7 @@ void FAnimNode_PoseSearchHistoryCollector::Evaluate_AnyThread(FPoseContext& Outp
 	}
 
 	UE::PoseSearch::FComponentSpacePoseProvider ComponentSpacePoseProvider(ComponentSpacePose);
-	PoseHistory.EvaluateComponentSpace_AnyThread(Output.AnimInstanceProxy->GetDeltaSeconds(), ComponentSpacePoseProvider, bStoreScales,
+	GetPoseHistory().EvaluateComponentSpace_AnyThread(Output.AnimInstanceProxy->GetDeltaSeconds(), ComponentSpacePoseProvider, bStoreScales,
 		RootBoneRecoveryTime, RootBoneTranslationRecoveryRatio, RootBoneRotationRecoveryRatio, bNeedsReset, bCacheBones, 
 		RequiredBones, Output.Curve, MakeConstArrayView(CollectedCurves));
 
@@ -232,7 +168,7 @@ void FAnimNode_PoseSearchHistoryCollector::Evaluate_AnyThread(FPoseContext& Outp
 #else // WITH_EDITORONLY_DATA
 	Color = FLinearColor::Red.ToFColor(true);
 #endif // WITH_EDITORONLY_DATA
-	PoseHistory.DebugDraw(*Output.AnimInstanceProxy, Color);
+	GetPoseHistory().DebugDraw(*Output.AnimInstanceProxy, Color);
 #endif // ENABLE_DRAW_DEBUG && ENABLE_ANIM_DEBUG
 }
 
@@ -282,7 +218,7 @@ void FAnimNode_PoseSearchComponentSpaceHistoryCollector::EvaluateComponentSpace_
 	}
 
 	UE::PoseSearch::FComponentSpacePoseProvider ComponentSpacePoseProvider(Output.Pose);
-	PoseHistory.EvaluateComponentSpace_AnyThread(Output.AnimInstanceProxy->GetDeltaSeconds(), ComponentSpacePoseProvider, bStoreScales, 
+	GetPoseHistory().EvaluateComponentSpace_AnyThread(Output.AnimInstanceProxy->GetDeltaSeconds(), ComponentSpacePoseProvider, bStoreScales, 
 		RootBoneRecoveryTime, RootBoneTranslationRecoveryRatio, RootBoneRotationRecoveryRatio, bNeedsReset, bCacheBones,
 		RequiredBones, Output.Curve, MakeConstArrayView(CollectedCurves));
 	
@@ -295,7 +231,7 @@ void FAnimNode_PoseSearchComponentSpaceHistoryCollector::EvaluateComponentSpace_
 #else // WITH_EDITORONLY_DATA
 	Color = FLinearColor::Red.ToFColor(true);
 #endif // WITH_EDITORONLY_DATA
-	PoseHistory.DebugDraw(*Output.AnimInstanceProxy, Color);
+	GetPoseHistory().DebugDraw(*Output.AnimInstanceProxy, Color);
 #endif // ENABLE_DRAW_DEBUG && ENABLE_ANIM_DEBUG
 }
 
diff --git a/Engine/Plugins/Animation/PoseSearch/Source/Runtime/Private/PoseHistoryProvider.cpp b/Engine/Plugins/Animation/PoseSearch/Source/Runtime/Private/PoseHistoryProvider.cpp
--- a/Engine/Plugins/Animation/PoseSearch/Source/Runtime/Private/PoseHistoryProvider.cpp
+++ b/Engine/Plugins/Animation/PoseSearch/Source/Runtime/Private/PoseHistoryProvider.cpp
@@ -9,7 +9,7 @@ namespace UE::PoseSearch
 const IPoseHistory& FPoseHistoryProvider::GetPoseHistory() const
 {
 	check(HistoryCollector);
-	return *HistoryCollector;
+	return HistoryCollector->GetPoseHistory();
 }
 
 } // namespace UE::PoseSearch
\ No newline at end of file
diff --git a/Engine/Plugins/Animation/PoseSearch/Source/Runtime/Private/PoseSearchHistory.cpp b/Engine/Plugins/Animation/PoseSearch/Source/Runtime/Private/PoseSearchHistory.cpp
--- a/Engine/Plugins/Animation/PoseSearch/Source/Runtime/Private/PoseSearchHistory.cpp
+++ b/Engine/Plugins/Animation/PoseSearch/Source/Runtime/Private/PoseSearchHistory.cpp
@@ -1063,6 +1063,17 @@ void FPoseHistory::DebugDraw(FAnimInstanceProxy& AnimInstanceProxy, FColor Color
 }
 #endif // ENABLE_DRAW_DEBUG && ENABLE_ANIM_DEBUG
 
+//////////////////////////////////////////////////////////////////////////
+// FGenerateTrajectoryPoseHistory
+void FGenerateTrajectoryPoseHistory::GenerateTrajectory(const UObject* AnimContext, float DeltaTime)
+{
+	if (bGenerateTrajectory && !bIsTrajectoryGeneratedBeforePreUpdate)
+	{
+		FPoseHistory::GenerateTrajectory(AnimContext, DeltaTime, TrajectoryData, TrajectoryDataSampling);
+		bIsTrajectoryGeneratedBeforePreUpdate = true;
+	}
+}
+
 //////////////////////////////////////////////////////////////////////////
 // FMemStackPoseHistory
 void FMemStackPoseHistory::Init(const IPoseHistory* InPoseHistory)
diff --git a/Engine/Plugins/Animation/PoseSearch/Source/Runtime/Private/PoseSearchHistoryCollectorAnimNodeLibrary.cpp b/Engine/Plugins/Animation/PoseSearch/Source/Runtime/Private/PoseSearchHistoryCollectorAnimNodeLibrary.cpp
--- a/Engine/Plugins/Animation/PoseSearch/Source/Runtime/Private/PoseSearchHistoryCollectorAnimNodeLibrary.cpp
+++ b/Engine/Plugins/Animation/PoseSearch/Source/Runtime/Private/PoseSearchHistoryCollectorAnimNodeLibrary.cpp
@@ -19,7 +19,7 @@ void UPoseSearchHistoryCollectorAnimNodeLibrary::GetPoseHistoryNodeTransformTraj
 {
 	if (FAnimNode_PoseSearchHistoryCollector* PoseSearchHistoryCollectorNodePtr = PoseSearchHistoryCollectorNode.GetAnimNodePtr<FAnimNode_PoseSearchHistoryCollector>())
 	{
-		Trajectory = PoseSearchHistoryCollectorNodePtr->GetTrajectory();
+		Trajectory = PoseSearchHistoryCollectorNodePtr->GetPoseHistory().GetTrajectory();
 	}
 }
 
@@ -27,7 +27,7 @@ void UPoseSearchHistoryCollectorAnimNodeLibrary::SetPoseHistoryNodeTransformTraj
 {
 	if (FAnimNode_PoseSearchHistoryCollector* PoseSearchHistoryCollectorNodePtr = PoseSearchHistoryCollectorNode.GetAnimNodePtr<FAnimNode_PoseSearchHistoryCollector>())
 	{
-		PoseSearchHistoryCollectorNodePtr->SetTrajectory(Trajectory);
+		PoseSearchHistoryCollectorNodePtr->GetPoseHistory().SetTrajectory(Trajectory);
 	}
 }
 
@@ -36,7 +36,7 @@ void UPoseSearchHistoryCollectorAnimNodeLibrary::GetPoseHistoryNodeTrajectory(co
 {
 	if (FAnimNode_PoseSearchHistoryCollector* PoseSearchHistoryCollectorNodePtr = PoseSearchHistoryCollectorNode.GetAnimNodePtr<FAnimNode_PoseSearchHistoryCollector>())
 	{
-		Trajectory = PoseSearchHistoryCollectorNodePtr->GetTrajectory();
+		Trajectory = PoseSearchHistoryCollectorNodePtr->GetPoseHistory().GetTrajectory();
 	}
 }
 
@@ -44,7 +44,7 @@ void UPoseSearchHistoryCollectorAnimNodeLibrary::SetPoseHistoryNodeTrajectory(co
 {
 	if (FAnimNode_PoseSearchHistoryCollector* PoseSearchHistoryCollectorNodePtr = PoseSearchHistoryCollectorNode.GetAnimNodePtr<FAnimNode_PoseSearchHistoryCollector>())
 	{
-		PoseSearchHistoryCollectorNodePtr->SetTrajectory(Trajectory);
+		PoseSearchHistoryCollectorNodePtr->GetPoseHistory().SetTrajectory(Trajectory);
 	}
 }
 PRAGMA_ENABLE_DEPRECATION_WARNINGS
\ No newline at end of file
diff --git a/Engine/Plugins/Animation/PoseSearch/Source/Runtime/Private/PoseSearchInteractionIsland.cpp b/Engine/Plugins/Animation/PoseSearch/Source/Runtime/Private/PoseSearchInteractionIsland.cpp
--- a/Engine/Plugins/Animation/PoseSearch/Source/Runtime/Private/PoseSearchInteractionIsland.cpp
+++ b/Engine/Plugins/Animation/PoseSearch/Source/Runtime/Private/PoseSearchInteractionIsland.cpp
@@ -35,7 +35,7 @@ static FInteractionSearchResult InitSearchResult(const FSearchResult& SearchResu
 	{
 		const UObject* AnimContext = SearchContext.AnimContexts[AnimContextIndex].Get();
 		const USkeleton* Skeleton = AnimContext ? GetContextSkeleton(AnimContext) : nullptr;
-		const IPoseHistory* PoseHistory = SearchContext.PoseHistories[AnimContextIndex];
+		const IPoseHistory* PoseHistory = SearchContext.GetPoseHistory(AnimContextIndex);
 		if (Skeleton && PoseHistory)
 		{
 			PoseHistory->GetTransformAtTime(0.f, InteractionSearchResult.ActorRootTransforms[AnimContextIndex], Skeleton, ComponentSpaceIndexType, WorldSpaceIndexType);
@@ -173,6 +173,16 @@ static void AddPrerequisite(FTickFunction& TickFunction, UObject* TargetObject,
 
 // FInteractionSearchContextBase
 ///////////////////////////////////////////////////////////
+const UE::PoseSearch::IPoseHistory* FInteractionSearchContextBase::GetPoseHistory(int32 Index) const
+{
+	check(PoseHistories.IsValidIndex(Index));
+	if (PoseHistories[Index].IsValid())
+	{
+		return PoseHistories[Index].Pin().Get();
+	}
+	return nullptr;
+}
+
 bool FInteractionSearchContextBase::IsEquivalent(const FInteractionSearchContextBase& Other) const
 {
 	if (Database != Other.Database)
@@ -320,16 +330,16 @@ void FInteractionSearchContextBase::TestHistoryCollectorsThreadingAccess() const
 {
 	for (int32 Index = 0; Index < AnimContexts.Num(); ++Index)
 	{
-		const UAnimInstance* AnimInstance = Cast<UAnimInstance>(AnimContexts[Index].Get());
-		const IPoseHistory* PoseHistory = PoseHistories[Index];
-		
-		if (AnimInstance && PoseHistory)
+		if (const UAnimInstance* AnimInstance = Cast<UAnimInstance>(AnimContexts[Index].Get()))
 		{
-			const USkeleton* Skeleton = AnimInstance->GetRequiredBonesOnAnyThread().GetSkeletonAsset();
-			check(Skeleton);
+			if (const IPoseHistory* PoseHistory = GetPoseHistory(Index))
+			{		
+				const USkeleton* Skeleton = AnimInstance->GetRequiredBonesOnAnyThread().GetSkeletonAsset();
+				check(Skeleton);
 
-			FTransform ActorRootBoneTransform;
-			PoseHistory->GetTransformAtTime(0.f, ActorRootBoneTransform, Skeleton, RootBoneIndexType, ComponentSpaceIndexType);
+				FTransform ActorRootBoneTransform;
+				PoseHistory->GetTransformAtTime(0.f, ActorRootBoneTransform, Skeleton, RootBoneIndexType, ComponentSpaceIndexType);
+			}
 		}
 	}
 }
@@ -361,11 +371,11 @@ void FInteractionIsland::FPreTickFunction::ExecuteTick(float DeltaTime, enum ELe
 			{
 				if (const UObject* AnimContext = SearchContext.AnimContexts[Index].Get())
 				{
-					const IPoseHistory* PoseHistory = SearchContext.PoseHistories[Index];
-					check(PoseHistory);
-
-					// since FInteractionIsland has a tick dependency with the USkeletalMeshComponent it's safe modify the IPoseHistory
-					const_cast<IPoseHistory*>(PoseHistory)->GenerateTrajectory(AnimContext, DeltaTime);
+					if (const IPoseHistory* PoseHistory = SearchContext.GetPoseHistory(Index))
+					{
+						// since FInteractionIsland has a tick dependency with the USkeletalMeshComponent it's safe modify the IPoseHistory
+						const_cast<IPoseHistory*>(PoseHistory)->GenerateTrajectory(AnimContext, DeltaTime);
+					}
 				}
 			}
 		}
@@ -831,6 +841,7 @@ bool FInteractionIsland::DoSearch_AnyThread(const UObject* AnimContext, const FP
 	FMemMark Mark(FMemStack::Get());
 
 	TArray<const UObject*, TInlineAllocator<PreallocatedRolesNum, TMemStackAllocator<>>> AnimContexts;
+	TArray<const UE::PoseSearch::IPoseHistory*, TInlineAllocator<PreallocatedRolesNum, TMemStackAllocator<>>> PoseHistories;
 	TArray<FSearchResult, TMemStackAllocator<>> PoseSearchResults;
 
 	// SearchContexts are modified only by UPoseSearchInteractionSubsystem::Tick and constant otherwise, so it's safe to access them in a threaded enviroment without locks
@@ -852,10 +863,15 @@ bool FInteractionIsland::DoSearch_AnyThread(const UObject* AnimContext, const FP
 			return false;
 		}
 
+		const int32 NumRoles = SearchContext.AnimContexts.Num();
+		check(NumRoles == SearchContext.PoseHistories.Num());
+
 		AnimContexts.Reset();
-		for (const TWeakObjectPtr<const UObject>& AnimContextPtr : SearchContext.AnimContexts)
+		PoseHistories.Reset();
+
+		for (int32 RoleIndex = 0; RoleIndex < NumRoles; ++RoleIndex)
 		{
-			const UObject* SearchContextAnimContext = AnimContextPtr.Get();
+			const UObject* SearchContextAnimContext = SearchContext.AnimContexts[RoleIndex].Get();
 			if (!SearchContextAnimContext)
 			{
 				UE_LOG(LogPoseSearch, Error, TEXT("FInteractionIsland::DoSearch_AnyThread null anim context"));
@@ -863,46 +879,23 @@ bool FInteractionIsland::DoSearch_AnyThread(const UObject* AnimContext, const FP
 			}
 
 			AnimContexts.Add(SearchContextAnimContext);
+
+			const UE::PoseSearch::IPoseHistory* PoseHistory = SearchContext.GetPoseHistory(RoleIndex);
+			if (!PoseHistory)
+			{
+				UE_LOG(LogPoseSearch, Error, TEXT("FInteractionIsland::DoSearch_AnyThread null PoseHistory"));
+				return false;
+			}
+
+			PoseHistories.Add(PoseHistory);
 		}
 
 		const UObject* AssetsToSearch[] = { Database };
 		FPoseSearchFutureProperties PoseSearchFutureProperties;
 
-		FPoseSearchContinuingProperties ContinuingPropertiesToUse = SearchContext.GetContinuingProperties();
-
-		// @todo: implement this carefully, since it'll cause thread safety concerns depending on which actors is calling the query first 
-		// and if their animations get integrated by different play rates
-			
-		//if (ContinuingProperties.PlayingAsset)
-		//{
-		//	if (const UMultiAnimAsset* MultiAnimAsset = Cast<UMultiAnimAsset>(ContinuingProperties.PlayingAsset))
-		//	{
-		//		// @todo: if they are compatible... etcetc MultiAnimAsset->GetNumRoles
-		//		ContinuingPropertiesToUse = ContinuingProperties;
-		//	}
-		//	else if (const UMultiAnimAsset* SearchContextMultiAnimAsset = Cast<UMultiAnimAsset>(SearchContext.ContinuingProperties.PlayingAsset))
-		//	{
-		//		for (int32 RoleIndex = 0; RoleIndex < SearchContextMultiAnimAsset->GetNumRoles(); ++RoleIndex)
-		//		{
-		//			if (SearchContextMultiAnimAsset->GetAnimationAsset(SearchContextMultiAnimAsset->GetRole(RoleIndex)) == ContinuingProperties.PlayingAsset)
-		//			{
-		//				ContinuingPropertiesToUse.PlayingAssetAccumulatedTime = ContinuingProperties.PlayingAssetAccumulatedTime;
-		//				break;
-		//			}
-		//		}
-		//	}
-		//	else if (!SearchContext.ContinuingProperties.PlayingAsset)
-		//	{
-		//		// @todo: should we scan all the database iterating over the UMultiAnimAsset to look for ContinuingProperties.PlayingAsset?
-		//	}
-		//	else
-		//	{
-		//	}
-		//}
-
 		// @todo: we could perform multiple UPoseSearchLibrary::MotionMatch in parallel!
 		PoseSearchResults[SearchIndex] = UPoseSearchLibrary::MotionMatch(AnimContexts, SearchContext.Roles,
-			SearchContext.PoseHistories, AssetsToSearch, ContinuingPropertiesToUse, PoseSearchFutureProperties, FPoseSearchEvent());
+			PoseHistories, AssetsToSearch, SearchContext.GetContinuingProperties(), PoseSearchFutureProperties, FPoseSearchEvent());
 	}
 		
 	InitSearchResults(SearchResults, PoseSearchResults, SearchContexts);
diff --git a/Engine/Plugins/Animation/PoseSearch/Source/Runtime/Private/PoseSearchInteractionSubsystem.cpp b/Engine/Plugins/Animation/PoseSearch/Source/Runtime/Private/PoseSearchInteractionSubsystem.cpp
--- a/Engine/Plugins/Animation/PoseSearch/Source/Runtime/Private/PoseSearchInteractionSubsystem.cpp
+++ b/Engine/Plugins/Animation/PoseSearch/Source/Runtime/Private/PoseSearchInteractionSubsystem.cpp
@@ -48,7 +48,7 @@ const IPoseHistory* FInteractionAvailability::GetPoseHistory(const UObject* Anim
 	{
 		if (const FAnimNode_PoseSearchHistoryCollector_Base* PoseSearchHistoryCollector = UPoseSearchLibrary::FindPoseHistoryNode(PoseHistoryName, AnimInstance))
 		{
-			return PoseSearchHistoryCollector;
+			return &PoseSearchHistoryCollector->GetPoseHistory();
 		}
 	}
 
@@ -898,9 +898,10 @@ void UPoseSearchInteractionSubsystem::GenerateSearchContexts(UE::PoseSearch::FIn
 								}
 							}
 
+							check(RoledAnimContextInfo.PoseHistory);
 							SearchContext.AnimContexts.Add(RoledAnimContextInfo.AnimContextInfo->AnimContext);
 							SearchContext.Roles.Add(RoledAnimContextInfo.Role);
-							SearchContext.PoseHistories.Add(RoledAnimContextInfo.PoseHistory);
+							SearchContext.PoseHistories.Add(RoledAnimContextInfo.PoseHistory->AsWeak());
 							SearchContext.bDisableCollisions |= RoledAnimContextInfo.bDisableCollisions;
 							SearchContext.TickPriorities.Add(RoledAnimContextInfo.AnimContextInfo->TickPriority);
 						}
diff --git a/Engine/Plugins/Animation/PoseSearch/Source/Runtime/Private/PoseSearchLibrary.cpp b/Engine/Plugins/Animation/PoseSearch/Source/Runtime/Private/PoseSearchLibrary.cpp
--- a/Engine/Plugins/Animation/PoseSearch/Source/Runtime/Private/PoseSearchLibrary.cpp
+++ b/Engine/Plugins/Animation/PoseSearch/Source/Runtime/Private/PoseSearchLibrary.cpp
@@ -924,7 +924,7 @@ void UPoseSearchLibrary::MotionMatch(
 	{
 		if (const FAnimNode_PoseSearchHistoryCollector_Base* PoseHistoryNode = FindPoseHistoryNode(PoseHistoryName, AnimInstance))
 		{
-			PoseHistories.Add(PoseHistoryNode);
+			PoseHistories.Add(&PoseHistoryNode->GetPoseHistory());
 		}
 
 		AnimContexts.Add(AnimInstance);
diff --git a/Engine/Plugins/Animation/PoseSearch/Source/Runtime/Public/PoseSearch/AnimNode_PoseSearchHistoryCollector.h b/Engine/Plugins/Animation/PoseSearch/Source/Runtime/Public/PoseSearch/AnimNode_PoseSearchHistoryCollector.h
--- a/Engine/Plugins/Animation/PoseSearch/Source/Runtime/Public/PoseSearch/AnimNode_PoseSearchHistoryCollector.h
+++ b/Engine/Plugins/Animation/PoseSearch/Source/Runtime/Public/PoseSearch/AnimNode_PoseSearchHistoryCollector.h
@@ -9,7 +9,7 @@
 #include "AnimNode_PoseSearchHistoryCollector.generated.h"
 
 USTRUCT(BlueprintInternalUseOnly)
-struct POSESEARCH_API FAnimNode_PoseSearchHistoryCollector_Base : public FAnimNode_Base, public UE::PoseSearch::IPoseHistory
+struct POSESEARCH_API FAnimNode_PoseSearchHistoryCollector_Base : public FAnimNode_Base
 {
 	GENERATED_BODY()
 
@@ -104,36 +104,15 @@ public:
 	bool bCacheBones = false;
 
 	// FAnimNode_Base interface
+	virtual bool NeedsOnInitializeAnimInstance() const override { return true; }
+	virtual void OnInitializeAnimInstance(const FAnimInstanceProxy* InProxy, const UAnimInstance* InAnimInstance) override;
 	virtual void Initialize_AnyThread(const FAnimationInitializeContext& Context) override;
 	virtual void CacheBones_AnyThread(const FAnimationCacheBonesContext& Context) override;
 	virtual void Update_AnyThread(const FAnimationUpdateContext& Context) override;
 	// End of FAnimNode_Base interface
 
-	// UE::PoseSearch::IPoseHistory interface
-	virtual bool GetTransformAtTime(float Time, FTransform& OutBoneTransform, const USkeleton* BoneIndexSkeleton = nullptr, FBoneIndexType BoneIndexType = UE::PoseSearch::RootBoneIndexType, FBoneIndexType ReferenceBoneIndexType = UE::PoseSearch::ComponentSpaceIndexType, bool bExtrapolate = false) const override;
-	virtual bool GetCurveValueAtTime(float Time, const FName& CurveName, float& OutCurveValue, bool bExtrapolate = false) const override;
-	virtual const FTransformTrajectory& GetTrajectory() const override;
-	// Experimental, this feature might be removed without warning, not for production use
-	virtual void SetTrajectory(const FTransformTrajectory& InTrajectory, float InTrajectorySpeedMultiplier = 1.f) override;
-	virtual void GenerateTrajectory(const UObject* AnimContext, float DeltaTime) override;
-	virtual float GetTrajectorySpeedMultiplier() const override;
-	virtual bool IsEmpty() const override;
-	virtual const UE::PoseSearch::FBoneToTransformMap& GetBoneToTransformMap() const override;
-	virtual const TConstArrayView<FName> GetCollectedCurves() const override;
-	virtual int32 GetNumEntries() const override;
-	virtual const UE::PoseSearch::FPoseHistoryEntry& GetEntry(int32 EntryIndex) const override;
-
-#if ENABLE_DRAW_DEBUG && ENABLE_ANIM_DEBUG
-	virtual void DebugDraw(const UWorld* World, FColor Color) const override;
-	virtual void DebugDraw(FAnimInstanceProxy& AnimInstanceProxy, FColor Color) const override;
-#endif
-	// End of UE::PoseSearch::IPoseHistory interface
-
-	UE_DEPRECATED(5.6, "Use this::IPoseHistory interface instead")
-	const UE::PoseSearch::FPoseHistory& GetPoseHistory() const { return PoseHistory; }
-
-	UE_DEPRECATED(5.6, "Use this::IPoseHistory interface instead")
-	UE::PoseSearch::FPoseHistory& GetPoseHistory() { return PoseHistory; }
+	const UE::PoseSearch::FPoseHistory& GetPoseHistory() const { check(PoseHistoryPtr.IsValid()); return *PoseHistoryPtr.Get(); }
+	UE::PoseSearch::FPoseHistory& GetPoseHistory() { check(PoseHistoryPtr.IsValid()); return *PoseHistoryPtr.Get(); }
 	
 	UE_DEPRECATED(5.6, "Use this::IPoseHistory::GenerateTrajectory interface instead")
 	void GenerateTrajectory(const UAnimInstance* InAnimInstance);
@@ -141,7 +120,7 @@ public:
 	PRAGMA_DISABLE_DEPRECATION_WARNINGS
 	// Note: We need to explicitly disable warnings on these constructors/operators for clang to be happy with deprecated variables
 	// this is a requirement for clang to compile without warnings.
-	 FAnimNode_PoseSearchHistoryCollector_Base() = default;
+	FAnimNode_PoseSearchHistoryCollector_Base() = default;
 	~FAnimNode_PoseSearchHistoryCollector_Base() = default;
 	FAnimNode_PoseSearchHistoryCollector_Base(const FAnimNode_PoseSearchHistoryCollector_Base&) = default;
 	FAnimNode_PoseSearchHistoryCollector_Base(FAnimNode_PoseSearchHistoryCollector_Base&&) = default;
@@ -153,8 +132,14 @@ protected:
 	// returns an array of skeleton bone indexes (encoded as FBoneIndexType)
 	TArray<FBoneIndexType> GetRequiredBones(const FAnimInstanceProxy* AnimInstanceProxy) const;
 
+	UE_DEPRECATED(5.6, "Use GetPoseHistory() instead")
 	UE::PoseSearch::FPoseHistory PoseHistory;
+
+	UE_DEPRECATED(5.6, "Use PoseHistoryPtr->bIsTrajectoryGeneratedBeforePreUpdate instead")
 	bool bIsTrajectoryGeneratedBeforePreUpdate = false;
+	
+private:
+	TSharedPtr<UE::PoseSearch::FGenerateTrajectoryPoseHistory, ESPMode::ThreadSafe> PoseHistoryPtr;
 };
 
 USTRUCT(BlueprintInternalUseOnly)
diff --git a/Engine/Plugins/Animation/PoseSearch/Source/Runtime/Public/PoseSearch/PoseSearchHistory.h b/Engine/Plugins/Animation/PoseSearch/Source/Runtime/Public/PoseSearch/PoseSearchHistory.h
--- a/Engine/Plugins/Animation/PoseSearch/Source/Runtime/Public/PoseSearch/PoseSearchHistory.h
+++ b/Engine/Plugins/Animation/PoseSearch/Source/Runtime/Public/PoseSearch/PoseSearchHistory.h
@@ -73,7 +73,7 @@ struct FPoseHistoryEntry
 };
 FArchive& operator<<(FArchive& Ar, FPoseHistoryEntry& Entry);
 
-struct POSESEARCH_API IPoseHistory
+struct POSESEARCH_API IPoseHistory : public TSharedFromThis<IPoseHistory, ESPMode::ThreadSafe>
 {
 	virtual ~IPoseHistory() {}
 	
@@ -156,6 +156,7 @@ struct POSESEARCH_API FPoseHistory : public IPoseHistory
 	FPoseHistory& operator=(const FPoseHistory& Other);
 	FPoseHistory& operator=(FPoseHistory&& Other);
 
+	// @todo: deprecate and delete this method in favor of providing the trajectory via SetTrajectory
 	void GenerateTrajectory(const UObject* AnimContext, float DeltaTime, const FPoseSearchTrajectoryData& TrajectoryData, const FPoseSearchTrajectoryData::FSampling& TrajectoryDataSampling);
 	
 	UE_DEPRECATED(5.6, "No longer necessary method")
@@ -239,6 +240,17 @@ private:
 	UE_MT_DECLARE_TS_RW_ACCESS_DETECTOR(PoseDataThreadSafeCounter);
 };
 
+// pose history struct that can generate a trajectory properly by calling it's IPoseHistory::GenerateTrajectory API (something other IPoseHistory don't do fully)
+struct POSESEARCH_API FGenerateTrajectoryPoseHistory : public FPoseHistory
+{
+	bool bGenerateTrajectory = false;
+	bool bIsTrajectoryGeneratedBeforePreUpdate = false;
+	FPoseSearchTrajectoryData::FSampling TrajectoryDataSampling;
+	FPoseSearchTrajectoryData TrajectoryData;
+
+	virtual void GenerateTrajectory(const UObject* AnimContext, float DeltaTime) override;
+};
+
 struct POSESEARCH_API FMemStackPoseHistory : public IPoseHistory
 {
 	void Init(const IPoseHistory* InPoseHistory);
diff --git a/Engine/Plugins/Animation/PoseSearch/Source/Runtime/Public/PoseSearch/PoseSearchInteractionIsland.h b/Engine/Plugins/Animation/PoseSearch/Source/Runtime/Public/PoseSearch/PoseSearchInteractionIsland.h
--- a/Engine/Plugins/Animation/PoseSearch/Source/Runtime/Public/PoseSearch/PoseSearchInteractionIsland.h
+++ b/Engine/Plugins/Animation/PoseSearch/Source/Runtime/Public/PoseSearch/PoseSearchInteractionIsland.h
@@ -19,7 +19,7 @@ struct IInteractionIslandDependency;
 struct FInteractionSearchContextBase
 {
 	TArray<TWeakObjectPtr<const UObject>, TInlineAllocator<PreallocatedRolesNum>> AnimContexts;
-	TArray<const UE::PoseSearch::IPoseHistory*, TInlineAllocator<PreallocatedRolesNum>> PoseHistories;
+	TArray<TWeakPtr<const UE::PoseSearch::IPoseHistory, ESPMode::ThreadSafe>, TInlineAllocator<PreallocatedRolesNum>> PoseHistories;
 	TArray<UE::PoseSearch::FRole, TInlineAllocator<PreallocatedRolesNum>> Roles;
 
 	TWeakObjectPtr<const UPoseSearchDatabase> Database;
@@ -27,6 +27,7 @@ struct FInteractionSearchContextBase
 
 	bool IsValid() const { return !AnimContexts.IsEmpty(); }
 	bool IsEquivalent(const FInteractionSearchContextBase& Other) const;
+	const UE::PoseSearch::IPoseHistory* GetPoseHistory(int32 Index) const;
 
 #if ENABLE_VISUAL_LOG
 	void VLogContext(const FColor& Color) const;
diff --git a/Engine/Plugins/Animation/PoseSearch/Source/Runtime/Public/PoseSearch/PoseSearchInteractionSubsystem.h b/Engine/Plugins/Animation/PoseSearch/Source/Runtime/Public/PoseSearch/PoseSearchInteractionSubsystem.h
--- a/Engine/Plugins/Animation/PoseSearch/Source/Runtime/Public/PoseSearch/PoseSearchInteractionSubsystem.h
+++ b/Engine/Plugins/Animation/PoseSearch/Source/Runtime/Public/PoseSearch/PoseSearchInteractionSubsystem.h
@@ -22,6 +22,7 @@ struct FInteractionAvailability : public FPoseSearchInteractionAvailability
 		, PoseHistoryName(InPoseHistoryName)
 		, PoseHistory(InPoseHistory)
 	{
+		check(!PoseHistory || PoseHistory->AsWeak().IsValid());
 	}
 
 	FString GetPoseHistoryName() const;
diff --git a/Engine/Plugins/Experimental/AnimNextPoseSearch/Source/AnimNextPoseSearch/Private/HistoryCollectorTrait.cpp b/Engine/Plugins/Experimental/AnimNextPoseSearch/Source/AnimNextPoseSearch/Private/HistoryCollectorTrait.cpp
--- a/Engine/Plugins/Experimental/AnimNextPoseSearch/Source/AnimNextPoseSearch/Private/HistoryCollectorTrait.cpp
+++ b/Engine/Plugins/Experimental/AnimNextPoseSearch/Source/AnimNextPoseSearch/Private/HistoryCollectorTrait.cpp
@@ -169,7 +169,8 @@ namespace UE::AnimNext
 		}
 		else
 		{
-			InstanceData->PoseHistory.SetTrajectory(SharedData->GetTrajectory(Binding));
+			check(InstanceData->PoseHistoryPtr.IsValid());
+			InstanceData->PoseHistoryPtr->SetTrajectory(SharedData->GetTrajectory(Binding));
 		}
 
 		IUpdate::PreUpdate(Context, Binding, TraitState);
@@ -183,8 +184,9 @@ namespace UE::AnimNext
 		FInstanceData* InstanceData = Binding.GetInstanceData<FInstanceData>();
 		const bool bStoreScales = SharedData->GetbStoreScales(Binding);
 
+		check(InstanceData->PoseHistoryPtr.IsValid());
 		FAnimNextHistoryCollectorTask Task;
-		Task.PoseHistory = &InstanceData->PoseHistory;
+		Task.PoseHistory = InstanceData->PoseHistoryPtr.Get();
 		Task.SharedData = SharedData;
 		Task.bStoreScales = bStoreScales;
 		Task.DeltaTime = InstanceData->DeltaTime;
@@ -194,6 +196,7 @@ namespace UE::AnimNext
 	const UE::PoseSearch::IPoseHistory* FHistoryCollectorTrait::GetPoseHistory(FExecutionContext& Context, const TTraitBinding<IPoseHistory>& Binding) const
 	{
 		const FInstanceData* InstanceData = Binding.GetInstanceData<FInstanceData>();
-		return &InstanceData->PoseHistory;
+		check(InstanceData->PoseHistoryPtr.IsValid());
+		return InstanceData->PoseHistoryPtr.Get();
 	}
 }
diff --git a/Engine/Plugins/Experimental/AnimNextPoseSearch/Source/AnimNextPoseSearch/Private/HistoryCollectorTrait.h b/Engine/Plugins/Experimental/AnimNextPoseSearch/Source/AnimNextPoseSearch/Private/HistoryCollectorTrait.h
--- a/Engine/Plugins/Experimental/AnimNextPoseSearch/Source/AnimNextPoseSearch/Private/HistoryCollectorTrait.h
+++ b/Engine/Plugins/Experimental/AnimNextPoseSearch/Source/AnimNextPoseSearch/Private/HistoryCollectorTrait.h
@@ -133,7 +133,8 @@ namespace UE::AnimNext
 
 		struct FInstanceData : FTrait::FInstanceData
 		{
-			UE::PoseSearch::FPoseHistory PoseHistory;
+			FInstanceData() : PoseHistoryPtr(MakeShareable(new UE::PoseSearch::FGenerateTrajectoryPoseHistory())) { }
+			TSharedPtr<UE::PoseSearch::FGenerateTrajectoryPoseHistory, ESPMode::ThreadSafe> PoseHistoryPtr;
 			float DeltaTime = 0.f;
 		};
 
