Option to set 'current time' when updating history in pose search trajectory
Mark the UpdateHistory_WorldSpace function as experimental. It is not released yet and so we are free to change it
Delete a verion of UpdateHistory_WorldSpace that was deprecated. It was never released so it is safe to delete
Update anim sandbox to match new API, fixes a bug where we record the current position at the wrong time

diff --git a/Engine/Plugins/Animation/PoseSearch/Source/Runtime/Private/PoseSearchTrajectoryLibrary.cpp b/Engine/Plugins/Animation/PoseSearch/Source/Runtime/Private/PoseSearchTrajectoryLibrary.cpp
--- a/Engine/Plugins/Animation/PoseSearch/Source/Runtime/Private/PoseSearchTrajectoryLibrary.cpp
+++ b/Engine/Plugins/Animation/PoseSearch/Source/Runtime/Private/PoseSearchTrajectoryLibrary.cpp
@@ -285,8 +285,12 @@ void UPoseSearchTrajectoryLibrary::UpdateHistory_TransformHistory(
 
 PRAGMA_ENABLE_DEPRECATION_WARNINGS
 
-void UPoseSearchTrajectoryLibrary::UpdateHistory_TransformHistory(FTransformTrajectory& Trajectory, FVector CurrentPosition, FVector CurrentVelocity,
-	const FPoseSearchTrajectoryData::FSampling& TrajectoryDataSampling, float DeltaTime)
+void UPoseSearchTrajectoryLibrary::UpdateHistory_TransformHistory(FTransformTrajectory& Trajectory,
+                                                                  FVector CurrentPosition,
+                                                                  FVector CurrentVelocity,
+                                                                  const FPoseSearchTrajectoryData::FSampling& TrajectoryDataSampling,
+                                                                  float DeltaTime,
+                                                                  float CurrentTime)
 {
 	const int32 NumHistorySamples = TrajectoryDataSampling.NumHistorySamples;
 	if (NumHistorySamples > 0)
@@ -312,7 +316,7 @@ void UPoseSearchTrajectoryLibrary::UpdateHistory_TransformHistory(FTransformTraj
 
 			// Adding a new history record
 			// Note we add current frame at delta time in the future since t = 0 should be the previous frame
-			Trajectory.Samples[NumHistorySamples - 1].TimeInSeconds = DeltaTime;
+			Trajectory.Samples[NumHistorySamples - 1].TimeInSeconds = CurrentTime;
 			Trajectory.Samples[NumHistorySamples - 1].Position = CurrentPosition;
 			Trajectory.Samples[NumHistorySamples - 1].Facing = Trajectory.Samples[NumHistorySamples].Facing;
 		}
@@ -329,25 +333,12 @@ void UPoseSearchTrajectoryLibrary::UpdateHistory_TransformHistory(FTransformTraj
 	}
 }
 
-PRAGMA_DISABLE_DEPRECATION_WARNINGS
-
-// Deprecated
-void UPoseSearchTrajectoryLibrary::UpdateHistory_WorldSpace(
-	FPoseSearchQueryTrajectory& Trajectory,
-	FVector CurrentPosition,
-	FQuat CurrentRotation,
-	const FPoseSearchTrajectoryData::FSampling& TrajectoryDataSampling,
-	float DeltaTime)
-{
-	FTransformTrajectory TransformTrajectory = Trajectory;
-	UpdateHistory_WorldSpace(TransformTrajectory, CurrentPosition, CurrentRotation, TrajectoryDataSampling, DeltaTime);
-	Trajectory = TransformTrajectory;
-}
-
-PRAGMA_ENABLE_DEPRECATION_WARNINGS
-
-void UPoseSearchTrajectoryLibrary::UpdateHistory_WorldSpace(FTransformTrajectory& Trajectory, FVector CurrentPosition, FQuat CurrentRotation, 
-	const FPoseSearchTrajectoryData::FSampling& TrajectoryDataSampling, float DeltaTime)
+void UPoseSearchTrajectoryLibrary::UpdateHistory_WorldSpace(FTransformTrajectory& Trajectory,
+                                                            FVector CurrentPosition,
+                                                            FQuat CurrentRotation,
+                                                            const FPoseSearchTrajectoryData::FSampling& TrajectoryDataSampling,
+                                                            float DeltaTime,
+                                                            float CurrentTime)
 {
 	const int32 NumHistorySamples = TrajectoryDataSampling.NumHistorySamples;
 	if (NumHistorySamples > 0)
@@ -368,7 +359,7 @@ void UPoseSearchTrajectoryLibrary::UpdateHistory_WorldSpace(FTransformTrajectory
 
 			// Adding a new history record
 			// Note we add current frame at delta time in the future since t = 0 should be the previous frame
-			Trajectory.Samples[NumHistorySamples - 1].TimeInSeconds = DeltaTime;
+			Trajectory.Samples[NumHistorySamples - 1].TimeInSeconds = CurrentTime;
 			Trajectory.Samples[NumHistorySamples - 1].Position = CurrentPosition;
 			Trajectory.Samples[NumHistorySamples - 1].Facing = CurrentRotation;
 		}
diff --git a/Engine/Plugins/Animation/PoseSearch/Source/Runtime/Public/PoseSearch/PoseSearchTrajectoryLibrary.h b/Engine/Plugins/Animation/PoseSearch/Source/Runtime/Public/PoseSearch/PoseSearchTrajectoryLibrary.h
--- a/Engine/Plugins/Animation/PoseSearch/Source/Runtime/Public/PoseSearch/PoseSearchTrajectoryLibrary.h
+++ b/Engine/Plugins/Animation/PoseSearch/Source/Runtime/Public/PoseSearch/PoseSearchTrajectoryLibrary.h
@@ -132,9 +132,6 @@ public:
 	UE_DEPRECATED(5.6, "Use the UpdateHistory_TransformHistory function thats in FTransformTrajectory instead of FPoseSearchQueryTrajectory.")
 	static void UpdateHistory_TransformHistory(FPoseSearchQueryTrajectory& Trajectory, FVector CurrentPosition, FVector CurrentVelocity, const FPoseSearchTrajectoryData::FSampling& TrajectoryDataSampling, float DeltaTime);
 
-	UE_DEPRECATED(5.6, "Use the UpdateHistory_WorldSpace function thats in FTransformTrajectory instead of FPoseSearchQueryTrajectory.")
-	static void UpdateHistory_WorldSpace(FPoseSearchQueryTrajectory& Trajectory, FVector CurrentPosition, FQuat CurrentRotation, const FPoseSearchTrajectoryData::FSampling& TrajectoryDataSampling, float DeltaTime);
-
 	UE_DEPRECATED(5.6, "Use the UpdatePrediction_SimulateCharacterMovement function thats in FTransformTrajectory instead of FPoseSearchQueryTrajectory.")
 	static void UpdatePrediction_SimulateCharacterMovement(FPoseSearchQueryTrajectory& Trajectory, const FPoseSearchTrajectoryData& TrajectoryData, const FPoseSearchTrajectoryData::FDerived& TrajectoryDataDerived, const FPoseSearchTrajectoryData::FSampling& TrajectoryDataSampling, float DeltaTime);
 
@@ -195,15 +192,48 @@ public:
 	// that to the current world transform. This works well on moving platforms as it only stores a history of movement
 	// that results from character intent, not movement from platforms.
 	// Important: CurrentVelocity should be the velocity relative to the ground as reported by the character movement component or character mover etc
-	static void UpdateHistory_TransformHistory(FTransformTrajectory& Trajectory, FVector CurrentPosition, FVector CurrentVelocity, const FPoseSearchTrajectoryData::FSampling& TrajectoryDataSampling, float DeltaTime);
+	static void UpdateHistory_TransformHistory(FTransformTrajectory& Trajectory,
+	                                           FVector CurrentPosition,
+	                                           FVector CurrentVelocity,
+	                                           const FPoseSearchTrajectoryData::FSampling& TrajectoryDataSampling,
+	                                           float DeltaTime,
+	                                           float CurrentTime);
+
+	static FORCEINLINE void UpdateHistory_TransformHistory(FTransformTrajectory& Trajectory,
+	                                                       FVector CurrentPosition,
+	                                                       FVector CurrentVelocity,
+	                                                       const FPoseSearchTrajectoryData::FSampling& TrajectoryDataSampling,
+	                                                       float DeltaTime)
+	{
+		// Old behaviour: assume this is called after position has been updated (e.g. post Character Movement), so use DeltaTime for current time
+		UpdateHistory_TransformHistory(Trajectory, CurrentPosition, CurrentVelocity, TrajectoryDataSampling, DeltaTime, DeltaTime);
+	}
 
 	// Update prediction by simulating the movement math for ground locomotion from UCharacterMovementComponent.
 	static void UpdatePrediction_SimulateCharacterMovement(FTransformTrajectory& Trajectory, const FPoseSearchTrajectoryData& TrajectoryData, const FPoseSearchTrajectoryData::FDerived& TrajectoryDataDerived, const FPoseSearchTrajectoryData::FSampling& TrajectoryDataSampling, float DeltaTime);
 
-	// Update the history purely based on current position, without taking into account ground velocities
-	// DeltaTime is the time since the last call to this method	
-	static void UpdateHistory_WorldSpace(FTransformTrajectory& Trajectory, FVector CurrentPosition, FQuat CurrentRotation, const FPoseSearchTrajectoryData::FSampling& TrajectoryDataSampling, float DeltaTime);
-	
+	// Experimental: Update the history purely based on current position, without taking into account ground velocities
+	// DeltaTime is the time since the last call to this method
+	UE_EXPERIMENTAL(5.6, "Experimental pose search function, API may change in the future");
+	static void UpdateHistory_WorldSpace(FTransformTrajectory& Trajectory,
+	                                     FVector CurrentPosition,
+	                                     FQuat CurrentRotation,
+	                                     const FPoseSearchTrajectoryData::FSampling& TrajectoryDataSampling,
+	                                     float DeltaTime,
+	                                     float CurrentTime);
+
+	// Experimental: Update the history purely based on current position, without taking into account ground velocities
+	UE_EXPERIMENTAL(5.6, "Experimental pose search function, API may change in the future");
+	static void UpdateHistory_WorldSpace(FTransformTrajectory& Trajectory,
+	                                     FVector CurrentPosition,
+	                                     FQuat CurrentRotation,
+	                                     const FPoseSearchTrajectoryData::FSampling& TrajectoryDataSampling,
+	                                     float DeltaTime)
+	{
+		// Old behaviour: assume this is called after position has been updated (e.g. post CMC)
+		UpdateHistory_WorldSpace(Trajectory, CurrentPosition, CurrentRotation, TrajectoryDataSampling, DeltaTime, DeltaTime);
+	}
+
 	// Generates a prediction trajectory based of the current character intent. For use with Character actors.
 	UFUNCTION(BlueprintCallable, Category = "Animation|PoseSearch", meta = (BlueprintThreadSafe, DisplayName = "Pose Search Generate Trajectory (for Character)"))
 	static void PoseSearchGenerateTransformTrajectory(const UObject* InAnimInstance, 
