PoseSearch - disabling UPoseSearchDatabase::PostLoad indexing WaitForCompletion in -game builds executions

diff --git a/Engine/Plugins/Animation/PoseSearch/Source/Runtime/Private/PoseSearchDatabase.cpp b/Engine/Plugins/Animation/PoseSearch/Source/Runtime/Private/PoseSearchDatabase.cpp
--- a/Engine/Plugins/Animation/PoseSearch/Source/Runtime/Private/PoseSearchDatabase.cpp
+++ b/Engine/Plugins/Animation/PoseSearch/Source/Runtime/Private/PoseSearchDatabase.cpp
@@ -908,18 +908,22 @@ void UPoseSearchDatabase::PostLoad()
 	// todo: should we SynchronizeWithExternalDependencies() here?
 	SynchronizeChooser();
 
-	ERequestAsyncBuildFlag Flag = ERequestAsyncBuildFlag::NewRequest;
-#if WITH_ENGINE
-	// If there isn't an EditorEngine (ex. Standalone Game via -game argument) we WaitForCompletion
-	if (Cast<UEditorEngine>(GEngine) == nullptr)
-	{
-		Flag |= ERequestAsyncBuildFlag::WaitForCompletion;
-	}
-#endif // WITH_ENGINE
-
 	// Delay indexing until BeginCacheForCookedPlatformData when running a CookCommandlet
 	if (!IsRunningCookCommandlet())
 	{
+		ERequestAsyncBuildFlag Flag = ERequestAsyncBuildFlag::NewRequest;
+#if WITH_ENGINE
+		// @todo: after CL 39338693 waiting for completion on a UPoseSearchDatabase is no longer possible,
+		//        because UAnimSequence are not implemented to wait for their DDC tasks (IAnimSequenceCompilingManager::FinishCompilation) in their PostLoad
+		//        leading FKeyBuilder::TryAddDependency not being able to compose a DDC key for the database, therefore not being able to complete the indexing.
+		// 
+		// If there isn't an EditorEngine (ex. Standalone Game via -game argument) we WaitForCompletion
+		//if (Cast<UEditorEngine>(GEngine) == nullptr)
+		//{
+		//	Flag |= ERequestAsyncBuildFlag::WaitForCompletion;
+		//}
+#endif // WITH_ENGINE
+	
 		FAsyncPoseSearchDatabasesManagement::RequestAsyncBuildIndex(this, Flag);
 	}
 #endif
