Build fix for experimental pose search method

diff --git a/Engine/Plugins/Animation/PoseSearch/Source/Runtime/Public/PoseSearch/PoseSearchTrajectoryLibrary.h b/Engine/Plugins/Animation/PoseSearch/Source/Runtime/Public/PoseSearch/PoseSearchTrajectoryLibrary.h
--- a/Engine/Plugins/Animation/PoseSearch/Source/Runtime/Public/PoseSearch/PoseSearchTrajectoryLibrary.h
+++ b/Engine/Plugins/Animation/PoseSearch/Source/Runtime/Public/PoseSearch/PoseSearchTrajectoryLibrary.h
@@ -214,7 +214,7 @@ public:
 
 	// Experimental: Update the history purely based on current position, without taking into account ground velocities
 	// DeltaTime is the time since the last call to this method
-	UE_EXPERIMENTAL(5.6, "Experimental pose search function, API may change in the future");
+	UE_EXPERIMENTAL(5.6, "Experimental pose search function, API may change in the future")
 	static void UpdateHistory_WorldSpace(FTransformTrajectory& Trajectory,
 	                                     FVector CurrentPosition,
 	                                     FQuat CurrentRotation,
@@ -223,7 +223,7 @@ public:
 	                                     float CurrentTime);
 
 	// Experimental: Update the history purely based on current position, without taking into account ground velocities
-	UE_EXPERIMENTAL(5.6, "Experimental pose search function, API may change in the future");
+	UE_EXPERIMENTAL(5.6, "Experimental pose search function, API may change in the future")
 	static void UpdateHistory_WorldSpace(FTransformTrajectory& Trajectory,
 	                                     FVector CurrentPosition,
 	                                     FQuat CurrentRotation,
@@ -231,7 +231,9 @@ public:
 	                                     float DeltaTime)
 	{
 		// Old behaviour: assume this is called after position has been updated (e.g. post CMC)
+		PRAGMA_DISABLE_EXPERIMENTAL_WARNINGS
 		UpdateHistory_WorldSpace(Trajectory, CurrentPosition, CurrentRotation, TrajectoryDataSampling, DeltaTime, DeltaTime);
+		PRAGMA_ENABLE_EXPERIMENTAL_WARNINGS
 	}
 
 	// Generates a prediction trajectory based of the current character intent. For use with Character actors.
