PoseSearch - fix for FAnimNode_PoseSearchHistoryCollector_Base::Initialize_AnyThread crashing when recompiling animation blue prints

diff --git a/Engine/Plugins/Animation/PoseSearch/Source/Runtime/Private/AnimNode_PoseSearchHistoryCollector.cpp b/Engine/Plugins/Animation/PoseSearch/Source/Runtime/Private/AnimNode_PoseSearchHistoryCollector.cpp
--- a/Engine/Plugins/Animation/PoseSearch/Source/Runtime/Private/AnimNode_PoseSearchHistoryCollector.cpp
+++ b/Engine/Plugins/Animation/PoseSearch/Source/Runtime/Private/AnimNode_PoseSearchHistoryCollector.cpp
@@ -30,7 +30,13 @@ void FAnimNode_PoseSearchHistoryCollector_Base::Initialize_AnyThread(const FAnim
 {
 	Super::Initialize_AnyThread(Context);
 
-	check(PoseHistoryPtr.IsValid());
+	// reinitializing PoseHistoryPtr if invalid here. it could happen when recompiling animation 
+	// blue prints when OnInitializeAnimInstance is not called
+	if (!PoseHistoryPtr.IsValid())
+	{
+		PoseHistoryPtr = MakeShareable(new UE::PoseSearch::FGenerateTrajectoryPoseHistory());
+	}
+
 	PoseHistoryPtr->Initialize_AnyThread(PoseCount, SamplingInterval);
 
 	FMemMark Mark(FMemStack::Get());
