PoseSearch
- improved AnimNext pose history discoverability
- removing unnecessary allocations
- stability improvements

diff --git a/Engine/Plugins/Animation/PoseSearch/Source/Runtime/Private/AnimNode_MotionMatchingInteraction.cpp b/Engine/Plugins/Animation/PoseSearch/Source/Runtime/Private/AnimNode_MotionMatchingInteraction.cpp
--- a/Engine/Plugins/Animation/PoseSearch/Source/Runtime/Private/AnimNode_MotionMatchingInteraction.cpp
+++ b/Engine/Plugins/Animation/PoseSearch/Source/Runtime/Private/AnimNode_MotionMatchingInteraction.cpp
@@ -68,7 +68,7 @@ void FAnimNode_MotionMatchingInteraction::UpdateAssetPlayer(const FAnimationUpda
 		ContinuingProperties.PlayingAsset = CurrentResult.SelectedAnimation; // GetAnimAsset();
 		ContinuingProperties.PlayingAssetAccumulatedTime = GetAccumulatedTime();
 
-		CurrentResult = UPoseSearchInteractionLibrary::MotionMatchInteraction(Availabilities, Context.AnimInstanceProxy->GetAnimInstanceObject(), ContinuingProperties, &PoseHistoryProvider->GetPoseHistory(), bValidateResultAgainstAvailabilities);
+		UPoseSearchInteractionLibrary::MotionMatchInteraction(CurrentResult, Availabilities, Context.AnimInstanceProxy->GetAnimInstanceObject(), ContinuingProperties, FName(),  & PoseHistoryProvider->GetPoseHistory(), bValidateResultAgainstAvailabilities);
 		check(CurrentResult.ActorRootTransforms.Num() == CurrentResult.ActorRootBoneTransforms.Num());
 
 		UAnimationAsset* RoledAnimAsset = nullptr;
diff --git a/Engine/Plugins/Animation/PoseSearch/Source/Runtime/Private/PoseSearchInteractionIsland.cpp b/Engine/Plugins/Animation/PoseSearch/Source/Runtime/Private/PoseSearchInteractionIsland.cpp
--- a/Engine/Plugins/Animation/PoseSearch/Source/Runtime/Private/PoseSearchInteractionIsland.cpp
+++ b/Engine/Plugins/Animation/PoseSearch/Source/Runtime/Private/PoseSearchInteractionIsland.cpp
@@ -823,7 +823,8 @@ bool FInteractionIsland::DoSearch_AnyThread(const UObject* AnimContext, const FP
 
 	if (bSearchPerfomed)
 	{
-		check(AnimContext != GetMainAnimContext());
+		// we now support multiple searches from the same AnimContext, and we return the already performed
+		// and cached result, so check(AnimContext != GetMainAnimContext()) is no longer valid!
 		return GetResult_AnyThread(AnimContext, Result);
 	}
 
diff --git a/Engine/Plugins/Animation/PoseSearch/Source/Runtime/Private/PoseSearchInteractionLibrary.cpp b/Engine/Plugins/Animation/PoseSearch/Source/Runtime/Private/PoseSearchInteractionLibrary.cpp
--- a/Engine/Plugins/Animation/PoseSearch/Source/Runtime/Private/PoseSearchInteractionLibrary.cpp
+++ b/Engine/Plugins/Animation/PoseSearch/Source/Runtime/Private/PoseSearchInteractionLibrary.cpp
@@ -9,26 +9,27 @@
 FPoseSearchInteractionBlueprintResult UPoseSearchInteractionLibrary::MotionMatchInteraction_Pure(TArray<FPoseSearchInteractionAvailability> Availabilities, const UObject* AnimContext, FPoseSearchContinuingProperties ContinuingProperties, FName PoseHistoryName, bool bValidateResultAgainstAvailabilities)
 {
 	FPoseSearchInteractionBlueprintResult Result;
-	if (UPoseSearchInteractionSubsystem* InteractionSubsystem = UPoseSearchInteractionSubsystem::GetSubsystem_AnyThread(AnimContext))
-	{
-		InteractionSubsystem->Query_AnyThread(Availabilities, AnimContext, ContinuingProperties, Result, PoseHistoryName, nullptr, bValidateResultAgainstAvailabilities);
-	}
+	MotionMatchInteraction(Result, Availabilities, AnimContext, ContinuingProperties, PoseHistoryName, nullptr, bValidateResultAgainstAvailabilities);
 	return Result;
 }
 
 FPoseSearchInteractionBlueprintResult UPoseSearchInteractionLibrary::MotionMatchInteraction(TArray<FPoseSearchInteractionAvailability> Availabilities, const UObject* AnimContext, FPoseSearchContinuingProperties ContinuingProperties, FName PoseHistoryName, bool bValidateResultAgainstAvailabilities)
 {
-	return MotionMatchInteraction_Pure(Availabilities, AnimContext, ContinuingProperties, PoseHistoryName, bValidateResultAgainstAvailabilities);
+	FPoseSearchInteractionBlueprintResult Result;
+	MotionMatchInteraction(Result, Availabilities, AnimContext, ContinuingProperties, PoseHistoryName, nullptr, bValidateResultAgainstAvailabilities);
+	return Result;
 }
 
-FPoseSearchInteractionBlueprintResult UPoseSearchInteractionLibrary::MotionMatchInteraction(const TArrayView<const FPoseSearchInteractionAvailability> Availabilities, const UObject* AnimContext, const FPoseSearchContinuingProperties& ContinuingProperties, const UE::PoseSearch::IPoseHistory* PoseHistory, bool bValidateResultAgainstAvailabilities)
+void UPoseSearchInteractionLibrary::MotionMatchInteraction(FPoseSearchInteractionBlueprintResult& Result, const TArrayView<const FPoseSearchInteractionAvailability> Availabilities, const UObject* AnimContext, const FPoseSearchContinuingProperties& ContinuingProperties, FName PoseHistoryName, const UE::PoseSearch::IPoseHistory* PoseHistory, bool bValidateResultAgainstAvailabilities)
 {
-	FPoseSearchInteractionBlueprintResult Result;
 	if (UPoseSearchInteractionSubsystem* InteractionSubsystem = UPoseSearchInteractionSubsystem::GetSubsystem_AnyThread(AnimContext))
 	{
-		InteractionSubsystem->Query_AnyThread(Availabilities, AnimContext, ContinuingProperties, Result, FName(), PoseHistory, bValidateResultAgainstAvailabilities);
+		InteractionSubsystem->Query_AnyThread(Availabilities, AnimContext, ContinuingProperties, Result, PoseHistoryName, PoseHistory, bValidateResultAgainstAvailabilities);
+	}
+	else
+	{
+		Result = FPoseSearchInteractionBlueprintResult();
 	}
-	return Result;
 }
 
 FPoseSearchInteractionBlueprintResult UPoseSearchInteractionLibrary::GetCachedInteraction(const UObject* AnimContext, bool bCompareOwningActors)
diff --git a/Engine/Plugins/Animation/PoseSearch/Source/Runtime/Private/PoseSearchInteractionSubsystem.cpp b/Engine/Plugins/Animation/PoseSearch/Source/Runtime/Private/PoseSearchInteractionSubsystem.cpp
--- a/Engine/Plugins/Animation/PoseSearch/Source/Runtime/Private/PoseSearchInteractionSubsystem.cpp
+++ b/Engine/Plugins/Animation/PoseSearch/Source/Runtime/Private/PoseSearchInteractionSubsystem.cpp
@@ -1018,10 +1018,10 @@ void UPoseSearchInteractionSubsystem::UpdateValidInteractionSearches()
 	using namespace UE::PoseSearch;
 
 	const int32 ValidInteractionSearchesNum = ValidInteractionSearches.Num();
-	TArray<bool, TInlineAllocator<PreallocatedSearchesNum * 4, TMemStackAllocator<>>> Visited;
+	TArray<bool, TInlineAllocator<PreallocatedSearchesNum * 4>> Visited;
 	Visited.SetNum(ValidInteractionSearchesNum);
 
-	TArray<FValidInteractionSearch, TInlineAllocator<PreallocatedSearchesNum * 4, TMemStackAllocator<>>> NewValidInteractionSearches;
+	TArray<FValidInteractionSearch, TInlineAllocator<PreallocatedSearchesNum * 4>> NewValidInteractionSearches;
 	for (FInteractionIsland* Island : Islands)
 	{
 		if (Island->IsInitialized())
diff --git a/Engine/Plugins/Animation/PoseSearch/Source/Runtime/Public/PoseSearch/PoseSearchInteractionLibrary.h b/Engine/Plugins/Animation/PoseSearch/Source/Runtime/Public/PoseSearch/PoseSearchInteractionLibrary.h
--- a/Engine/Plugins/Animation/PoseSearch/Source/Runtime/Public/PoseSearch/PoseSearchInteractionLibrary.h
+++ b/Engine/Plugins/Animation/PoseSearch/Source/Runtime/Public/PoseSearch/PoseSearchInteractionLibrary.h
@@ -49,6 +49,16 @@ struct POSESEARCH_API FPoseSearchInteractionAvailability
 	int32 TickPriority = UE::PoseSearch::DefaultTickPriority;
 
 	bool IsTagValid() const { return !Tag.IsNone(); }
+
+	bool operator==(const FPoseSearchInteractionAvailability& Other) const
+	{
+		return Database == Other.Database
+			&& Tag == Other.Tag
+			&& RolesFilter == Other.RolesFilter
+			&& BroadPhaseRadius == Other.BroadPhaseRadius
+			&& bDisableCollisions == Other.bDisableCollisions
+			&& TickPriority == Other.TickPriority;
+	}
 };
 
 // Experimental, this feature might be removed without warning, not for production use
@@ -136,8 +146,7 @@ public:
 	UFUNCTION(Experimental, BlueprintCallable, Category = "Animation|Pose Search", meta = (BlueprintThreadSafe))
 	static FPoseSearchInteractionBlueprintResult MotionMatchInteraction(TArray<FPoseSearchInteractionAvailability> Availabilities, const UObject* AnimContext, FPoseSearchContinuingProperties ContinuingProperties, FName PoseHistoryName, bool bValidateResultAgainstAvailabilities = true);
 
-	// version of MotionMatchInteraction_Pure referencing directly the HistoryCollector rather than looking for it by name
-	static FPoseSearchInteractionBlueprintResult MotionMatchInteraction(const TArrayView<const FPoseSearchInteractionAvailability> Availabilities, const UObject* AnimContext, const FPoseSearchContinuingProperties& ContinuingProperties, const UE::PoseSearch::IPoseHistory* PoseHistory, bool bValidateResultAgainstAvailabilities = true);
+	static void MotionMatchInteraction(FPoseSearchInteractionBlueprintResult& Result, const TArrayView<const FPoseSearchInteractionAvailability> Availabilities, const UObject* AnimContext, const FPoseSearchContinuingProperties& ContinuingProperties, FName PoseHistoryName, const UE::PoseSearch::IPoseHistory* PoseHistory, bool bValidateResultAgainstAvailabilities = true);
 
 	static FPoseSearchInteractionBlueprintResult GetCachedInteraction(const UObject* AnimContext, bool bCompareOwningActors = false);
 	
diff --git a/Engine/Plugins/Animation/PoseSearch/Source/Runtime/Public/PoseSearch/PoseSearchInteractionSubsystem.h b/Engine/Plugins/Animation/PoseSearch/Source/Runtime/Public/PoseSearch/PoseSearchInteractionSubsystem.h
--- a/Engine/Plugins/Animation/PoseSearch/Source/Runtime/Public/PoseSearch/PoseSearchInteractionSubsystem.h
+++ b/Engine/Plugins/Animation/PoseSearch/Source/Runtime/Public/PoseSearch/PoseSearchInteractionSubsystem.h
@@ -25,6 +25,13 @@ struct FInteractionAvailability : public FPoseSearchInteractionAvailability
 		check(!PoseHistory || PoseHistory->AsWeak().IsValid());
 	}
 
+	bool operator==(const FInteractionAvailability& Other) const
+	{
+		return static_cast<const FPoseSearchInteractionAvailability&>(*this) == static_cast<const FPoseSearchInteractionAvailability&>(Other)
+			&& PoseHistoryName == Other.PoseHistoryName
+			&& PoseHistory == Other.PoseHistory;
+	}
+
 	FString GetPoseHistoryName() const;
 	const IPoseHistory* GetPoseHistory(const UObject* AnimContext) const;
 
@@ -53,13 +60,12 @@ struct FAvailabilitiesMap : protected TMap<TWeakObjectPtr<const UObject>, FInter
 		{
 			if (Availability.Database || Availability.IsTagValid())
 			{
-				AnimContextAvailabilities.Emplace(Availability, PoseHistoryName, PoseHistory);
+				// Avoiding adding trivial duplicates. FInteractionAvailability could not be fully specified to understand if it's an actual duplicate in case the pose
+				// history is passed by name or the Availability.Database is null and supposed to be resolved using other availabilities Database(s) with the same Availability.Tag.
+				// The duplicated availabilities are excluded when creating the combinations of possible interactions during FAnimContextInfoVisitor when 
+				// FRoledAnimContextInfos.AddRoledAnimContextInfos calls FRoledAnimContextInfos.AddUnique
+				AnimContextAvailabilities.AddUnique(FInteractionAvailability(Availability, PoseHistoryName, PoseHistory));
 			}
-			// @todo: make this Warning less spammy and reenable it
-			//else
-			//{
-			//	UE_LOG(LogPoseSearch, Warning, TEXT("FAvailabilitiesMap::AddAvailabilities trying adding an invalid Availability for AnimContext %s!"), *AnimContext->GetName());
-			//}
 		}
 	}
 };
@@ -193,6 +199,7 @@ private:
 	static void OnInteractionEnd(UE::PoseSearch::FValidInteractionSearch& ValidInteractionSearch);
 
 	// @todo: use a thread local storage data structure to make publishing availabilities lockless to avoid the need of AvailabilitiesMapMutex
+	// it could contain duplicated avaialbilities
 	UE::PoseSearch::FAvailabilitiesMap AvailabilitiesMap;
 	FCriticalSection AvailabilitiesMapMutex;
 
diff --git a/Engine/Plugins/Experimental/AnimNextPoseSearch/Source/AnimNextPoseSearch/Private/HistoryCollectorTrait.cpp b/Engine/Plugins/Experimental/AnimNextPoseSearch/Source/AnimNextPoseSearch/Private/HistoryCollectorTrait.cpp
--- a/Engine/Plugins/Experimental/AnimNextPoseSearch/Source/AnimNextPoseSearch/Private/HistoryCollectorTrait.cpp
+++ b/Engine/Plugins/Experimental/AnimNextPoseSearch/Source/AnimNextPoseSearch/Private/HistoryCollectorTrait.cpp
@@ -181,9 +181,17 @@ namespace UE::AnimNext
 			InstanceData->PoseHistoryPtr->SetTrajectory(SharedData->GetTrajectory(Binding));
 		}
 
+		Context.PushScopedInterface<IPoseHistory>(Binding);
+
 		IUpdate::PreUpdate(Context, Binding, TraitState);
 	}
 
+	void FHistoryCollectorTrait::PostUpdate(FUpdateTraversalContext& Context, const TTraitBinding<IUpdate>& Binding, const FTraitUpdateState& TraitState) const
+	{
+		ensure(Context.PopScopedInterface<IPoseHistory>(Binding));
+		IUpdate::PostUpdate(Context, Binding, TraitState);
+	}
+
 	void FHistoryCollectorTrait::PostEvaluate(FEvaluateTraversalContext& Context, const TTraitBinding<IEvaluate>& Binding) const
 	{
 		IEvaluate::PostEvaluate(Context, Binding);
diff --git a/Engine/Plugins/Experimental/AnimNextPoseSearch/Source/AnimNextPoseSearch/Private/HistoryCollectorTrait.h b/Engine/Plugins/Experimental/AnimNextPoseSearch/Source/AnimNextPoseSearch/Private/HistoryCollectorTrait.h
--- a/Engine/Plugins/Experimental/AnimNextPoseSearch/Source/AnimNextPoseSearch/Private/HistoryCollectorTrait.h
+++ b/Engine/Plugins/Experimental/AnimNextPoseSearch/Source/AnimNextPoseSearch/Private/HistoryCollectorTrait.h
@@ -142,6 +142,7 @@ namespace UE::AnimNext
 
 		// IUpdate impl
 		virtual void PreUpdate(FUpdateTraversalContext& Context, const TTraitBinding<IUpdate>& Binding, const FTraitUpdateState& TraitState) const override;
+		virtual void PostUpdate(FUpdateTraversalContext& Context, const TTraitBinding<IUpdate>& Binding, const FTraitUpdateState& TraitState) const override;
 
 		// IEvaluate impl
 		virtual void PostEvaluate(FEvaluateTraversalContext& Context, const TTraitBinding<IEvaluate>& Binding) const override;
diff --git a/Engine/Plugins/Experimental/AnimNextPoseSearch/Source/AnimNextPoseSearch/Private/IPoseHistory.h b/Engine/Plugins/Experimental/AnimNextPoseSearch/Source/AnimNextPoseSearch/Private/IPoseHistory.h
--- a/Engine/Plugins/Experimental/AnimNextPoseSearch/Source/AnimNextPoseSearch/Private/IPoseHistory.h
+++ b/Engine/Plugins/Experimental/AnimNextPoseSearch/Source/AnimNextPoseSearch/Private/IPoseHistory.h
@@ -3,7 +3,7 @@
 #pragma once
 
 #include "CoreMinimal.h"
-#include "TraitCore/ITraitInterface.h"
+#include "TraitCore/IScopedTraitInterface.h"
 #include "TraitCore/TraitBinding.h"
 
 namespace UE::PoseSearch
@@ -16,7 +16,7 @@ namespace UE::AnimNext
 	/**
 	 * IPoseHistory
 	 */
-	struct IPoseHistory : ITraitInterface
+	struct IPoseHistory : IScopedTraitInterface
 	{
 		DECLARE_ANIM_TRAIT_INTERFACE(IPoseHistory)
 
diff --git a/Engine/Plugins/Experimental/AnimNextPoseSearch/Source/AnimNextPoseSearch/Private/MotionMatchingTrait.cpp b/Engine/Plugins/Experimental/AnimNextPoseSearch/Source/AnimNextPoseSearch/Private/MotionMatchingTrait.cpp
--- a/Engine/Plugins/Experimental/AnimNextPoseSearch/Source/AnimNextPoseSearch/Private/MotionMatchingTrait.cpp
+++ b/Engine/Plugins/Experimental/AnimNextPoseSearch/Source/AnimNextPoseSearch/Private/MotionMatchingTrait.cpp
@@ -79,7 +79,7 @@ namespace UE::AnimNext
 		}
 
 		TTraitBinding<IPoseHistory> PoseHistoryTrait;
-		if (!Binding.GetStackInterface(PoseHistoryTrait))
+		if (!Context.GetScopedInterface<IPoseHistory>(PoseHistoryTrait))
 		{
 			UE_LOG(LogPoseSearch, Error, TEXT("FMotionMatchingTrait::PreUpdate, missing IPoseHistory"));
 			return;
@@ -187,7 +187,7 @@ namespace UE::AnimNext
 			ContinuingProperties.PlayingAsset = MotionMatchingInteractionState.SelectedAnimation; // GetAnimAsset();
 			ContinuingProperties.PlayingAssetAccumulatedTime = MotionMatchingInteractionState.SelectedTime;
 
-			MotionMatchingInteractionState = UPoseSearchInteractionLibrary::MotionMatchInteraction(Availabilities, AnimContext, ContinuingProperties, PoseHistory, SharedData->GetbValidateResultAgainstAvailabilities(Binding));
+			UPoseSearchInteractionLibrary::MotionMatchInteraction(MotionMatchingInteractionState, Availabilities, AnimContext, ContinuingProperties, FName(), PoseHistory, SharedData->GetbValidateResultAgainstAvailabilities(Binding));
 			check(MotionMatchingInteractionState.ActorRootTransforms.Num() == MotionMatchingInteractionState.ActorRootBoneTransforms.Num());
 		}
 
