PoseSearch
- fix for database blendspaces indexing in ProcessRootTransform, sometime missing the last root transform sample, that was leading to incorrect first frame velocity for indexed looping blendspaces (resulted sometimes zero velocities...)
- added test for blendspace sampling determinism
- added BlendParams display row in database editor data detail tab
- small indexing performance improvement

diff --git a/Engine/Plugins/Animation/PoseSearch/Source/Editor/Private/PoseSearchDatabaseDataDetails.cpp b/Engine/Plugins/Animation/PoseSearch/Source/Editor/Private/PoseSearchDatabaseDataDetails.cpp
--- a/Engine/Plugins/Animation/PoseSearch/Source/Editor/Private/PoseSearchDatabaseDataDetails.cpp
+++ b/Engine/Plugins/Animation/PoseSearch/Source/Editor/Private/PoseSearchDatabaseDataDetails.cpp
@@ -360,9 +360,46 @@ public:
 													const int32 AssetIdx = PreviewActor.GetIndexAssetIndex();
 													if (Assets.IsValidIndex(AssetIdx))
 													{
-														FString AssetIdxString;
-														AssetIdxString.AppendInt(Assets[AssetIdx].GetSourceAssetIdx());
-														return FText::FromString(AssetIdxString);
+														FString SourceAssetIndexString;
+														SourceAssetIndexString.AppendInt(Assets[AssetIdx].GetSourceAssetIdx());
+														return FText::FromString(SourceAssetIndexString);
+													}
+												}
+											}
+										}
+									}
+								}
+							}
+							return FText::GetEmpty();
+						});
+			}
+
+			if (ChannelItem->GetCompactLabel() == "BlendParams")
+			{
+				return SNew(STextBlock)
+					.Margin(FMargin(1.f, 1.f, 1.f, 1.f))
+					.Text_Lambda([this, ColumnName]() -> FText
+						{
+							if (!ChannelItem->IsExpanded())
+							{
+								if (const TSharedPtr<FDatabaseViewModel> ViewModel = EditorViewModel.Pin())
+								{
+									const UPoseSearchDatabase* PoseSearchDatabase = ViewModel->GetPoseSearchDatabase();
+									if (EAsyncBuildIndexResult::Success == FAsyncPoseSearchDatabasesManagement::RequestAsyncBuildIndex(PoseSearchDatabase, ERequestAsyncBuildFlag::ContinueRequest))
+									{
+										for (const TArray<FDatabasePreviewActor>& PreviewActorGroup : ViewModel->GetPreviewActors())
+										{
+											for (const FDatabasePreviewActor& PreviewActor : PreviewActorGroup)
+											{
+												if (*GetNameSafe(PreviewActor.GetActor()) == ColumnName)
+												{
+													const TConstArrayView<FSearchIndexAsset> Assets = PoseSearchDatabase->GetSearchIndex().Assets;
+													const int32 AssetIdx = PreviewActor.GetIndexAssetIndex();
+													if (Assets.IsValidIndex(AssetIdx))
+													{
+														const FVector BlendParameters = Assets[AssetIdx].GetBlendParameters();
+														FString BlendParamsString = FString::Printf(TEXT("%f, %f"), BlendParameters.X, BlendParameters.Y);
+														return FText::FromString(BlendParamsString);
 													}
 												}
 											}
@@ -616,6 +653,7 @@ void SDatabaseDataDetails::RebuildChannelItemsStats(TArray<FChannelItemPtr>& Cha
 	ChannelItem->GetChannelItems().Add(MakeShareable(new FChannelItem("SDatabaseDataDetailsPoseIndex", "PoseIndex")));
 	ChannelItem->GetChannelItems().Add(MakeShareable(new FChannelItem("SDatabaseDataDetailsPoseIndex", "SearchAssetIndex")));
 	ChannelItem->GetChannelItems().Add(MakeShareable(new FChannelItem("SDatabaseDataDetailsPoseIndex", "SourceAssetIndex")));
+	ChannelItem->GetChannelItems().Add(MakeShareable(new FChannelItem("SDatabaseDataDetailsPoseIndex", "BlendParams")));
 	// @todo: add support for additional stats
 	//ChannelItem->GetChannelItems().Add(MakeShareable(new FChannelItem("SDatabaseDataDetailsValueVectorIndex", "ValueVectorIndex")));
 	//ChannelItem->GetChannelItems().Add(MakeShareable(new FChannelItem("SDatabaseDataDetailsPCAValueVectorIndex", "PCAValueVectorIndex")));
diff --git a/Engine/Plugins/Animation/PoseSearch/Source/Runtime/Private/PoseSearchAssetSampler.cpp b/Engine/Plugins/Animation/PoseSearch/Source/Runtime/Private/PoseSearchAssetSampler.cpp
--- a/Engine/Plugins/Animation/PoseSearch/Source/Runtime/Private/PoseSearchAssetSampler.cpp
+++ b/Engine/Plugins/Animation/PoseSearch/Source/Runtime/Private/PoseSearchAssetSampler.cpp
@@ -157,7 +157,7 @@ static void ProcessRootTransform(const UBlendSpace* BlendSpace, const FVector& B
 	int32 RootTransformSamplingRate, bool bIsLoopable, TArray<FTransform>& AccumulatedRootTransform)
 {
 	// Pre-compute root motion
-	const int32 NumRootSamples = FMath::Max(CachedPlayLength * RootTransformSamplingRate + 1, 1);
+	const int32 NumRootSamples = FMath::Max(FMath::CeilToInt(CachedPlayLength * RootTransformSamplingRate) + 1, 1);
 	AccumulatedRootTransform.Init(FTransform::Identity, NumRootSamples);
 	
 	TArray<FBlendSampleData> BlendSamplesData;
diff --git a/Engine/Plugins/Animation/PoseSearch/Source/Runtime/Private/PoseSearchDatabase.cpp b/Engine/Plugins/Animation/PoseSearch/Source/Runtime/Private/PoseSearchDatabase.cpp
--- a/Engine/Plugins/Animation/PoseSearch/Source/Runtime/Private/PoseSearchDatabase.cpp
+++ b/Engine/Plugins/Animation/PoseSearch/Source/Runtime/Private/PoseSearchDatabase.cpp
@@ -392,7 +392,7 @@ const FString FPoseSearchDatabaseAnimationAssetBase::GetName() const
 UAnimationAsset* FPoseSearchDatabaseAnimationAssetBase::GetAnimationAssetForRole(const UE::PoseSearch::FRole& Role) const
 {
 	check(GetNumRoles() == 1);
-	return CastChecked<UAnimationAsset>(GetAnimationAsset());
+	return Cast<UAnimationAsset>(GetAnimationAsset());
 }
 
 FTransform FPoseSearchDatabaseAnimationAssetBase::GetRootTransformOriginForRole(const UE::PoseSearch::FRole& Role) const
@@ -979,7 +979,7 @@ void UPoseSearchDatabase::SynchronizeWithExternalDependencies()
 		{
 			if (Asset.IsInstanceOf(UAnimSequenceBase::StaticClass()))
 			{
-				if (UAnimSequenceBase* SequenceBase = CastChecked<UAnimSequenceBase>(Asset.FastGetAsset(true)))
+				if (UAnimSequenceBase* SequenceBase = Cast<UAnimSequenceBase>(Asset.FastGetAsset(true)))
 				{
 					for (const FAnimNotifyEvent& NotifyEvent : SequenceBase->Notifies)
 					{
diff --git a/Engine/Plugins/Animation/PoseSearch/Source/Runtime/Private/PoseSearchDerivedData.cpp b/Engine/Plugins/Animation/PoseSearch/Source/Runtime/Private/PoseSearchDerivedData.cpp
--- a/Engine/Plugins/Animation/PoseSearch/Source/Runtime/Private/PoseSearchDerivedData.cpp
+++ b/Engine/Plugins/Animation/PoseSearch/Source/Runtime/Private/PoseSearchDerivedData.cpp
@@ -913,22 +913,25 @@ static bool IndexDatabase(FSearchIndexBase& SearchIndexBase, const UPoseSearchDa
 	{
 		if (const FPoseSearchDatabaseAnimationAssetBase* DatabaseAnimationAssetBase = DatabaseToLookForAssets->GetDatabaseAnimationAsset<FPoseSearchDatabaseAnimationAssetBase>(AnimationAssetIndex))
 		{
-			DatabaseAnimationAssetBase->IterateOverSamplingParameter([DatabaseAnimationAssetBase, &SamplerMap, &Samplers](const FVector& BlendParameters)
-				{
-					const int32 NumRoles = DatabaseAnimationAssetBase->GetNumRoles();
-					for (int32 RoleIndex = 0; RoleIndex < NumRoles; ++RoleIndex)
+			if (DatabaseAnimationAssetBase->IsEnabled())
+			{
+				DatabaseAnimationAssetBase->IterateOverSamplingParameter([DatabaseAnimationAssetBase, &SamplerMap, &Samplers](const FVector& BlendParameters)
 					{
-						const FRole& Role = DatabaseAnimationAssetBase->GetRole(RoleIndex);
-						const UAnimationAsset* AnimationAsset = DatabaseAnimationAssetBase->GetAnimationAssetForRole(Role);
-						const FTransform RootTransformOrigin = DatabaseAnimationAssetBase->GetRootTransformOriginForRole(Role);
-						const FSamplerMapKey SamplerMapKey(AnimationAsset, RootTransformOrigin, BlendParameters);
-						if (!SamplerMap.Contains(SamplerMapKey))
+						const int32 NumRoles = DatabaseAnimationAssetBase->GetNumRoles();
+						for (int32 RoleIndex = 0; RoleIndex < NumRoles; ++RoleIndex)
 						{
-							SamplerMap.Add(SamplerMapKey, Samplers.Num());
-							Samplers.Emplace(AnimationAsset, RootTransformOrigin, BlendParameters, FAnimationAssetSampler::DefaultRootTransformSamplingRate, false);
+							const FRole& Role = DatabaseAnimationAssetBase->GetRole(RoleIndex);
+							const UAnimationAsset* AnimationAsset = DatabaseAnimationAssetBase->GetAnimationAssetForRole(Role);
+							const FTransform RootTransformOrigin = DatabaseAnimationAssetBase->GetRootTransformOriginForRole(Role);
+							const FSamplerMapKey SamplerMapKey(AnimationAsset, RootTransformOrigin, BlendParameters);
+							if (!SamplerMap.Contains(SamplerMapKey))
+							{
+								SamplerMap.Add(SamplerMapKey, Samplers.Num());
+								Samplers.Emplace(AnimationAsset, RootTransformOrigin, BlendParameters, FAnimationAssetSampler::DefaultRootTransformSamplingRate, false);
+							}
 						}
-					}
-				});
+					});
+			}
 		}
 	}
 
@@ -2318,8 +2321,10 @@ void FAsyncPoseSearchDatabasesManagement::Tick(float DeltaTime)
 			const UPoseSearchDatabase* Database = nullptr;
 			int32 AnimationAssetIndex = INDEX_NONE;
 			const UAnimationAsset* AnimationAsset = nullptr;
+			FTransform RootTransformOrigin = FTransform::Identity;
+			FVector BlendParameters = FVector::ZeroVector;
 			int32 SampleIndex = INDEX_NONE;
-			float SampleTime = 0.f;
+			float SampleNormalizedTime = 0.f;
 
 			TAlignedArray<uint8> SerializedData;
 
@@ -2344,29 +2349,36 @@ void FAsyncPoseSearchDatabasesManagement::Tick(float DeltaTime)
 					{
 						if (const FPoseSearchDatabaseAnimationAssetBase* DatabaseAsset = Database->GetDatabaseAnimationAsset<FPoseSearchDatabaseAnimationAssetBase>(AnimationAssetIndex))
 						{
-							for (int32 RoleIndex = 0; RoleIndex < DatabaseAsset->GetNumRoles(); ++RoleIndex)
+							if (DatabaseAsset->IsEnabled())
 							{
-								const FRole Role = DatabaseAsset->GetRole(RoleIndex);
-								if (const UAnimationAsset* AnimationAsset = DatabaseAsset->GetAnimationAssetForRole(Role))
-								{
-									if (const FBoneContainer* BoneContainer = RoledBoneContainers.Find(Role))
+								DatabaseAsset->IterateOverSamplingParameter([Database, DatabaseAsset, AnimationAssetIndex, &TestSamples, &RoledBoneContainers](const FVector& BlendParameters)
 									{
-										static float SAMPLING_TIME = 1 / 30.f;
-										const float PlayLength = AnimationAsset->GetPlayLength();
-										const int32 NumSamples = FMath::FloorToInt(PlayLength / SAMPLING_TIME);
-
-										for (int32 SampleIndex = 0; SampleIndex < NumSamples; ++SampleIndex)
+										const int32 NumRoles = DatabaseAsset->GetNumRoles();
+										for (int32 RoleIndex = 0; RoleIndex < NumRoles; ++RoleIndex)
 										{
-											FTestSample& TestSample = TestSamples.AddDefaulted_GetRef();
-											TestSample.BoneContainer = *BoneContainer;
-											TestSample.Database = Database;
-											TestSample.AnimationAssetIndex = AnimationAssetIndex;
-											TestSample.AnimationAsset = AnimationAsset;
-											TestSample.SampleIndex = SampleIndex;
-											TestSample.SampleTime = SampleIndex * SAMPLING_TIME;
+											const FRole& Role = DatabaseAsset->GetRole(RoleIndex);
+											if (const FBoneContainer* BoneContainer = RoledBoneContainers.Find(Role))
+											{
+												const UAnimationAsset* AnimationAsset = DatabaseAsset->GetAnimationAssetForRole(Role);
+												const FTransform RootTransformOrigin = DatabaseAsset->GetRootTransformOriginForRole(Role);
+
+												static int32 NUM_TESTING_SAMPLES_PER_DATABASE_ASSET = 100; // make sure it's greater than 1
+
+												for (int32 SampleIndex = 0; SampleIndex < NUM_TESTING_SAMPLES_PER_DATABASE_ASSET; SampleIndex++)
+												{
+													FTestSample& TestSample = TestSamples.AddDefaulted_GetRef();
+													TestSample.BoneContainer = *BoneContainer;
+													TestSample.Database = Database;
+													TestSample.AnimationAssetIndex = AnimationAssetIndex;
+													TestSample.AnimationAsset = AnimationAsset;
+													TestSample.RootTransformOrigin = RootTransformOrigin;
+													TestSample.BlendParameters = BlendParameters;
+													TestSample.SampleIndex = SampleIndex;
+													TestSample.SampleNormalizedTime = SampleIndex / (NUM_TESTING_SAMPLES_PER_DATABASE_ASSET - 1);
+												}
+											}
 										}
-									}
-								}
+									});
 							}
 						}
 					}
@@ -2380,11 +2392,12 @@ void FAsyncPoseSearchDatabasesManagement::Tick(float DeltaTime)
 
 				FTestSample& TestSample = TestSamples[TestSampleIndex];
 				const FAnimationAssetSampler AssetSampler(TestSample.AnimationAsset);
+				const float SampleTime = AssetSampler.GetPlayLength() * TestSample.SampleNormalizedTime;
 
 				FCompactPose Pose;
 				Pose.SetBoneContainer(&TestSample.BoneContainer);
-				AssetSampler.ExtractPose(TestSample.SampleTime, Pose);
-				const FTransform RootTransform = AssetSampler.ExtractRootTransform(TestSample.SampleTime);
+				AssetSampler.ExtractPose(SampleTime, Pose);
+				const FTransform RootTransform = AssetSampler.ExtractRootTransform(SampleTime);
 
 				const TConstArrayView<FTransform> Bones = Pose.GetBones();
 
@@ -2394,8 +2407,8 @@ void FAsyncPoseSearchDatabasesManagement::Tick(float DeltaTime)
 					{
 						FCompactPose TestPose;
 						TestPose.SetBoneContainer(&TestSample.BoneContainer);
-						AssetSampler.ExtractPose(TestSample.SampleTime, TestPose);
-						const FTransform TestRootTransform = AssetSampler.ExtractRootTransform(TestSample.SampleTime);
+						AssetSampler.ExtractPose(SampleTime, TestPose);
+						const FTransform TestRootTransform = AssetSampler.ExtractRootTransform(SampleTime);
 
 						if (FMemory::Memcmp(&RootTransform, &TestRootTransform, sizeof(FTransform)) != 0)
 						{
diff --git a/Engine/Plugins/Animation/PoseSearch/Source/Runtime/Public/PoseSearch/PoseSearchAssetIndexer.h b/Engine/Plugins/Animation/PoseSearch/Source/Runtime/Public/PoseSearch/PoseSearchAssetIndexer.h
--- a/Engine/Plugins/Animation/PoseSearch/Source/Runtime/Public/PoseSearch/PoseSearchAssetIndexer.h
+++ b/Engine/Plugins/Animation/PoseSearch/Source/Runtime/Public/PoseSearch/PoseSearchAssetIndexer.h
@@ -146,9 +146,9 @@ private:
 		bool bClamped = false;
 
 		// RootTransform and ComponentSpacePose are stored mirrored in case SearchIndexAsset.IsMirrored
-		TArray<FTransform> RootTransform;
-		TArray<FCachedCSPose> ComponentSpacePose;
-		TArray<FBlendedHeapCurve> Curves;
+		TArray<FTransform, TInlineAllocator<PreallocatedRolesNum>> RootTransform;
+		TArray<FCachedCSPose, TInlineAllocator<PreallocatedRolesNum>> ComponentSpacePose;
+		TArray<FBlendedHeapCurve, TInlineAllocator<PreallocatedRolesNum>> Curves;
 	};
 
 	void GetSampleInfo(float SampleTime, int32 RoleIndex, FTransform& OutRootTransform, float& OutClipTime, bool& bOutClamped) const;
diff --git a/Engine/Plugins/Animation/PoseSearch/Source/Runtime/Public/PoseSearch/PoseSearchDefines.h b/Engine/Plugins/Animation/PoseSearch/Source/Runtime/Public/PoseSearch/PoseSearchDefines.h
--- a/Engine/Plugins/Animation/PoseSearch/Source/Runtime/Public/PoseSearch/PoseSearchDefines.h
+++ b/Engine/Plugins/Animation/PoseSearch/Source/Runtime/Public/PoseSearch/PoseSearchDefines.h
@@ -42,7 +42,7 @@ static constexpr int32 PreallocatedSearchesNum = 4;
 static constexpr float FiniteDelta = 1 / 60.0f; // Time delta used for computing pose derivatives
 static constexpr int32 MaxNumberOfCollectedPoseCandidatesPerDatabase = 256;
 static constexpr int32 InvalidSearchId = INDEX_NONE; // @todo: it should be 0, since the search id is an hash!
-static constexpr int32 DatabaseIndexDerivedDataCacheKeyVersion = 2677; // bump this value to force reindex ALL the UPoseSearchDatabases, without having to change the global FDevSystemGuids::POSESEARCHDB_DERIVEDDATA_VER
+static constexpr int32 DatabaseIndexDerivedDataCacheKeyVersion = 2679; // bump this value to force reindex ALL the UPoseSearchDatabases, without having to change the global FDevSystemGuids::POSESEARCHDB_DERIVEDDATA_VER
 static constexpr int32 DefaultTickPriority = 100;
 
 template<uint32 Alignment = DEFAULT_ALIGNMENT> class TMemStackSetAllocator : public TSetAllocator<TSparseArrayAllocator<TMemStackAllocator<Alignment>, TMemStackAllocator<Alignment>>, TMemStackAllocator<Alignment>> {};
