PoseSearch
- fixed interaction tick dependency cycle when any non main actors had already a dependency to the main actor
- added debug code to diagnose interaction island tick functions problems, enabled via a.PoseSearchInteraction.DiagnoseTickDependencies

#rnx

diff --git a/Engine/Plugins/Animation/PoseSearch/Source/Runtime/Private/PoseSearchInteractionIsland.cpp b/Engine/Plugins/Animation/PoseSearch/Source/Runtime/Private/PoseSearchInteractionIsland.cpp
--- a/Engine/Plugins/Animation/PoseSearch/Source/Runtime/Private/PoseSearchInteractionIsland.cpp
+++ b/Engine/Plugins/Animation/PoseSearch/Source/Runtime/Private/PoseSearchInteractionIsland.cpp
@@ -21,6 +21,103 @@
 namespace UE::PoseSearch
 {
 
+#if ENABLE_ANIM_DEBUG
+static bool GVarPoseSearchInteractionDiagnoseTickDependencies = false;
+static FAutoConsoleVariableRef CVarPoseSearchInteractionDiagnoseTickDependencies(TEXT("a.PoseSearchInteraction.DiagnoseTickDependencies"), GVarPoseSearchInteractionDiagnoseTickDependencies, TEXT("Enable Pose Search Interaction Tick Dependencies Diagnostic (SLOW!)"));
+
+// recursion safe FTickFunction logging functions
+static void ShowPrerequistes(const FTickFunction& NestedTick, int32 Indent, int32 MaxIndent)
+{
+	if (Indent >= MaxIndent)
+	{
+		UE_LOG(LogPoseSearch, Log, TEXT("%s==== REACHED MAX INDENT ===="), FCString::Spc(Indent * 2));
+	}
+	else
+	{
+		for (const FTickPrerequisite& Prereq : NestedTick.GetPrerequisites())
+		{
+			if (Prereq.PrerequisiteTickFunction)
+			{
+				UE_LOG(LogPoseSearch, Log, TEXT("%s prereq %s"), FCString::Spc(Indent * 2), *Prereq.PrerequisiteTickFunction->DiagnosticMessage());
+				ShowPrerequistes(*Prereq.PrerequisiteTickFunction, Indent + 1, MaxIndent);
+			}
+		}
+	}
+}
+
+static void LogTickFunction(FTickFunction& Tick, ENamedThreads::Type CurrentThread, bool bLogPrerequisites, int32 Indent, int32 MaxIndent = 50)
+{
+	if (Indent >= MaxIndent)
+	{
+		UE_LOG(LogPoseSearch, Log, TEXT("%s==== REACHED MAX INDENT ===="), FCString::Spc(Indent * 2));
+	}
+	else
+	{
+		// scoping brackets to save some heap for the recursion
+		{
+			UE_LOG(LogPoseSearch, Log, TEXT("%stick %s [%1d, %1d] %6llu %2d %s"), FCString::Spc(Indent * 2), Tick.bHighPriority ? TEXT("*") : TEXT(" "), (int32)Tick.GetActualTickGroup(), (int32)Tick.GetActualEndTickGroup(), (uint64)GFrameCounter, (int32)CurrentThread, *Tick.DiagnosticMessage());
+			if (bLogPrerequisites)
+			{
+				ShowPrerequistes(Tick, Indent, MaxIndent);
+			}
+		}
+
+		// Handle nested ticks
+		Tick.ForEachNestedTick([CurrentThread, bLogPrerequisites, Indent, MaxIndent](FTickFunction& NestedTick)
+			{
+				LogTickFunction(NestedTick, CurrentThread, bLogPrerequisites, Indent + 1, MaxIndent);
+			});
+	}
+}
+
+// check if there's any cycle within the prerequisites of Tick
+static bool ValidateTickDependenciesCycles(FTickFunction& Tick, TSet<FTickFunction*>& VisitedTickFunctions)
+{
+	bool bValidatedCorrectly = true;
+
+	bool bAlreadyInSet = false;
+	VisitedTickFunctions.Add(&Tick, &bAlreadyInSet);
+
+	if (bAlreadyInSet)
+	{
+		UE_LOG(LogPoseSearch, Error, TEXT("ValidateTickDependencies: TickFunction '%s' form a cycle"), *Tick.DiagnosticMessage());
+		bValidatedCorrectly = false;
+	}
+	else
+	{
+		for (const FTickPrerequisite& Prereq : Tick.GetPrerequisites())
+		{
+			if (Prereq.PrerequisiteTickFunction)
+			{
+				if (!ValidateTickDependenciesCycles(*Prereq.PrerequisiteTickFunction, VisitedTickFunctions))
+				{
+					VisitedTickFunctions.Remove(&Tick);
+					bValidatedCorrectly = false;
+					break;
+				}
+			}
+		}
+
+		if (bValidatedCorrectly)
+		{
+			Tick.ForEachNestedTick([&VisitedTickFunctions, &bValidatedCorrectly](FTickFunction& NestedTick)
+				{
+					if (bValidatedCorrectly)
+					{
+						if (!ValidateTickDependenciesCycles(NestedTick, VisitedTickFunctions))
+						{
+							bValidatedCorrectly = false;
+						}
+					}
+				});
+		}
+	}
+	
+	VisitedTickFunctions.Remove(&Tick);
+	return bValidatedCorrectly;
+}
+#endif // ENABLE_ANIM_DEBUG
+
 static FInteractionSearchResult InitSearchResult(const FSearchResult& SearchResult, const FInteractionSearchContext& SearchContext, int32 SearchIndex)
 {
 	FInteractionSearchResult InteractionSearchResult;
@@ -160,15 +257,46 @@ static UActorComponent* FindComponentForTickDependencies(const UObject* AnimCont
 
 static void AddPrerequisite(FTickFunction& TickFunction, UObject* TargetObject, FTickFunction& TargetTickFunction)
 {
-	check(TargetObject);
-	TickFunction.AddPrerequisite(TargetObject, TargetTickFunction);
+	if (TargetObject)
+	{
+#if ENABLE_ANIM_DEBUG
+		if (GVarPoseSearchInteractionDiagnoseTickDependencies)
+		{
+			const ENamedThreads::Type Type = IsInGameThread() ? ENamedThreads::GameThread : ENamedThreads::AnyThread;
+			TSet<FTickFunction*> VisitedTickFunctions;
+			if (!ValidateTickDependenciesCycles(TickFunction, VisitedTickFunctions))
+			{
+				LogTickFunction(TickFunction, Type, true, 1);
+			}
+			check(VisitedTickFunctions.IsEmpty());
+			if (!ValidateTickDependenciesCycles(TargetTickFunction, VisitedTickFunctions))
+			{
+				LogTickFunction(TargetTickFunction, Type, true, 1);
+			}
+			check(VisitedTickFunctions.IsEmpty());
+		}
+#endif // ENABLE_ANIM_DEBUG
+
+		TickFunction.AddPrerequisite(TargetObject, TargetTickFunction);
 
 #if ENABLE_ANIM_DEBUG
-	if (!TickFunction.GetPrerequisites().Contains(FTickPrerequisite(TargetObject, TargetTickFunction)))
-	{
-		UE_LOG(LogPoseSearch, Error, TEXT("UE::PoseSearch::AddPrerequisite, Failed to add prerequisite from [%s] to [%s, %s]!"), *TickFunction.DiagnosticMessage(), *TargetObject->GetName(), *TargetTickFunction.DiagnosticMessage());
-	}
+		if (!TickFunction.GetPrerequisites().Contains(FTickPrerequisite(TargetObject, TargetTickFunction)))
+		{
+			UE_LOG(LogPoseSearch, Error, TEXT("UE::PoseSearch::AddPrerequisite, Failed to add prerequisite from [%s] to [%s, %s]!"), *TickFunction.DiagnosticMessage(), *TargetObject->GetName(), *TargetTickFunction.DiagnosticMessage());
+		}
+
+		if (GVarPoseSearchInteractionDiagnoseTickDependencies)
+		{
+			const ENamedThreads::Type Type = IsInGameThread() ? ENamedThreads::GameThread : ENamedThreads::AnyThread;
+			TSet<FTickFunction*> VisitedTickFunctions;
+			if (!ValidateTickDependenciesCycles(TickFunction, VisitedTickFunctions))
+			{
+				LogTickFunction(TickFunction, Type, true, 1);
+			}
+			check(VisitedTickFunctions.IsEmpty());
+		}
 #endif // ENABLE_ANIM_DEBUG
+	}
 }
 
 // FInteractionSearchContextBase
@@ -486,41 +614,64 @@ void FInteractionIsland::AddTickDependencies(UActorComponent* TickActorComponent
 	{
 		if (const FTickFunction* TickActorComponentTickFunction = IslandDependency->FindTickFunction(TickActorComponent))
 		{
-			for (const FTickPrerequisite& TickActorComponentPrerequisite : TickActorComponentTickFunction->GetPrerequisites())
+			if(bInIsMainActor)
 			{
-				AddPrerequisite(PreTickFunction, TickActorComponentPrerequisite.PrerequisiteObject.Get(), *TickActorComponentPrerequisite.PrerequisiteTickFunction);
+				// PostTickFunction prerequisites should be empty since we haven't add the main actor tick function yet
+				check(PostTickFunction.GetPrerequisites().IsEmpty());
+
+				// adding to PreTickFunction all the tick dependencies TickActorComponent has, so it runs after all the tick dependencies of ALL the TickActorComponents in this FInteractionIsland
+				for (const FTickPrerequisite& TickActorComponentPrerequisite : TickActorComponentTickFunction->GetPrerequisites())
+				{
+					AddPrerequisite(PreTickFunction, TickActorComponentPrerequisite.PrerequisiteObject.Get(), *TickActorComponentPrerequisite.PrerequisiteTickFunction);
+				}
+
+				// Add post-tick function dependencies
+				IslandDependency->AddSubsequent(InteractionSubsystem, PreTickFunction, TickActorComponent);
+				IslandDependency->AddPrerequisite(InteractionSubsystem, PostTickFunction, TickActorComponent);
+
+				check(!bHasTickDependencies);
+				bHasTickDependencies = true;
 			}
-		}
-		else
-		{
-			UE_LOG(LogPoseSearch, Error, TEXT("FInteractionIsland::AddTickDependencies, error while retrieving the tick function for %s"), *TickActorComponent->GetName());
-		}
+			else
+			{
+				// PostTickFunction should contain only the tick function to the main actor's one
+				check(PostTickFunction.GetPrerequisites().Num());
 
-		if(bInIsMainActor)
-		{
-			IslandDependency->AddSubsequent(InteractionSubsystem, PreTickFunction, TickActorComponent);
-			IslandDependency->AddPrerequisite(InteractionSubsystem, PostTickFunction, TickActorComponent);
+				// adding to PreTickFunction all the tick dependencies TickActorComponent has, so it runs after all the tick dependencies of ALL the TickActorComponents in this FInteractionIsland
+				// BUT excluding the main actor tick fuction that is PostTickFunction.GetPrerequisites()[0]
+				const FTickFunction* MainActorTickFunction = PostTickFunction.GetPrerequisites()[0].Get();
+				for (const FTickPrerequisite& TickActorComponentPrerequisite : TickActorComponentTickFunction->GetPrerequisites())
+				{
+					if (TickActorComponentPrerequisite.PrerequisiteTickFunction != MainActorTickFunction)
+					{
+						AddPrerequisite(PreTickFunction, TickActorComponentPrerequisite.PrerequisiteObject.Get(), *TickActorComponentPrerequisite.PrerequisiteTickFunction);
+					}
+				}
 
-			check(!bHasTickDependencies);
-			bHasTickDependencies = true;
+				// Add post-tick function dependencies
+				IslandDependency->AddSubsequent(InteractionSubsystem, PostTickFunction, TickActorComponent);
+				check(bHasTickDependencies);
+			}
 		}
 		else
 		{
-			IslandDependency->AddSubsequent(InteractionSubsystem, PostTickFunction, TickActorComponent);
-			check(bHasTickDependencies);
+			UE_LOG(LogPoseSearch, Error, TEXT("FInteractionIsland::AddTickDependencies, error while retrieving the tick function for %s"), *TickActorComponent->GetName());
 		}
 	}
 	else
 	{
-		// adding to PreTickFunction all the tick dependencies TickActorComponent has, so it runs after all the tick dependencies of ALL the TickActorComponents in this FInteractionIsland
-		for (FTickPrerequisite& TickActorComponentPrerequisite : TickActorComponent->PrimaryComponentTick.GetPrerequisites())
-		{
-			AddPrerequisite(PreTickFunction, TickActorComponentPrerequisite.PrerequisiteObject.Get(), *TickActorComponentPrerequisite.PrerequisiteTickFunction);
-		}
-
-		// Add post-tick function dependencies
 		if (bInIsMainActor)
 		{
+			// PostTickFunction prerequisites should be empty since we haven't add the main actor tick function yet
+			check(PostTickFunction.GetPrerequisites().IsEmpty());
+
+			// adding to PreTickFunction all the tick dependencies TickActorComponent has, so it runs after all the tick dependencies of ALL the TickActorComponents in this FInteractionIsland
+			for (FTickPrerequisite& TickActorComponentPrerequisite : TickActorComponent->PrimaryComponentTick.GetPrerequisites())
+			{
+				AddPrerequisite(PreTickFunction, TickActorComponentPrerequisite.PrerequisiteObject.Get(), *TickActorComponentPrerequisite.PrerequisiteTickFunction);
+			}
+
+			// Add post-tick function dependencies
 			// it should be this island, but it's not a UObject, so we use the InteractionSubsystem
 			AddPrerequisite(TickActorComponent->PrimaryComponentTick, InteractionSubsystem, PreTickFunction);
 			AddPrerequisite(PostTickFunction, TickActorComponent, TickActorComponent->PrimaryComponentTick);
@@ -530,11 +681,43 @@ void FInteractionIsland::AddTickDependencies(UActorComponent* TickActorComponent
 		}
 		else
 		{
+			// PostTickFunction should contain only the tick function to the main actor's one
+			check(PostTickFunction.GetPrerequisites().Num());
+
+			// adding to PreTickFunction all the tick dependencies TickActorComponent has, so it runs after all the tick dependencies of ALL the TickActorComponents in this FInteractionIsland
+			// BUT excluding the main actor tick fuction that is PostTickFunction.GetPrerequisites()[0]
+			const FTickFunction* MainActorTickFunction = PostTickFunction.GetPrerequisites()[0].Get();
+			for (FTickPrerequisite& TickActorComponentPrerequisite : TickActorComponent->PrimaryComponentTick.GetPrerequisites())
+			{
+				if (TickActorComponentPrerequisite.PrerequisiteTickFunction != MainActorTickFunction)
+				{
+					AddPrerequisite(PreTickFunction, TickActorComponentPrerequisite.PrerequisiteObject.Get(), *TickActorComponentPrerequisite.PrerequisiteTickFunction);
+				}
+			}
+
+			// Add post-tick function dependencies
 			// it should be this island, but it's not a UObject, so we use the InteractionSubsystem
 			AddPrerequisite(TickActorComponent->PrimaryComponentTick, InteractionSubsystem, PostTickFunction);
 			check(bHasTickDependencies);
 		}
 	}
+
+#if ENABLE_ANIM_DEBUG
+	if (GVarPoseSearchInteractionDiagnoseTickDependencies)
+	{
+		TSet<FTickFunction*> VisitedTickFunctions;
+		bool bIsPreTickFunctionValid = ValidateTickDependenciesCycles(PreTickFunction, VisitedTickFunctions);
+		check(VisitedTickFunctions.IsEmpty());
+		bool bIsPostTickFunctionValid = ValidateTickDependenciesCycles(PostTickFunction, VisitedTickFunctions);
+		check(VisitedTickFunctions.IsEmpty());
+		if (!bIsPreTickFunctionValid || !bIsPostTickFunctionValid)
+		{
+			// if this validation triggers here, FInteractionIsland is not respectint the already present dependencies, creating cycles
+			UE_LOG(LogPoseSearch, Error, TEXT("============== FInteractionIsland::AddTickDependencies ValidateTickDependencies failed! Analyze the log and tune the FPoseSearchInteractionAvailability::TickPriority =============="));
+			LogTickDependencies();
+		}
+	}
+#endif // ENABLE_ANIM_DEBUG
 }
 
 void FInteractionIsland::RemoveTickDependencies(bool bValidateTickDependencies)
@@ -695,6 +878,23 @@ void FInteractionIsland::AddSearchContext(const FInteractionSearchContext& Searc
 
 void FInteractionIsland::Uninitialize(bool bValidateTickDependencies)
 {
+#if ENABLE_ANIM_DEBUG
+	if (GVarPoseSearchInteractionDiagnoseTickDependencies)
+	{
+		TSet<FTickFunction*> VisitedTickFunctions;
+		bool bIsPreTickFunctionValid = ValidateTickDependenciesCycles(PreTickFunction, VisitedTickFunctions);
+		check(VisitedTickFunctions.IsEmpty());
+		bool bIsPostTickFunctionValid = ValidateTickDependenciesCycles(PostTickFunction, VisitedTickFunctions);
+		check(VisitedTickFunctions.IsEmpty());
+		if (!bIsPreTickFunctionValid || !bIsPostTickFunctionValid)
+		{
+			// if this validation triggers here, some additional tick dependency outside FInteractionIsland has been injected witout respecting the already present dependencies, creating cycles
+			UE_LOG(LogPoseSearch, Error, TEXT("============== FInteractionIsland::Uninitialize ValidateTickDependencies failed! =============="));
+			LogTickDependencies();
+		}
+	}
+#endif // ENABLE_ANIM_DEBUG
+
 	RemoveTickDependencies(bValidateTickDependencies);
 
 	if (IsInitialized())
@@ -737,52 +937,6 @@ const AActor* FInteractionIsland::GetMainActor() const
 }
 
 #if ENABLE_ANIM_DEBUG
-// recursion safe FTickFunction logging functions
-static void ShowPrerequistes(const FTickFunction& NestedTick, int32 Indent, int32 MaxIndent)
-{
-	if (Indent >= MaxIndent)
-	{
-		UE_LOG(LogPoseSearch, Log, TEXT("%s==== REACHED MAX INDENT ===="), FCString::Spc(Indent * 2));
-	}
-	else
-	{
-		for (const FTickPrerequisite& Prereq : NestedTick.GetPrerequisites())
-		{
-			if (Prereq.PrerequisiteTickFunction)
-			{
-				UE_LOG(LogPoseSearch, Log, TEXT("%s prereq %s"), FCString::Spc(Indent * 2), *Prereq.PrerequisiteTickFunction->DiagnosticMessage());
-				ShowPrerequistes(*Prereq.PrerequisiteTickFunction, Indent + 1, MaxIndent);
-			}
-		}
-	}
-}
-
-static void LogTickFunction(FTickFunction& Tick, ENamedThreads::Type CurrentThread, bool bLogPrerequisites, int32 Indent, int32 MaxIndent = 50)
-{
-	if (Indent >= MaxIndent)
-	{
-		UE_LOG(LogPoseSearch, Log, TEXT("%s==== REACHED MAX INDENT ===="), FCString::Spc(Indent * 2));
-	}
-	else
-	{
-		// scoping brackets to save some heap for the recursion
-		{
-			UE_LOG(LogPoseSearch, Log, TEXT("%stick %s [%1d, %1d] %6llu %2d %s"), FCString::Spc(Indent * 2), Tick.bHighPriority ? TEXT("*") : TEXT(" "), (int32)Tick.GetActualTickGroup(), (int32)Tick.GetActualEndTickGroup(), (uint64)GFrameCounter, (int32)CurrentThread, *Tick.DiagnosticMessage());
-			if (bLogPrerequisites)
-			{
-				ShowPrerequistes(Tick, Indent, MaxIndent);
-			}
-		}
-
-		// Handle nested ticks
-		Tick.ForEachNestedTick([CurrentThread, bLogPrerequisites, Indent, MaxIndent](FTickFunction& NestedTick)
-			{
-				LogTickFunction(NestedTick, CurrentThread, bLogPrerequisites, Indent + 1, MaxIndent);
-			});
-	}
-}
-
-
 void FInteractionIsland::LogTickDependencies() const
 {
 	const ENamedThreads::Type Type = IsInGameThread() ? ENamedThreads::GameThread : ENamedThreads::AnyThread;
diff --git a/Engine/Plugins/Animation/PoseSearch/Source/Runtime/Private/PoseSearchInteractionSubsystem.cpp b/Engine/Plugins/Animation/PoseSearch/Source/Runtime/Private/PoseSearchInteractionSubsystem.cpp
--- a/Engine/Plugins/Animation/PoseSearch/Source/Runtime/Private/PoseSearchInteractionSubsystem.cpp
+++ b/Engine/Plugins/Animation/PoseSearch/Source/Runtime/Private/PoseSearchInteractionSubsystem.cpp
@@ -2,7 +2,6 @@
 
 #include "PoseSearch/PoseSearchInteractionSubsystem.h"
 #include "PoseSearch/AnimNode_PoseSearchHistoryCollector.h"
-#include "PoseSearch/PoseSearchContext.h"
 #include "PoseSearch/PoseSearchDatabase.h"
 #include "PoseSearch/PoseSearchInteractionIsland.h"
 #include "PoseSearch/PoseSearchRole.h"
diff --git a/Engine/Plugins/Animation/PoseSearch/Source/Runtime/Public/PoseSearch/PoseSearchInteractionSubsystem.h b/Engine/Plugins/Animation/PoseSearch/Source/Runtime/Public/PoseSearch/PoseSearchInteractionSubsystem.h
--- a/Engine/Plugins/Animation/PoseSearch/Source/Runtime/Public/PoseSearch/PoseSearchInteractionSubsystem.h
+++ b/Engine/Plugins/Animation/PoseSearch/Source/Runtime/Public/PoseSearch/PoseSearchInteractionSubsystem.h
@@ -3,6 +3,7 @@
 #pragma once
 
 #include "Subsystems/WorldSubsystem.h"
+#include "PoseSearch/PoseSearchContext.h"
 #include "PoseSearch/PoseSearchInteractionLibrary.h"
 #include "PoseSearch/PoseSearchInteractionValidator.h"
 #include "PoseSearchInteractionSubsystem.generated.h"
@@ -55,6 +56,29 @@ struct FAvailabilitiesMap : protected TMap<TWeakObjectPtr<const UObject>, FInter
 	
 	void AddAvailabilities(const TArrayView<const FPoseSearchInteractionAvailability> Availabilities, const UObject* AnimContext, FName PoseHistoryName, const IPoseHistory* PoseHistory)
 	{
+#if WITH_EDITOR && DO_CHECK
+		check(AnimContext);
+		const AActor* ContextOwningActor = GetContextOwningActor(AnimContext);
+		check(ContextOwningActor);
+
+		for (const TPair<TWeakObjectPtr<const UObject>, FInteractionAvailabilities>& AvailabilitiesMapPair : *this)
+		{
+			if (const UObject* OtherAnimContext = AvailabilitiesMapPair.Key.Get())
+			{
+				if (const AActor* OtherContextOwningActor = GetContextOwningActor(OtherAnimContext))
+				{
+					// same actor but different AnimContext? Do we have multiple AddAvailabilities from different components?
+					// It's not currently supported! and if not filtered out it could potentially add interactions within the same actor...
+					if (ContextOwningActor == OtherContextOwningActor && AnimContext != OtherAnimContext)
+					{
+						UE_LOG(LogPoseSearch, Error, TEXT("FAvailabilitiesMap::AddAvailabilities from different AnimContexts %s, %s sharing the same Actor %s. This is not currently supported!"), *AnimContext->GetName(), *OtherAnimContext->GetName(), *ContextOwningActor->GetName());
+						return;
+					}
+				}
+			}
+		}
+#endif // WITH_EDITOR && DO_CHECK
+
 		FInteractionAvailabilities& AnimContextAvailabilities = FindOrAdd(AnimContext);
 		for (const FPoseSearchInteractionAvailability& Availability : Availabilities)
 		{
