Fixed root motion being zero in the case where no assets are resultant from motion matching

#rnx

diff --git a/Engine/Plugins/Animation/PoseSearch/Source/Runtime/Private/AnimNode_MotionMatchingInteraction.cpp b/Engine/Plugins/Animation/PoseSearch/Source/Runtime/Private/AnimNode_MotionMatchingInteraction.cpp
--- a/Engine/Plugins/Animation/PoseSearch/Source/Runtime/Private/AnimNode_MotionMatchingInteraction.cpp
+++ b/Engine/Plugins/Animation/PoseSearch/Source/Runtime/Private/AnimNode_MotionMatchingInteraction.cpp
@@ -100,8 +100,11 @@ void FAnimNode_MotionMatchingInteraction::UpdateAssetPlayer(const FAnimationUpda
 	const bool bDidBlendToRequestAnInertialBlend = bBlendToExecuted && bUseInertialBlend;
 	UE::Anim::TOptionalScopedGraphMessage<UE::Anim::FAnimInertializationSyncScope> InertializationSync(bDidBlendToRequestAnInertialBlend, Context);
 	
-	UpdatePlayRate(CurrentResult.WantedPlayRate);
-	UpdateBlendspaceParameters(BlendspaceUpdateMode, CurrentResult.BlendParameters);
+	if (CurrentResult.SelectedAnim.IsValid())
+	{
+		UpdatePlayRate(CurrentResult.WantedPlayRate);
+		UpdateBlendspaceParameters(BlendspaceUpdateMode, CurrentResult.BlendParameters);
+	}
 
 #if ENABLE_ANIM_DEBUG
 	if (UE_TRACE_CHANNELEXPR_IS_ENABLED(AnimationChannel))
