PoseSearch - fix for database viwq model crash on db assets with only one pose

diff --git a/Engine/Plugins/Animation/PoseSearch/Source/Editor/Private/PoseSearchDatabaseViewModel.cpp b/Engine/Plugins/Animation/PoseSearch/Source/Editor/Private/PoseSearchDatabaseViewModel.cpp
--- a/Engine/Plugins/Animation/PoseSearch/Source/Editor/Private/PoseSearchDatabaseViewModel.cpp
+++ b/Engine/Plugins/Animation/PoseSearch/Source/Editor/Private/PoseSearchDatabaseViewModel.cpp
@@ -137,26 +137,27 @@ bool FDatabasePreviewActor::SpawnPreviewActor(UWorld* World, const UPoseSearchDa
 	Trajectory.Samples.SetNumUninitialized(NumPoses);
 	TrajectorySpeed.SetNumUninitialized(NumPoses);
 
-	if (NumPoses > 1)
+	for (int32 Index = 0; Index < NumPoses; ++Index)
 	{
-		for (int32 Index = 0; Index < NumPoses; ++Index)
-		{
-			const int32 IndexAssetPoseIdx = Index + IndexAsset.GetFirstPoseIdx();
-			const float IndexAssetPoseTime = IndexAsset.GetTimeFromPoseIndex(IndexAssetPoseIdx, PoseSearchDatabase->Schema->SampleRate);
-			const FTransform IndexAssetPoseTransform = MirrorDataCache.MirrorTransform(Sampler.ExtractRootTransform(IndexAssetPoseTime));
+		const int32 IndexAssetPoseIdx = Index + IndexAsset.GetFirstPoseIdx();
+		const float IndexAssetPoseTime = IndexAsset.GetTimeFromPoseIndex(IndexAssetPoseIdx, PoseSearchDatabase->Schema->SampleRate);
+		const FTransform IndexAssetPoseTransform = MirrorDataCache.MirrorTransform(Sampler.ExtractRootTransform(IndexAssetPoseTime));
 
-			Trajectory.Samples[Index].SetTransform(IndexAssetPoseTransform);
-			Trajectory.Samples[Index].TimeInSeconds = IndexAssetPoseTime;
-		}
+		Trajectory.Samples[Index].SetTransform(IndexAssetPoseTransform);
+		Trajectory.Samples[Index].TimeInSeconds = IndexAssetPoseTime;
+	}
 
-		for (int32 Index = 1; Index < NumPoses; ++Index)
-		{
-			const float DeltaAccumulatedSeconds = FMath::Max(UE_KINDA_SMALL_NUMBER, Trajectory.Samples[Index].TimeInSeconds - Trajectory.Samples[Index - 1].TimeInSeconds);
-			const FVector& Start = Trajectory.Samples[Index - 1].Position;
-			const FVector& End = Trajectory.Samples[Index].Position;
-			TrajectorySpeed[Index] = (Start - End).Length() / DeltaAccumulatedSeconds;
-		}
-		TrajectorySpeed[0] = TrajectorySpeed[1];
+	for (int32 Index = 1; Index < NumPoses; ++Index)
+	{
+		const float DeltaAccumulatedSeconds = FMath::Max(UE_KINDA_SMALL_NUMBER, Trajectory.Samples[Index].TimeInSeconds - Trajectory.Samples[Index - 1].TimeInSeconds);
+		const FVector& Start = Trajectory.Samples[Index - 1].Position;
+		const FVector& End = Trajectory.Samples[Index].Position;
+		TrajectorySpeed[Index] = (Start - End).Length() / DeltaAccumulatedSeconds;
+	}
+	
+	if (NumPoses > 0)
+	{
+		TrajectorySpeed[0] = NumPoses > 1 ? TrajectorySpeed[1] : 0.f;
 	}
 
 	UE_LOG(LogPoseSearchEditor, Log, TEXT("Spawned preview Actor: %s"), *GetNameSafe(ActorPtr.Get()));
