PoseSearch
- FAnimNode_MotionMatching::MotionMatchingState conditional reset fixes (not on FMotionMatchingTrait yet, since it implements MM interaction as well)
- removed MM node conditional search under bIsRelevantEnough condition, introduced for MM interactions, and now solved via blends updating the children independently of their relevance

diff --git a/Engine/Plugins/Animation/PoseSearch/Source/Runtime/Private/AnimNode_MotionMatching.cpp b/Engine/Plugins/Animation/PoseSearch/Source/Runtime/Private/AnimNode_MotionMatching.cpp
--- a/Engine/Plugins/Animation/PoseSearch/Source/Runtime/Private/AnimNode_MotionMatching.cpp
+++ b/Engine/Plugins/Animation/PoseSearch/Source/Runtime/Private/AnimNode_MotionMatching.cpp
@@ -104,50 +104,62 @@ void FAnimNode_MotionMatching::UpdateAssetPlayer(const FAnimationUpdateContext&
 
 	GetEvaluateGraphExposedInputs().Execute(Context);
 
-	bool bNeedsReset =
-		bResetOnBecomingRelevant &&
-		UpdateCounter.HasEverBeenUpdated() &&
-		!UpdateCounter.WasSynchronizedCounter(Context.AnimInstanceProxy->GetUpdateCounter());
-
+	// synchronizing with GetAccumulatedTime or resetting MotionMatchingState, and conditionally resetting FAnimNode_BlendStack_Standalone
+	if (bResetOnBecomingRelevant && UpdateCounter.HasEverBeenUpdated() && !UpdateCounter.WasSynchronizedCounter(Context.AnimInstanceProxy->GetUpdateCounter()))
+	{
+		// If we just became relevant and haven't been initialized yet, then reset motion matching state, otherwise update the asset time using the player node.
+		MotionMatchingState.Reset();
+		FAnimNode_BlendStack_Standalone::Reset();
+	}
+	else if (MotionMatchingState.CurrentSearchResult.Database == nullptr || MotionMatchingState.CurrentSearchResult.Database->Schema == nullptr)
+	{
+	}
 #if WITH_EDITOR
-	if (EAsyncBuildIndexResult::Success != FAsyncPoseSearchDatabasesManagement::RequestAsyncBuildIndex(MotionMatchingState.CurrentSearchResult.Database.Get(), ERequestAsyncBuildFlag::ContinueRequest))
+	else if (EAsyncBuildIndexResult::Success != FAsyncPoseSearchDatabasesManagement::RequestAsyncBuildIndex(MotionMatchingState.CurrentSearchResult.Database.Get(), ERequestAsyncBuildFlag::ContinueRequest))
 	{
-		bNeedsReset = true;
+		// MotionMatchingState.CurrentSearchResult.Database is indexing, and it's not safe to use its previous index properties cached in MotionMatchingState
+		MotionMatchingState.Reset();
 	}
-	// in case this node is not updated, and MotionMatchingState.CurrentSearchResult.Database gets modified, we could end up with CurrentSearchResult being out of synch with the updated database, so we need to reset the state
 	else if (MotionMatchingState.CurrentSearchResult.IsValid() && MotionMatchingState.CurrentSearchResult.PoseIdx >= MotionMatchingState.CurrentSearchResult.Database->GetSearchIndex().GetNumPoses())
 	{
-		bNeedsReset = true;
+		// in case this node is not updated, and MotionMatchingState.CurrentSearchResult.Database gets modified, we could end up with
+		// CurrentSearchResult being out of synch with the updated database, so we need to reset the state
+		MotionMatchingState.Reset();
 	}
 #endif // WITH_EDITOR
-
-	// If we just became relevant and haven't been initialized yet, then reset motion matching state, otherwise update the asset time using the player node.
-	if (bNeedsReset)
+	else if (GetAnimAsset() != MotionMatchingState.CurrentSearchResult.GetCurrentResultAnimationAsset())
 	{
+		// MotionMatchingState and blendstack are not in synch! MotionMatchingState thinks it's playing MotionMatchingState.CurrentSearchResult.GetCurrentResultAnimationAsset()
+		// and blendstack thinks it's playing GetAnimAsset(). The two could have gone out of synch becasue or external blendstack requests, and so we reset the MotionMatchingState
 		MotionMatchingState.Reset();
-		FAnimNode_BlendStack_Standalone::Reset();
 	}
 	else
 	{
 		// We adjust the motion matching state asset time to the current player node's asset time. This is done 
 		// because the player node may have ticked more or less time than we expected due to variable dt or the 
 		// dynamic playback rate adjustment and as such the motion matching state does not update by itself
+		
+		// @todo: this is missing eventual update of the blend parameters that could change the pose!
 		MotionMatchingState.AdjustAssetTime(GetAccumulatedTime());
 	}
+
 	UpdateCounter.SynchronizeWith(Context.AnimInstanceProxy->GetUpdateCounter());
 
 	// If the Database property hasn't been overridden, set it as the only database to search.
-	if (!bOverrideDatabaseInput && Database)
+	if (!bOverrideDatabaseInput)
 	{
-		DatabasesToSearch.Reset(1);
-		DatabasesToSearch.Add(Database);
+		DatabasesToSearch.Reset();
+		if (Database != nullptr)
+		{
+			DatabasesToSearch.Add(Database);
+		}
 	}
 
 #if ENABLE_ANIM_DEBUG
 	if (GVarAnimNodeMotionMatchingDrawInfo)
 	{
 		const UPoseSearchDatabase* CurrentDatabase = MotionMatchingState.CurrentSearchResult.Database.Get();
-		const UAnimationAsset* CurrentAnimationAsset = AnimPlayers.IsEmpty() ? nullptr : AnimPlayers[0].GetAnimationAsset();
+		const UAnimationAsset* CurrentAnimationAsset = GetAnimAsset();
 
 		FString DebugInfo = FString::Printf(TEXT("NextUpdateInterruptMode(%s)\n"), *UEnum::GetValueAsString(NextUpdateInterruptMode));
 		DebugInfo += FString::Printf(TEXT("Current Database(%s)\n"), *GetNameSafe(CurrentDatabase));
@@ -169,57 +181,41 @@ void FAnimNode_MotionMatching::UpdateAssetPlayer(const FAnimationUpdateContext&
 	}
 #endif // ENABLE_ANIM_DEBUG
 
-	const bool bIsRelevantEnough = Context.GetFinalBlendWeight() > FAnimWeight::GetSmallestRelevantWeight();
-	if (DatabasesToSearch.IsEmpty())
+	const IPoseHistory* PoseHistory = nullptr;
+	if (const FPoseHistoryProvider* PoseHistoryProvider = Context.GetMessage<FPoseHistoryProvider>())
 	{
-		Context.LogMessage(
-			EMessageSeverity::Error,
-			LOCTEXT("NoDatabases", "No database assets provided for motion matching."));
+		PoseHistory = &PoseHistoryProvider->GetPoseHistory();
 	}
-	else if (bIsRelevantEnough)
-	{
-		const float DeltaTime = Context.GetDeltaTime();
-
-		const IPoseHistory* PoseHistory = nullptr;
-		if (const FPoseHistoryProvider* PoseHistoryProvider = Context.GetMessage<FPoseHistoryProvider>())
-		{
-			PoseHistory = &PoseHistoryProvider->GetPoseHistory();
-		}
 
-		FFloatInterval PoseSearchPlayRate = PlayRate;
+	FFloatInterval PoseSearchPlayRate = PlayRate;
 #if ENABLE_ANIM_DEBUG
-		bool bPoseSearchPlayRateEnabled = GVarAnimNodeMotionMatchingPlayRateEnabled == UE::Private::EPlayRateState::Enabled || GVarAnimNodeMotionMatchingPlayRateEnabled == UE::Private::EPlayRateState::PoseSearchOnly;
-		PoseSearchPlayRate = bPoseSearchPlayRateEnabled ? PoseSearchPlayRate : FFloatInterval(1.f, 1.f);
+	bool bPoseSearchPlayRateEnabled = GVarAnimNodeMotionMatchingPlayRateEnabled == UE::Private::EPlayRateState::Enabled || GVarAnimNodeMotionMatchingPlayRateEnabled == UE::Private::EPlayRateState::PoseSearchOnly;
+	PoseSearchPlayRate = bPoseSearchPlayRateEnabled ? PoseSearchPlayRate : FFloatInterval(1.f, 1.f);
 #endif // ENABLE_ANIM_DEBUG
 
-		UPoseSearchLibrary::UpdateMotionMatchingState(
-			Context.AnimInstanceProxy->GetAnimInstanceObject(),
-			PoseHistory,
-			DatabasesToSearch,
-			Context.GetDeltaTime(),
-			PoseJumpThresholdTime,
-			PoseReselectHistory,
-			bShouldSearch ? SearchThrottleTime : UE_BIG_NUMBER,
-			PoseSearchPlayRate,
-			MotionMatchingState,
-			NextUpdateInterruptMode,
-			bShouldUseCachedChannelData
+	UPoseSearchLibrary::UpdateMotionMatchingState(
+		Context.AnimInstanceProxy->GetAnimInstanceObject(),
+		PoseHistory,
+		DatabasesToSearch,
+		Context.GetDeltaTime(),
+		PoseJumpThresholdTime,
+		PoseReselectHistory,
+		bShouldSearch ? SearchThrottleTime : UE_BIG_NUMBER,
+		PoseSearchPlayRate,
+		MotionMatchingState,
+		NextUpdateInterruptMode,
+		bShouldUseCachedChannelData
 #if ENABLE_ANIM_DEBUG
-			, GVarAnimNodeMotionMatchingDrawQuery
-			, GVarAnimNodeMotionMatchingDrawCurResult
+		, GVarAnimNodeMotionMatchingDrawQuery
+		, GVarAnimNodeMotionMatchingDrawCurResult
 #else // ENABLE_ANIM_DEBUG
-			, false
-			, false
+		, false
+		, false
 #endif // ENABLE_ANIM_DEBUG
-			, EventToSearch
-		);
-	}
-	else
-	{
-		MotionMatchingState.Reset();
-	}
+		, EventToSearch
+	);
 
-	const bool bJumpToPose = bIsRelevantEnough && MotionMatchingState.CurrentSearchResult.IsValid() && !MotionMatchingState.CurrentSearchResult.bIsContinuingPoseSearch;
+	const bool bJumpToPose = MotionMatchingState.CurrentSearchResult.IsValid() && !MotionMatchingState.CurrentSearchResult.bIsContinuingPoseSearch;
 
 	PRAGMA_DISABLE_DEPRECATION_WARNINGS
 	MotionMatchingState.bJumpedToPose = bJumpToPose;
diff --git a/Engine/Plugins/Animation/PoseSearch/Source/Runtime/Private/MotionMatchingAnimNodeLibrary.cpp b/Engine/Plugins/Animation/PoseSearch/Source/Runtime/Private/MotionMatchingAnimNodeLibrary.cpp
--- a/Engine/Plugins/Animation/PoseSearch/Source/Runtime/Private/MotionMatchingAnimNodeLibrary.cpp
+++ b/Engine/Plugins/Animation/PoseSearch/Source/Runtime/Private/MotionMatchingAnimNodeLibrary.cpp
@@ -33,16 +33,17 @@ void UMotionMatchingAnimNodeLibrary::GetMotionMatchingSearchResult(const FMotion
 		const UPoseSearchDatabase* CurrentResultDatabase = MotionMatchingState.CurrentSearchResult.Database.Get();
 		if (CurrentResultDatabase == nullptr || CurrentResultDatabase->Schema == nullptr)
 		{
-			// @todo: should we log a warning?
 		}
 #if WITH_EDITOR
 		else if (EAsyncBuildIndexResult::Success != FAsyncPoseSearchDatabasesManagement::RequestAsyncBuildIndex(CurrentResultDatabase, ERequestAsyncBuildFlag::ContinueRequest))
 		{
+			// MotionMatchingState.CurrentSearchResult.Database is indexing, and it's not safe to use its previous index properties cached in MotionMatchingState
 			UE_LOG(LogPoseSearch, Warning, TEXT("UMotionMatchingAnimNodeLibrary::GetMotionMatchingSearchResult called while '%s' is indexing. returning an invalid result.."), *CurrentResultDatabase->GetName());
 		}
-		// in case this node is not updated, and MotionMatchingState.CurrentSearchResult.Database gets modified, we could end up with CurrentSearchResult being out of synch with the updated database, so we need to reset the state
 		else if (MotionMatchingState.CurrentSearchResult.IsValid() && MotionMatchingState.CurrentSearchResult.PoseIdx >= MotionMatchingState.CurrentSearchResult.Database->GetSearchIndex().GetNumPoses())
 		{
+			// in case this node is not updated, and MotionMatchingState.CurrentSearchResult.Database gets modified, we could end up with 
+			// CurrentSearchResult being out of synch with the updated database, so we need to reset the state
 			UE_LOG(LogPoseSearch, Warning, TEXT("UMotionMatchingAnimNodeLibrary::GetMotionMatchingSearchResult inconsistency between MotionMatchingState and '%s', probably because of stale MotionMatchingState after database got reindexed. returning an invalid result.."), *CurrentResultDatabase->GetName());
 		}
 #endif // WITH_EDITOR
diff --git a/Engine/Plugins/Animation/PoseSearch/Source/Runtime/Private/PoseSearchResult.cpp b/Engine/Plugins/Animation/PoseSearch/Source/Runtime/Private/PoseSearchResult.cpp
--- a/Engine/Plugins/Animation/PoseSearch/Source/Runtime/Private/PoseSearchResult.cpp
+++ b/Engine/Plugins/Animation/PoseSearch/Source/Runtime/Private/PoseSearchResult.cpp
@@ -81,6 +81,24 @@ const FSearchIndexAsset* FSearchResult::GetSearchIndexAsset(bool bMandatory) con
 	return &Database->GetSearchIndex().GetAssetForPose(PoseIdx);
 }
 
+const UAnimationAsset* FSearchResult::GetCurrentResultAnimationAsset() const
+{
+	if (const FSearchIndexAsset* SearchIndexAsset = GetSearchIndexAsset())
+	{
+		return Database->GetDatabaseAnimationAsset<FPoseSearchDatabaseAnimationAssetBase>(*SearchIndexAsset)->GetAnimationAssetForRole(Database->Schema->GetDefaultRole());
+	}
+	return nullptr;
+}
+
+const UAnimationAsset* FSearchResult::GetCurrentResultAnimationAsset(const FRole& Role) const
+{
+	if (const FSearchIndexAsset* SearchIndexAsset = GetSearchIndexAsset())
+	{
+		return Database->GetDatabaseAnimationAsset<FPoseSearchDatabaseAnimationAssetBase>(*SearchIndexAsset)->GetAnimationAssetForRole(Role);
+	}
+	return nullptr;
+}
+
 bool FSearchResult::CanAdvance(float DeltaTime) const
 {
 	const FSearchIndexAsset* SearchIndexAsset = GetSearchIndexAsset();
diff --git a/Engine/Plugins/Animation/PoseSearch/Source/Runtime/Public/PoseSearch/AnimNode_MotionMatching.h b/Engine/Plugins/Animation/PoseSearch/Source/Runtime/Public/PoseSearch/AnimNode_MotionMatching.h
--- a/Engine/Plugins/Animation/PoseSearch/Source/Runtime/Public/PoseSearch/AnimNode_MotionMatching.h
+++ b/Engine/Plugins/Animation/PoseSearch/Source/Runtime/Public/PoseSearch/AnimNode_MotionMatching.h
@@ -167,7 +167,7 @@ private:
 	FGraphTraversalCounter UpdateCounter;
 
 	// List of databases this node is searching.
-	UPROPERTY()
+	UPROPERTY(Transient)
 	TArray<TObjectPtr<const UPoseSearchDatabase>> DatabasesToSearch;
 
 	// Applied EPoseSearchInterruptMode on the next update that controls the continuing pose search eveluation. This is set back to EPoseSearchInterruptMode::DoNotInterrupt after each update.
diff --git a/Engine/Plugins/Animation/PoseSearch/Source/Runtime/Public/PoseSearch/PoseSearchAssetSampler.h b/Engine/Plugins/Animation/PoseSearch/Source/Runtime/Public/PoseSearch/PoseSearchAssetSampler.h
--- a/Engine/Plugins/Animation/PoseSearch/Source/Runtime/Public/PoseSearch/PoseSearchAssetSampler.h
+++ b/Engine/Plugins/Animation/PoseSearch/Source/Runtime/Public/PoseSearch/PoseSearchAssetSampler.h
@@ -68,7 +68,7 @@ protected:
 	// members used to sample blend spaces only!
 	FVector BlendParameters = FVector::ZeroVector;
 	int32 RootTransformSamplingRate = DefaultRootTransformSamplingRate;
-	float CachedPlayLength = 0.f;
+	float CachedPlayLength = -1.f;
 	TArray<FTransform> AccumulatedRootTransform;
 
 	const float ExtrapolationSampleTime = 1.f / 30.f;
diff --git a/Engine/Plugins/Animation/PoseSearch/Source/Runtime/Public/PoseSearch/PoseSearchResult.h b/Engine/Plugins/Animation/PoseSearch/Source/Runtime/Public/PoseSearch/PoseSearchResult.h
--- a/Engine/Plugins/Animation/PoseSearch/Source/Runtime/Public/PoseSearch/PoseSearchResult.h
+++ b/Engine/Plugins/Animation/PoseSearch/Source/Runtime/Public/PoseSearch/PoseSearchResult.h
@@ -3,7 +3,7 @@
 #pragma once
 
 #include "PoseSearch/PoseSearchCost.h"
-#include "PoseSearch/PoseSearchDefines.h"
+#include "PoseSearch/PoseSearchRole.h"
 #include "PoseSearchResult.generated.h"
 
 class UPoseSearchDatabase;
@@ -70,14 +70,24 @@ struct FSearchResult
 
 	void Reset()
 	{
+		PoseCost = FPoseSearchCost();
 		PoseIdx = INDEX_NONE;
 		EventPoseIdx = INDEX_NONE;
 		Database = nullptr;
 		AssetTime = 0.0f;
 		bIsContinuingPoseSearch = false;
+
+#if WITH_EDITOR && ENABLE_ANIM_DEBUG && UE_POSE_SEARCH_TRACE_ENABLED
+		BruteForcePoseCost = FPoseSearchCost();
+		BestPosePos = 0;
+#endif // WITH_EDITOR && ENABLE_ANIM_DEBUG && UE_POSE_SEARCH_TRACE_ENABLED
 	}
 
 	POSESEARCH_API const FSearchIndexAsset* GetSearchIndexAsset(bool bMandatory = false) const;
+	// Experimental, this feature might be removed without warning, not for production use
+	POSESEARCH_API const UAnimationAsset* GetCurrentResultAnimationAsset() const;
+	// Experimental, this feature might be removed without warning, not for production use
+	POSESEARCH_API const UAnimationAsset* GetCurrentResultAnimationAsset(const FRole& Role) const;
 	
 	bool CanAdvance(float DeltaTime) const;
 
