PoseSearch - MM logging improvements

diff --git a/Engine/Plugins/Animation/PoseSearch/Source/Runtime/Private/PoseSearchContext.cpp b/Engine/Plugins/Animation/PoseSearch/Source/Runtime/Private/PoseSearchContext.cpp
--- a/Engine/Plugins/Animation/PoseSearch/Source/Runtime/Private/PoseSearchContext.cpp
+++ b/Engine/Plugins/Animation/PoseSearch/Source/Runtime/Private/PoseSearchContext.cpp
@@ -864,8 +864,7 @@ FTransform FSearchContext::GetWorldRootBoneTransformAtTime(float SampleTime, con
 		// So we use PoseHistoryTransformAtTime as such below (in the form of "const FTransform& ComponentToWorldTransform = PoseHistory ? PoseHistoryTransformAtTime : ...)
 	}
 
-	// accessing AnimInstance or SkeletalMesh properties is not thread unsafe from a worker thread! Logging an error and returning identity transform
-	UE_LOG(LogPoseSearch, Error, TEXT("FSearchContext::GetWorldRootBoneTransformAtTime - Couldn't find root bone transform at time %f for role %s for schema %s!"), SampleTime, *SampleRole.ToString(), *Schema->GetName());
+	// accessing AnimInstance or SkeletalMesh properties is not thread safe from a worker thread! Logging an error and returning identity transform
 	
 	//const UObject* AnimContext = GetAnimContext(SampleRole);
 	//if (const UAnimInstance* AnimInstance = Cast<UAnimInstance>(AnimContext))
diff --git a/Engine/Plugins/Animation/PoseSearch/Source/Runtime/Private/PoseSearchHistory.cpp b/Engine/Plugins/Animation/PoseSearch/Source/Runtime/Private/PoseSearchHistory.cpp
--- a/Engine/Plugins/Animation/PoseSearch/Source/Runtime/Private/PoseSearchHistory.cpp
+++ b/Engine/Plugins/Animation/PoseSearch/Source/Runtime/Private/PoseSearchHistory.cpp
@@ -735,56 +735,70 @@ bool FPoseHistory::GetTransformAtTime(float Time, FTransform& OutBoneTransform,
 
 	static_assert(RootBoneIndexType == 0 && ComponentSpaceIndexType == FBoneIndexType(-1) && WorldSpaceIndexType == FBoneIndexType(-2)); // some assumptions
 	
+	bool bSuccess = true;
 	if (BoneIndexType == ReferenceBoneIndexType)
 	{
 		OutBoneTransform = FTransform::Identity;
-		return true;
+		bSuccess = true;
 	}
-
-	if (ReferenceBoneIndexType == WorldSpaceIndexType)
+	else if (ReferenceBoneIndexType == WorldSpaceIndexType)
 	{
 		if (BoneIndexType == ComponentSpaceIndexType)
 		{
 			OutBoneTransform = Trajectory.GetSampleAtTime(Time, bExtrapolate).GetTransform();
-			return true;
+			bSuccess = true;
+		}
+		else
+		{
+			// getting BoneIndexType in ComponentSpaceIndexType and then multiplying it with the component to world transform from the trajectory
+			bSuccess = GetTransformAtTime(Time, OutBoneTransform, BoneIndexSkeleton, BoneIndexType, ComponentSpaceIndexType, bExtrapolate);
+			OutBoneTransform *= Trajectory.GetSampleAtTime(Time, bExtrapolate).GetTransform();
 		}
-
-		// getting BoneIndexType in ComponentSpaceIndexType and then multiplying it with the component to world transform from the trajectory
-		const bool bSuccess = GetTransformAtTime(Time, OutBoneTransform, BoneIndexSkeleton, BoneIndexType, ComponentSpaceIndexType, bExtrapolate);
-		OutBoneTransform *= Trajectory.GetSampleAtTime(Time, bExtrapolate).GetTransform();
-		return bSuccess;
 	}
-
-	if (BoneIndexType == WorldSpaceIndexType || BoneIndexType == ComponentSpaceIndexType)
+	else if (BoneIndexType == WorldSpaceIndexType || BoneIndexType == ComponentSpaceIndexType)
 	{
-		const bool bSuccess = GetTransformAtTime(Time, OutBoneTransform, BoneIndexSkeleton, ReferenceBoneIndexType, BoneIndexType, bExtrapolate);
+		bSuccess = GetTransformAtTime(Time, OutBoneTransform, BoneIndexSkeleton, ReferenceBoneIndexType, BoneIndexType, bExtrapolate);
 		OutBoneTransform = OutBoneTransform.Inverse();
-		return bSuccess;
 	}
-
-	check(BoneIndexType != ComponentSpaceIndexType && BoneIndexType != WorldSpaceIndexType && ReferenceBoneIndexType != WorldSpaceIndexType);
-	
-	const int32 NumEntries = PoseData.Entries.Num();
-	if (NumEntries > 0)
+	else
 	{
-		int32 NextIdx = 0;
-		int32 PrevIdx = 0;
+		check(BoneIndexType != ComponentSpaceIndexType && BoneIndexType != WorldSpaceIndexType && ReferenceBoneIndexType != WorldSpaceIndexType);
 
-		if (NumEntries > 1)
+		const int32 NumEntries = PoseData.Entries.Num();
+		if (NumEntries > 0)
 		{
-			const int32 LowerBoundIdx = LowerBound(PoseData.Entries.begin(), PoseData.Entries.end(), Time, [](const FPoseHistoryEntry& Entry, float Value) { return Value > Entry.AccumulatedSeconds; });
-			NextIdx = FMath::Clamp(LowerBoundIdx, 1, NumEntries - 1);
-			PrevIdx = NextIdx - 1;
+			int32 NextIdx = 0;
+			int32 PrevIdx = 0;
+
+			if (NumEntries > 1)
+			{
+				const int32 LowerBoundIdx = LowerBound(PoseData.Entries.begin(), PoseData.Entries.end(), Time, [](const FPoseHistoryEntry& Entry, float Value) { return Value > Entry.AccumulatedSeconds; });
+				NextIdx = FMath::Clamp(LowerBoundIdx, 1, NumEntries - 1);
+				PrevIdx = NextIdx - 1;
+			}
+
+			const FPoseHistoryEntry& PrevEntry = PoseData.Entries[PrevIdx];
+			const FPoseHistoryEntry& NextEntry = PoseData.Entries[NextIdx];
+
+			bSuccess = LerpEntries(Time, bExtrapolate, PrevEntry, NextEntry, BoneIndexSkeleton, PoseData.LastUpdateSkeleton.Get(), PoseData.BoneToTransformMap, BoneIndexType, ReferenceBoneIndexType, OutBoneTransform);
 		}
-	
-		const FPoseHistoryEntry& PrevEntry = PoseData.Entries[PrevIdx];
-		const FPoseHistoryEntry& NextEntry = PoseData.Entries[NextIdx];
+		else
+		{
+			OutBoneTransform = FTransform::Identity;
+			bSuccess = false;
+		}
+	}
 
-		return LerpEntries(Time, bExtrapolate, PrevEntry, NextEntry, BoneIndexSkeleton, PoseData.LastUpdateSkeleton.Get(), PoseData.BoneToTransformMap, BoneIndexType, ReferenceBoneIndexType, OutBoneTransform);
+#if WITH_EDITOR
+	// @todo: logging errors only if pose history was properly initialized for now (since AnimNext version doesn't implement initialization with ref pose yet),
+	// and we ultimately want to catch errors of missing bone transforms only checking valid USkeleton(s) against valid PoseSearchSchema(s)
+	if (!bSuccess && BoneIndexSkeleton)
+	{
+		UE_LOG(LogPoseSearch, Error, TEXT("FPoseHistory::GetTransformAtTime - Couldn't find root bone transform at time %f for Skeleton %s!"), Time, *BoneIndexSkeleton->GetName());
 	}
-	
-	OutBoneTransform = FTransform::Identity;
-	return false;
+#endif // WITH_EDITOR
+
+	return bSuccess;
 }
 
 bool FPoseHistory::GetCurveValueAtTime(float Time, const FName& CurveName, float& OutCurveValue, bool bExtrapolate /*= false*/) const
diff --git a/Engine/Plugins/Animation/PoseSearch/Source/Runtime/Private/PoseSearchLibrary.cpp b/Engine/Plugins/Animation/PoseSearch/Source/Runtime/Private/PoseSearchLibrary.cpp
--- a/Engine/Plugins/Animation/PoseSearch/Source/Runtime/Private/PoseSearchLibrary.cpp
+++ b/Engine/Plugins/Animation/PoseSearch/Source/Runtime/Private/PoseSearchLibrary.cpp
@@ -655,21 +655,6 @@ void UPoseSearchLibrary::UpdateMotionMatchingState(
 
 	using namespace UE::PoseSearch;
 
-	bool bAnyValidDatabases = false;
-	for (const TObjectPtr<const UPoseSearchDatabase>& Database : Databases)
-	{
-		if (Database)
-		{
-			bAnyValidDatabases = true;
-			break;
-		}
-	}
-	if (!bAnyValidDatabases)
-	{
-		UE_LOG(LogPoseSearch, Error, TEXT("No database assets provided for motion matching"));
-		return;
-	}
-
 	FMemMark Mark(FMemStack::Get());
 
 	const UPoseSearchDatabase* CurrentResultDatabase = InOutMotionMatchingState.CurrentSearchResult.Database.Get();
