PoseSearch - fix for FAnimNode_MotionMatching::UpdateAssetPlayer blendspaces bAbortBlend logic

diff --git a/Engine/Plugins/Animation/PoseSearch/Source/Runtime/Private/AnimNode_MotionMatching.cpp b/Engine/Plugins/Animation/PoseSearch/Source/Runtime/Private/AnimNode_MotionMatching.cpp
--- a/Engine/Plugins/Animation/PoseSearch/Source/Runtime/Private/AnimNode_MotionMatching.cpp
+++ b/Engine/Plugins/Animation/PoseSearch/Source/Runtime/Private/AnimNode_MotionMatching.cpp
@@ -252,7 +252,7 @@ void FAnimNode_MotionMatching::UpdateAssetPlayer(const FAnimationUpdateContext&
 				}
 
 				bool bAbortBlend = false;
-				if (!AnimPlayers.IsEmpty() && GetBlendspaceParametersDeltaThreshold() > 0.0f)
+				if (!AnimPlayers.IsEmpty() && GetBlendspaceParametersDeltaThreshold() > 0.f && AnimationAsset->IsA<UBlendSpace>())
 				{
 					if (AnimPlayers[0].GetAnimationAsset() == AnimationAsset)
 					{
