PoseSearch
- fix fo IterateOverBlendSpaceSamplingParameter bUseGridForSampling implementation
- fix for IndexDatabase sometimes not finding the asset sampler during indexing for numerical errors on the blend parameters (IterateOverSamplingParameter returned BlendParameters are an FVector, and SearchIndexAsset.GetBlendParameters() is reconstructed from 2 floats)

diff --git a/Engine/Plugins/Animation/PoseSearch/Source/Runtime/Private/PoseSearchDatabase.cpp b/Engine/Plugins/Animation/PoseSearch/Source/Runtime/Private/PoseSearchDatabase.cpp
--- a/Engine/Plugins/Animation/PoseSearch/Source/Runtime/Private/PoseSearchDatabase.cpp
+++ b/Engine/Plugins/Animation/PoseSearch/Source/Runtime/Private/PoseSearchDatabase.cpp
@@ -239,22 +239,17 @@ static void IterateOverBlendSpaceSamplingParameter(const UBlendSpace* BlendSpace
 	{
 		ProcessSamplingParameter(SingleSampleBlendParameters);
 	}
-	else
+	else if (bUseGridForSampling)
 	{
-		int32 HorizontalBlendNum, VerticalBlendNum;
-
-		if (bUseGridForSampling)
-		{
-			HorizontalBlendNum = BlendSpace->GetBlendParameter(0).GridNum + 1;
-			VerticalBlendNum = BlendSpace->IsA<UBlendSpace1D>() ? 1 : BlendSpace->GetBlendParameter(1).GridNum + 1;
-		}
-		else
+		for (const FBlendSample& BlendSample : BlendSpace->GetBlendSamples())
 		{
-			HorizontalBlendNum = FMath::Max(NumberOfHorizontalSamples, 1);
-			VerticalBlendNum = BlendSpace->IsA<UBlendSpace1D>() ? 1 : FMath::Max(NumberOfVerticalSamples, 1);
+			ProcessSamplingParameter(BlendSample.SampleValue);
 		}
-
-		check(HorizontalBlendNum >= 1 && VerticalBlendNum >= 1);
+	}
+	else
+	{
+		const int32 HorizontalBlendNum = FMath::Max(NumberOfHorizontalSamples, 1);
+		const int32 VerticalBlendNum = BlendSpace->IsA<UBlendSpace1D>() ? 1 : FMath::Max(NumberOfVerticalSamples, 1);
 
 		const FBlendParameter& HorizontalBlendParameter = BlendSpace->GetBlendParameter(0);
 		const FBlendParameter& VerticalBlendParameter = BlendSpace->GetBlendParameter(1);
diff --git a/Engine/Plugins/Animation/PoseSearch/Source/Runtime/Private/PoseSearchDerivedData.cpp b/Engine/Plugins/Animation/PoseSearch/Source/Runtime/Private/PoseSearchDerivedData.cpp
--- a/Engine/Plugins/Animation/PoseSearch/Source/Runtime/Private/PoseSearchDerivedData.cpp
+++ b/Engine/Plugins/Animation/PoseSearch/Source/Runtime/Private/PoseSearchDerivedData.cpp
@@ -909,28 +909,26 @@ static bool IndexDatabase(FSearchIndexBase& SearchIndexBase, const UPoseSearchDa
 	TMap<FSamplerMapKey, int32> SamplerMap;
 	SamplerMap.Reserve(256);
 
-	for (int32 AnimationAssetIndex = 0; AnimationAssetIndex < DatabaseToLookForAssets->GetNumAnimationAssets(); ++AnimationAssetIndex)
+	for (int32 AssetIdx = 0; AssetIdx != SearchIndexBase.Assets.Num(); ++AssetIdx)
 	{
-		if (const FPoseSearchDatabaseAnimationAssetBase* DatabaseAnimationAssetBase = DatabaseToLookForAssets->GetDatabaseAnimationAsset<FPoseSearchDatabaseAnimationAssetBase>(AnimationAssetIndex))
+		FSearchIndexAsset& SearchIndexAsset = SearchIndexBase.Assets[AssetIdx];
+
+		const FPoseSearchDatabaseAnimationAssetBase* DatabaseAnimationAssetBase = DatabaseToLookForAssets->GetDatabaseAnimationAsset<FPoseSearchDatabaseAnimationAssetBase>(SearchIndexAsset.GetSourceAssetIdx());
+		check(DatabaseAnimationAssetBase);
+
+		const int32 NumRoles = DatabaseAnimationAssetBase->GetNumRoles();
+		for (int32 RoleIndex = 0; RoleIndex < NumRoles; ++RoleIndex)
 		{
-			if (DatabaseAnimationAssetBase->IsEnabled())
+			const FRole& Role = DatabaseAnimationAssetBase->GetRole(RoleIndex);
+			UAnimationAsset* AnimationAsset = DatabaseAnimationAssetBase->GetAnimationAssetForRole(Role);
+			const FTransform RootTransformOrigin = DatabaseAnimationAssetBase->GetRootTransformOriginForRole(Role);
+			const FVector BlendParameters = SearchIndexAsset.GetBlendParameters();
+
+			const FSamplerMapKey SamplerMapKey(AnimationAsset, RootTransformOrigin, BlendParameters);
+			if (!SamplerMap.Contains(SamplerMapKey))
 			{
-				DatabaseAnimationAssetBase->IterateOverSamplingParameter([DatabaseAnimationAssetBase, &SamplerMap, &Samplers](const FVector& BlendParameters)
-					{
-						const int32 NumRoles = DatabaseAnimationAssetBase->GetNumRoles();
-						for (int32 RoleIndex = 0; RoleIndex < NumRoles; ++RoleIndex)
-						{
-							const FRole& Role = DatabaseAnimationAssetBase->GetRole(RoleIndex);
-							const UAnimationAsset* AnimationAsset = DatabaseAnimationAssetBase->GetAnimationAssetForRole(Role);
-							const FTransform RootTransformOrigin = DatabaseAnimationAssetBase->GetRootTransformOriginForRole(Role);
-							const FSamplerMapKey SamplerMapKey(AnimationAsset, RootTransformOrigin, BlendParameters);
-							if (!SamplerMap.Contains(SamplerMapKey))
-							{
-								SamplerMap.Add(SamplerMapKey, Samplers.Num());
-								Samplers.Emplace(AnimationAsset, RootTransformOrigin, BlendParameters, FAnimationAssetSampler::DefaultRootTransformSamplingRate, false);
-							}
-						}
-					});
+				SamplerMap.Add(SamplerMapKey, Samplers.Num());
+				Samplers.Emplace(AnimationAsset, RootTransformOrigin, BlendParameters, FAnimationAssetSampler::DefaultRootTransformSamplingRate, false);
 			}
 		}
 	}
@@ -982,8 +980,12 @@ static bool IndexDatabase(FSearchIndexBase& SearchIndexBase, const UPoseSearchDa
 				UAnimationAsset* AnimationAsset = DatabaseAnimationAssetBase->GetAnimationAssetForRole(Role);
 				const FTransform RootTransformOrigin = DatabaseAnimationAssetBase->GetRootTransformOriginForRole(Role);
 				const FVector BlendParameters = SearchIndexAsset.GetBlendParameters();
-				const int32 SamplerIndex = SamplerMap[{ AnimationAsset, RootTransformOrigin, BlendParameters }];
-				TempAssetSamplers.AnimationAssetSamplers.Emplace(&Samplers[SamplerIndex]);
+
+				const FSamplerMapKey SamplerMapKey(AnimationAsset, RootTransformOrigin, BlendParameters);
+				const int32* SamplerIndex = SamplerMap.Find(SamplerMapKey);
+				check(SamplerIndex && Samplers.IsValidIndex(*SamplerIndex));
+
+				TempAssetSamplers.AnimationAssetSamplers.Emplace(&Samplers[*SamplerIndex]);
 				TempAssetSamplers.MirrorDataCaches.Emplace(MirrorDataCache);
 				TempBoneContainers.Emplace(RoledBoneContainers[Role]);
 				TempRoleToIndex.Add(Role) = RoleIndex;
diff --git a/Engine/Plugins/Animation/PoseSearch/Source/Runtime/Public/PoseSearch/PoseSearchDefines.h b/Engine/Plugins/Animation/PoseSearch/Source/Runtime/Public/PoseSearch/PoseSearchDefines.h
--- a/Engine/Plugins/Animation/PoseSearch/Source/Runtime/Public/PoseSearch/PoseSearchDefines.h
+++ b/Engine/Plugins/Animation/PoseSearch/Source/Runtime/Public/PoseSearch/PoseSearchDefines.h
@@ -42,7 +42,7 @@ static constexpr int32 PreallocatedSearchesNum = 4;
 static constexpr float FiniteDelta = 1 / 60.0f; // Time delta used for computing pose derivatives
 static constexpr int32 MaxNumberOfCollectedPoseCandidatesPerDatabase = 256;
 static constexpr int32 InvalidSearchId = INDEX_NONE; // @todo: it should be 0, since the search id is an hash!
-static constexpr int32 DatabaseIndexDerivedDataCacheKeyVersion = 2679; // bump this value to force reindex ALL the UPoseSearchDatabases, without having to change the global FDevSystemGuids::POSESEARCHDB_DERIVEDDATA_VER
+static constexpr int32 DatabaseIndexDerivedDataCacheKeyVersion = 2683; // bump this value to force reindex ALL the UPoseSearchDatabases, without having to change the global FDevSystemGuids::POSESEARCHDB_DERIVEDDATA_VER
 static constexpr int32 DefaultTickPriority = 100;
 
 template<uint32 Alignment = DEFAULT_ALIGNMENT> class TMemStackSetAllocator : public TSetAllocator<TSparseArrayAllocator<TMemStackAllocator<Alignment>, TMemStackAllocator<Alignment>>, TMemStackAllocator<Alignment>> {};
