PoseSearch - exposing InterruptMode to UPoseSearchDebuggerReflection

diff --git a/Engine/Plugins/Animation/PoseSearch/Source/Editor/Private/PoseSearchDebuggerReflection.h b/Engine/Plugins/Animation/PoseSearch/Source/Editor/Private/PoseSearchDebuggerReflection.h
--- a/Engine/Plugins/Animation/PoseSearch/Source/Editor/Private/PoseSearchDebuggerReflection.h
+++ b/Engine/Plugins/Animation/PoseSearch/Source/Editor/Private/PoseSearchDebuggerReflection.h
@@ -2,7 +2,7 @@
 
 #pragma once
 
-#include "CoreMinimal.h"
+#include "PoseSearch/PoseSearchLibrary.h"
 #include "PoseSearchDebuggerReflection.generated.h"
 
 /**
@@ -14,6 +14,9 @@ class POSESEARCHEDITOR_API UPoseSearchDebuggerReflection : public UObject
 	GENERATED_BODY()
 
 public:
+	UPROPERTY(VisibleAnywhere, Category = "Motion Matching State")
+	EPoseSearchInterruptMode InterruptMode = EPoseSearchInterruptMode::DoNotInterrupt;
+
 	/** Time since last PoseSearch */
 	UPROPERTY(VisibleAnywhere, Category="Motion Matching State", meta = (ForceUnits = "s"))
 	float ElapsedPoseSearchTime = 0.0f;
diff --git a/Engine/Plugins/Animation/PoseSearch/Source/Editor/Private/PoseSearchDebuggerView.cpp b/Engine/Plugins/Animation/PoseSearch/Source/Editor/Private/PoseSearchDebuggerView.cpp
--- a/Engine/Plugins/Animation/PoseSearch/Source/Editor/Private/PoseSearchDebuggerView.cpp
+++ b/Engine/Plugins/Animation/PoseSearch/Source/Editor/Private/PoseSearchDebuggerView.cpp
@@ -94,6 +94,7 @@ void SDebuggerDetailsView::UpdateReflection(const FTraceMotionMatchingStateMessa
 {
 	check(Reflection);
 
+	Reflection->InterruptMode = State.InterruptMode;
 	Reflection->ElapsedPoseSearchTime = State.ElapsedPoseSearchTime;
 	Reflection->AssetPlayerTime = State.AssetPlayerTime;
 	Reflection->LastDeltaTime = State.DeltaTime;
diff --git a/Engine/Plugins/Animation/PoseSearch/Source/Editor/Private/Trace/PoseSearchTraceAnalyzer.cpp b/Engine/Plugins/Animation/PoseSearch/Source/Editor/Private/Trace/PoseSearchTraceAnalyzer.cpp
--- a/Engine/Plugins/Animation/PoseSearch/Source/Editor/Private/Trace/PoseSearchTraceAnalyzer.cpp
+++ b/Engine/Plugins/Animation/PoseSearch/Source/Editor/Private/Trace/PoseSearchTraceAnalyzer.cpp
@@ -23,6 +23,7 @@ void FTraceAnalyzer::OnAnalysisBegin(const FOnAnalysisContext& Context)
 	
 	Builder.RouteEvent(RouteId_MotionMatchingState, "PoseSearch", "MotionMatchingState");
 	Builder.RouteEvent(RouteId_MotionMatchingState2, "PoseSearch", "MotionMatchingState2");
+	Builder.RouteEvent(RouteId_MotionMatchingState3, "PoseSearch", "MotionMatchingState3");
 }
 
 bool FTraceAnalyzer::OnEvent(uint16 RouteId, EStyle Style, const FOnEventContext& Context)
@@ -31,37 +32,34 @@ bool FTraceAnalyzer::OnEvent(uint16 RouteId, EStyle Style, const FOnEventContext
 
 	TraceServices::FAnalysisSessionEditScope Scope(Session);
 
-	if (RouteId == RouteId_MotionMatchingState)
-	{
-		FTraceMotionMatchingStateMessage Message;
-		FMemoryReaderView Archive(Context.EventData.GetArrayView<uint8>("Data"));
+	FCustomVersionContainer CustomVersionContainer;
 
-		// Ensure we have not versions.
-		FCustomVersionContainer CustomVersionContainer;
-		Archive.SetCustomVersions(CustomVersionContainer);
-		
-		Archive << Message;
-		TraceProvider.AppendMotionMatchingState(Message, Context.EventTime.AsSeconds(Message.Cycle));
-	}
-	else if (RouteId == RouteId_MotionMatchingState2)
+	switch (RouteId)
 	{
-		FTraceMotionMatchingStateMessage Message;
-		FMemoryReaderView Archive(Context.EventData.GetArrayView<uint8>("Data"));
-
+	case RouteId_MotionMatchingState:
+		// Ensure we have not versions.
+		break;
+	case RouteId_MotionMatchingState2:
 		// Ensure we have the version that matches this RouteId event.
-		FCustomVersionContainer CustomVersionContainer;
 		CustomVersionContainer.SetVersion(FPoseSearchCustomVersion::GUID, FPoseSearchCustomVersion::DeprecatedTrajectoryTypes, TEXT("Dev-PoseSearch-Version"));
-		Archive.SetCustomVersions(CustomVersionContainer);
-
-		Archive << Message;
-		TraceProvider.AppendMotionMatchingState(Message, Context.EventTime.AsSeconds(Message.Cycle));
-	}
-	else
-	{
-		// Should not happen
+		break;
+	case RouteId_MotionMatchingState3:
+		// Ensure we have the version that matches this RouteId event.
+		CustomVersionContainer.SetVersion(FPoseSearchCustomVersion::GUID, FPoseSearchCustomVersion::AddedInterruptModeToDebugger, TEXT("Dev-PoseSearch-Version"));
+		break;
+	default:
+		// Should not happen...
 		checkNoEntry();
+		break;
 	}
 
+	FTraceMotionMatchingStateMessage Message;
+	FMemoryReaderView Archive(Context.EventData.GetArrayView<uint8>("Data"));
+	Archive.SetCustomVersions(CustomVersionContainer);
+	Archive << Message;
+	TraceProvider.AppendMotionMatchingState(Message, Context.EventTime.AsSeconds(Message.Cycle));
+
 	return true;
 }
+
 } // namespace UE::PoseSearch
diff --git a/Engine/Plugins/Animation/PoseSearch/Source/Editor/Private/Trace/PoseSearchTraceAnalyzer.h b/Engine/Plugins/Animation/PoseSearch/Source/Editor/Private/Trace/PoseSearchTraceAnalyzer.h
--- a/Engine/Plugins/Animation/PoseSearch/Source/Editor/Private/Trace/PoseSearchTraceAnalyzer.h
+++ b/Engine/Plugins/Animation/PoseSearch/Source/Editor/Private/Trace/PoseSearchTraceAnalyzer.h
@@ -33,6 +33,7 @@ private:
 	{
 		RouteId_MotionMatchingState,
 		RouteId_MotionMatchingState2,
+		RouteId_MotionMatchingState3,
 	};
 
 	TraceServices::IAnalysisSession& Session;
diff --git a/Engine/Plugins/Animation/PoseSearch/Source/Runtime/Private/PoseSearchLibrary.cpp b/Engine/Plugins/Animation/PoseSearch/Source/Runtime/Private/PoseSearchLibrary.cpp
--- a/Engine/Plugins/Animation/PoseSearch/Source/Runtime/Private/PoseSearchLibrary.cpp
+++ b/Engine/Plugins/Animation/PoseSearch/Source/Runtime/Private/PoseSearchLibrary.cpp
@@ -406,7 +406,8 @@ void UPoseSearchLibrary::TraceMotionMatching(
 		float ElapsedPoseSearchTime,
 		float DeltaTime,
 		bool bSearch,
-		float WantedPlayRate)
+		float WantedPlayRate,
+		EPoseSearchInterruptMode InterruptMode)
 {
 	using namespace UE::PoseSearch;
 
@@ -425,6 +426,9 @@ void UPoseSearchLibrary::TraceMotionMatching(
 	uint32 SearchId = 787;
 
 	FTraceMotionMatchingStateMessage TraceState;
+
+	TraceState.InterruptMode = InterruptMode;
+
 	const FSearchResult& CurrentResult = SearchResult;
 
 	const int32 AnimContextsNum = SearchContext.GetAnimContexts().Num();
@@ -777,7 +781,7 @@ void UPoseSearchLibrary::UpdateMotionMatchingState(
 	InOutMotionMatchingState.PoseIndicesHistory.Update(InOutMotionMatchingState.CurrentSearchResult, DeltaTime, PoseReselectHistory);
 
 #if UE_POSE_SEARCH_TRACE_ENABLED
-	TraceMotionMatching(SearchContext, InOutMotionMatchingState.CurrentSearchResult, InOutMotionMatchingState.ElapsedPoseSearchTime, DeltaTime, bSearch, InOutMotionMatchingState.WantedPlayRate);
+	TraceMotionMatching(SearchContext, InOutMotionMatchingState.CurrentSearchResult, InOutMotionMatchingState.ElapsedPoseSearchTime, DeltaTime, bSearch, InOutMotionMatchingState.WantedPlayRate, InterruptMode);
 #endif // UE_POSE_SEARCH_TRACE_ENABLED
 
 #if WITH_EDITORONLY_DATA && ENABLE_ANIM_DEBUG
@@ -1260,7 +1264,8 @@ UE::PoseSearch::FSearchResult UPoseSearchLibrary::MotionMatch(
 #endif // ENABLE_DRAW_DEBUG && ENABLE_ANIM_DEBUG
 
 #if UE_POSE_SEARCH_TRACE_ENABLED
-	TraceMotionMatching(SearchContext, SearchResult, 0.f, DeltaSeconds, true, 1.f);
+	// @todo: add and handle an interrupt mode input param MotionMatch
+	TraceMotionMatching(SearchContext, SearchResult, 0.f, DeltaSeconds, true, 1.f, EPoseSearchInterruptMode::DoNotInterrupt);
 #endif // UE_POSE_SEARCH_TRACE_ENABLED
 
 	return SearchResult;
diff --git a/Engine/Plugins/Animation/PoseSearch/Source/Runtime/Private/Trace/PoseSearchTraceLogger.cpp b/Engine/Plugins/Animation/PoseSearch/Source/Runtime/Private/Trace/PoseSearchTraceLogger.cpp
--- a/Engine/Plugins/Animation/PoseSearch/Source/Runtime/Private/Trace/PoseSearchTraceLogger.cpp
+++ b/Engine/Plugins/Animation/PoseSearch/Source/Runtime/Private/Trace/PoseSearchTraceLogger.cpp
@@ -22,12 +22,17 @@ UE_TRACE_EVENT_BEGIN(PoseSearch, MotionMatchingState2)
 	UE_TRACE_EVENT_FIELD(uint8[], Data)
 UE_TRACE_EVENT_END()
 
+// We apply the FPoseSearchCustomVersion::AddedInterruptModeToDebugger version when loading.
+UE_TRACE_EVENT_BEGIN(PoseSearch, MotionMatchingState3)
+	UE_TRACE_EVENT_FIELD(uint8[], Data)
+UE_TRACE_EVENT_END()
+
 namespace UE::PoseSearch
 {
 
 // @todo: These should be deprecated. No need to expose them.
 const FName FTraceLogger::Name("PoseSearch");
-const FName FTraceMotionMatchingStateMessage::Name("MotionMatchingState2");
+const FName FTraceMotionMatchingStateMessage::Name("MotionMatchingState3");
 
 FArchive& operator<<(FArchive& Ar, FTraceMessage& State)
 {
@@ -77,6 +82,12 @@ FArchive& operator<<(FArchive& Ar, FTraceMotionMatchingStateMessage& State)
 	Ar << State.PoseHistories;	
 	Ar << State.CurrentDbEntryIdx;
 	Ar << State.CurrentPoseEntryIdx;
+
+	// InterruptMode has been introduced with AddedInterruptModeToDebugger version
+	if (Ar.CustomVer(FPoseSearchCustomVersion::GUID) >= FPoseSearchCustomVersion::AddedInterruptModeToDebugger)
+	{
+		Ar << State.InterruptMode;
+	}
 	return Ar;
 }
 
@@ -92,7 +103,7 @@ void FTraceMotionMatchingStateMessage::Output()
 	Archive.SetCustomVersions(CustomVersionContainer);
 	
 	Archive << *this;
-	UE_TRACE_LOG(PoseSearch, MotionMatchingState2, PoseSearchChannel) << MotionMatchingState2.Data(ArchiveData.GetData(), ArchiveData.Num());
+	UE_TRACE_LOG(PoseSearch, MotionMatchingState3, PoseSearchChannel) << MotionMatchingState3.Data(ArchiveData.GetData(), ArchiveData.Num());
 #endif
 }
 
diff --git a/Engine/Plugins/Animation/PoseSearch/Source/Runtime/Public/PoseSearch/PoseSearchCustomVersion.h b/Engine/Plugins/Animation/PoseSearch/Source/Runtime/Public/PoseSearch/PoseSearchCustomVersion.h
--- a/Engine/Plugins/Animation/PoseSearch/Source/Runtime/Public/PoseSearch/PoseSearchCustomVersion.h
+++ b/Engine/Plugins/Animation/PoseSearch/Source/Runtime/Public/PoseSearch/PoseSearchCustomVersion.h
@@ -13,9 +13,12 @@ struct POSESEARCH_API FPoseSearchCustomVersion
 		// Before any version changes were made in the plugin
 		BeforeCustomVersionWasAdded = 0,
 
-		// Deprecated FPoseSearchQueryTrajectory and FPoseSearchQueryTrajectorySample. This version matches the RouteId_MotionMatchingState_2 event.
+		// Deprecated FPoseSearchQueryTrajectory and FPoseSearchQueryTrajectorySample. This version matches the RouteId_MotionMatchingState2 event.
 		DeprecatedTrajectoryTypes,
 
+		// Added InterruptMode to posesearch debugger. This version matches the RouteId_MotionMatchingState3 event.
+		AddedInterruptModeToDebugger,
+
 		// -----<new versions can be added above this line>-------------------------------------------------
 		VersionPlusOne,
 		LatestVersion = VersionPlusOne - 1
@@ -27,3 +30,4 @@ struct POSESEARCH_API FPoseSearchCustomVersion
 private:
 	FPoseSearchCustomVersion() {}
 };
+ 
\ No newline at end of file
diff --git a/Engine/Plugins/Animation/PoseSearch/Source/Runtime/Public/PoseSearch/PoseSearchLibrary.h b/Engine/Plugins/Animation/PoseSearch/Source/Runtime/Public/PoseSearch/PoseSearchLibrary.h
--- a/Engine/Plugins/Animation/PoseSearch/Source/Runtime/Public/PoseSearch/PoseSearchLibrary.h
+++ b/Engine/Plugins/Animation/PoseSearch/Source/Runtime/Public/PoseSearch/PoseSearchLibrary.h
@@ -126,7 +126,20 @@ class POSESEARCH_API UPoseSearchLibrary : public UBlueprintFunctionLibrary
 		float ElapsedPoseSearchTime,
 		float DeltaTime,
 		bool bSearch,
-		float WantedPlayRate);
+		float WantedPlayRate,
+		EPoseSearchInterruptMode InterruptMode);
+
+	UE_DEPRECATED(5.6, "Use TraceMotionMatching with different signature instead.")
+	static void TraceMotionMatching(
+		UE::PoseSearch::FSearchContext& SearchContext,
+		const UE::PoseSearch::FSearchResult& SearchResult,
+		float ElapsedPoseSearchTime,
+		float DeltaTime,
+		bool bSearch,
+		float WantedPlayRate)
+	{
+		TraceMotionMatching(SearchContext, SearchResult, ElapsedPoseSearchTime, DeltaTime, bSearch, WantedPlayRate, EPoseSearchInterruptMode::DoNotInterrupt);
+	}
 
 	UE_DEPRECATED(5.6, "Use TraceMotionMatching with different signature instead.")
 	static void TraceMotionMatching(
@@ -136,7 +149,7 @@ class POSESEARCH_API UPoseSearchLibrary : public UBlueprintFunctionLibrary
 		bool bSearch,
 		float RecordingTime)
 	{
-		TraceMotionMatching(SearchContext, CurrentState.CurrentSearchResult, CurrentState.ElapsedPoseSearchTime, DeltaTime, bSearch, CurrentState.WantedPlayRate);
+		TraceMotionMatching(SearchContext, CurrentState.CurrentSearchResult, CurrentState.ElapsedPoseSearchTime, DeltaTime, bSearch, CurrentState.WantedPlayRate, EPoseSearchInterruptMode::DoNotInterrupt);
 	}
 	
 	UE_DEPRECATED(5.5, "Use TraceMotionMatching with different signature instead.")
@@ -149,7 +162,7 @@ class POSESEARCH_API UPoseSearchLibrary : public UBlueprintFunctionLibrary
 		bool bSearch,
 		float RecordingTime)
 	{
-		TraceMotionMatching(SearchContext, CurrentResult, ElapsedPoseSearchTime, DeltaTime, bSearch, 1.f);
+		TraceMotionMatching(SearchContext, CurrentResult, ElapsedPoseSearchTime, DeltaTime, bSearch, 1.f, EPoseSearchInterruptMode::DoNotInterrupt);
 	}
 
 #endif // UE_POSE_SEARCH_TRACE_ENABLED
diff --git a/Engine/Plugins/Animation/PoseSearch/Source/Runtime/Public/PoseSearch/Trace/PoseSearchTraceLogger.h b/Engine/Plugins/Animation/PoseSearch/Source/Runtime/Public/PoseSearch/Trace/PoseSearchTraceLogger.h
--- a/Engine/Plugins/Animation/PoseSearch/Source/Runtime/Public/PoseSearch/Trace/PoseSearchTraceLogger.h
+++ b/Engine/Plugins/Animation/PoseSearch/Source/Runtime/Public/PoseSearch/Trace/PoseSearchTraceLogger.h
@@ -3,7 +3,7 @@
 #pragma once
 
 #include "PoseSearch/PoseSearchContext.h"
-#include "PoseSearch/PoseSearchHistory.h"
+#include "PoseSearch/PoseSearchLibrary.h"
 
 UE_TRACE_CHANNEL_EXTERN(PoseSearchChannel, POSESEARCH_API);
 
@@ -89,6 +89,8 @@ struct POSESEARCH_API FTraceMotionMatchingStateMessage : public FTraceMessage
 	/** Index of the current pose in DatabaseEntries[CurrentDbEntryIdx].PoseEntries */
 	int32 CurrentPoseEntryIdx = INDEX_NONE;
 
+	EPoseSearchInterruptMode InterruptMode = EPoseSearchInterruptMode::DoNotInterrupt;
+	
 	/** Output the current state info to the logger */
 	void Output();
 
