Add flag on AnimationAssetSample to disregard compressed data enforcement as its used outside of the DDC build API

diff --git a/Engine/Plugins/Animation/PoseSearch/Source/Runtime/Private/PoseSearchAssetSampler.cpp b/Engine/Plugins/Animation/PoseSearch/Source/Runtime/Private/PoseSearchAssetSampler.cpp
--- a/Engine/Plugins/Animation/PoseSearch/Source/Runtime/Private/PoseSearchAssetSampler.cpp
+++ b/Engine/Plugins/Animation/PoseSearch/Source/Runtime/Private/PoseSearchAssetSampler.cpp
@@ -50,7 +50,7 @@ static FTransform ExtrapolateRootMotion(FTransform SampleToExtrapolate, float Sa
 	return ExtrapolatedRootMotion;
 }
 
-static FTransform ExtractRootTransformInternal(const UAnimMontage* AnimMontage, float StartTime, float EndTime)
+static FTransform ExtractRootTransformInternal(const UAnimMontage* AnimMontage, float StartTime, float EndTime, bool bEnforceCompressedDataSampling)
 {
 	// @todo: add support for SlotName / multiple SlotAnimTracks
 	if (AnimMontage->SlotAnimTracks.Num() != 1)
@@ -61,7 +61,7 @@ static FTransform ExtractRootTransformInternal(const UAnimMontage* AnimMontage,
 
 	FAnimExtractContext ExtractionCtx(0.0, true, FDeltaTimeRecord(StartTime), true);
 #if WITH_EDITOR
-	ExtractionCtx.bEnforceCompressedDataSampling = true;
+	ExtractionCtx.bEnforceCompressedDataSampling = bEnforceCompressedDataSampling;
 	ExtractionCtx.TargetPlatform = GetTargetPlatformManagerRef().GetRunningTargetPlatform();
 #endif // WITH_EDITOR
 
@@ -154,7 +154,7 @@ static FTransform ExtractBlendSpaceRootMotion(float StartTime, float DeltaTime,
 }
 
 static void ProcessRootTransform(const UBlendSpace* BlendSpace, const FVector& BlendParameters, float CachedPlayLength,
-	int32 RootTransformSamplingRate, bool bIsLoopable, TArray<FTransform>& AccumulatedRootTransform)
+	int32 RootTransformSamplingRate, bool bIsLoopable, TArray<FTransform>& AccumulatedRootTransform, bool bEnforceCompressedDataSampling)
 {
 	// Pre-compute root motion
 	const int32 NumRootSamples = FMath::Max(FMath::CeilToInt(CachedPlayLength * RootTransformSamplingRate) + 1, 1);
@@ -185,7 +185,7 @@ static void ProcessRootTransform(const UBlendSpace* BlendSpace, const FVector& B
 
 			FAnimExtractContext ExtractionCtx(0.0, true, {}, bIsLoopable);
 #if WITH_EDITOR
-			ExtractionCtx.bEnforceCompressedDataSampling = true;
+			ExtractionCtx.bEnforceCompressedDataSampling = bEnforceCompressedDataSampling;
 			ExtractionCtx.TargetPlatform = GetTargetPlatformManagerRef().GetRunningTargetPlatform();
 #endif // WITH_EDITOR
 			
@@ -238,12 +238,12 @@ static int32 GetHighestWeightSample(const TArray<struct FBlendSampleData>& Sampl
 
 //////////////////////////////////////////////////////////////////////////
 // FAssetSamplerBase
-FAnimationAssetSampler::FAnimationAssetSampler(const UAnimationAsset* InAnimationAsset, const FTransform& InRootTransformOrigin, const FVector& InBlendParameters, int32 InRootTransformSamplingRate, bool bPreProcessRootTransform)
+FAnimationAssetSampler::FAnimationAssetSampler(const UAnimationAsset* InAnimationAsset, const FTransform& InRootTransformOrigin, const FVector& InBlendParameters, int32 InRootTransformSamplingRate, bool bPreProcessRootTransform, bool bInEnforceCompressedDataSampling)
 {
 	Init(InAnimationAsset, InRootTransformOrigin, InBlendParameters, InRootTransformSamplingRate, bPreProcessRootTransform);
 }
 
-void FAnimationAssetSampler::Init(const UAnimationAsset* InAnimationAsset, const FTransform& InRootTransformOrigin, const FVector& InBlendParameters, int32 InRootTransformSamplingRate, bool bPreProcessRootTransform)
+void FAnimationAssetSampler::Init(const UAnimationAsset* InAnimationAsset, const FTransform& InRootTransformOrigin, const FVector& InBlendParameters, int32 InRootTransformSamplingRate, bool bPreProcessRootTransform, bool bInEnforceCompressedDataSampling)
 {
 	AnimationAssetPtr = InAnimationAsset;
 	RootTransformOrigin = InRootTransformOrigin;
@@ -252,6 +252,7 @@ void FAnimationAssetSampler::Init(const UAnimationAsset* InAnimationAsset, const
 	CachedPlayLength = GetPlayLength(AnimationAssetPtr.Get(), BlendParameters);
 
 #if WITH_EDITOR
+	bEnforceCompressedDataSampling = bInEnforceCompressedDataSampling;
 	ITargetPlatform* TargetPlatform = GetTargetPlatformManagerRef().GetRunningTargetPlatform();		
 	PlatformHash = FIoHash::Zero;
 	if (const UAnimSequence* AnimSequence = Cast<UAnimSequence>(InAnimationAsset))
@@ -447,7 +448,7 @@ void FAnimationAssetSampler::ExtractPose(float Time, FCompactPose& OutPose) cons
 
 #if WITH_EDITOR
 	ExtractionCtx.bExtractWithRootMotionProvider = false;
-	ExtractionCtx.bEnforceCompressedDataSampling = true;
+	ExtractionCtx.bEnforceCompressedDataSampling = bEnforceCompressedDataSampling;
 	ExtractionCtx.TargetPlatform = GetTargetPlatformManagerRef().GetRunningTargetPlatform();
 #endif // WITH_EDITOR
 
@@ -466,7 +467,7 @@ void FAnimationAssetSampler::ExtractPose(float Time, FCompactPose& OutPose, FBle
 
 #if WITH_EDITOR
 	ExtractionCtx.bExtractWithRootMotionProvider = false;
-	ExtractionCtx.bEnforceCompressedDataSampling = true;
+	ExtractionCtx.bEnforceCompressedDataSampling = bEnforceCompressedDataSampling;
 	ExtractionCtx.TargetPlatform = GetTargetPlatformManagerRef().GetRunningTargetPlatform();
 #endif // WITH_EDITOR
 
@@ -523,7 +524,7 @@ FTransform FAnimationAssetSampler::ExtractRootTransform(float Time) const
 	{
 		if (IsLoopable())
 		{
-			RootTransform = ExtractRootTransformInternal(AnimMontage, 0.f, Time);
+			RootTransform = ExtractRootTransformInternal(AnimMontage, 0.f, Time, bEnforceCompressedDataSampling);
 		}
 		else
 		{
@@ -535,7 +536,7 @@ FTransform FAnimationAssetSampler::ExtractRootTransform(float Time) const
 			// animation to estimate where the root would be at Time
 			if (ExtrapolationTime < -SMALL_NUMBER)
 			{
-				FTransform SampleToExtrapolate = ExtractRootTransformInternal(AnimMontage, 0.f, ExtrapolationSampleTime);
+				FTransform SampleToExtrapolate = ExtractRootTransformInternal(AnimMontage, 0.f, ExtrapolationSampleTime, bEnforceCompressedDataSampling);
 
 				const FTransform ExtrapolatedRootMotion = UE::PoseSearch::ExtrapolateRootMotion(
 					SampleToExtrapolate,
@@ -545,13 +546,13 @@ FTransform FAnimationAssetSampler::ExtractRootTransform(float Time) const
 			}
 			else
 			{
-				RootTransform = ExtractRootTransformInternal(AnimMontage, 0.f, ClampedTime);
+				RootTransform = ExtractRootTransformInternal(AnimMontage, 0.f, ClampedTime, bEnforceCompressedDataSampling);
 
 				// If Time is greater than PlayLength, ExtrapolationTime will be a positive number. In this case, we extrapolate
 				// the end of the animation to estimate where the root would be at Time
 				if (ExtrapolationTime > SMALL_NUMBER)
 				{
-					FTransform SampleToExtrapolate = ExtractRootTransformInternal(AnimMontage, PlayLength - ExtrapolationSampleTime, PlayLength);
+					FTransform SampleToExtrapolate = ExtractRootTransformInternal(AnimMontage, PlayLength - ExtrapolationSampleTime, PlayLength, bEnforceCompressedDataSampling);
 
 					const FTransform ExtrapolatedRootMotion = UE::PoseSearch::ExtrapolateRootMotion(
 						SampleToExtrapolate,
@@ -566,7 +567,7 @@ FTransform FAnimationAssetSampler::ExtractRootTransform(float Time) const
 	{
 		FAnimExtractContext ExtractionCtx(0.0, true, FDeltaTimeRecord(Time), true);
 #if WITH_EDITOR
-		ExtractionCtx.bEnforceCompressedDataSampling = true;
+		ExtractionCtx.bEnforceCompressedDataSampling = bEnforceCompressedDataSampling;
 		ExtractionCtx.TargetPlatformHash = PlatformHash;
 #endif // WITH_EDITOR
 		
@@ -619,7 +620,7 @@ void FAnimationAssetSampler::Process()
 {
 	if (const UBlendSpace* BlendSpace = Cast<UBlendSpace>(AnimationAssetPtr.Get()))
 	{
-		ProcessRootTransform(BlendSpace, BlendParameters, CachedPlayLength, RootTransformSamplingRate, IsLoopable(), AccumulatedRootTransform);
+		ProcessRootTransform(BlendSpace, BlendParameters, CachedPlayLength, RootTransformSamplingRate, IsLoopable(), AccumulatedRootTransform, bEnforceCompressedDataSampling);
 	}
 }
 
diff --git a/Engine/Plugins/Animation/PoseSearch/Source/Runtime/Private/PoseSearchDatabase.cpp b/Engine/Plugins/Animation/PoseSearch/Source/Runtime/Private/PoseSearchDatabase.cpp
--- a/Engine/Plugins/Animation/PoseSearch/Source/Runtime/Private/PoseSearchDatabase.cpp
+++ b/Engine/Plugins/Animation/PoseSearch/Source/Runtime/Private/PoseSearchDatabase.cpp
@@ -1560,7 +1560,7 @@ UE::PoseSearch::FSearchResult UPoseSearchDatabase::SearchContinuingPose(UE::Pose
 		if (const UAnimationAsset* AnimationAsset = DatabaseAnimationAssetBase->GetAnimationAssetForRole(DatabaseAnimationAssetBase->GetRole(RoleIndex)))
 		{
 			// sampler used only to extract the notify states. RootTransformOrigin can be set as Identity, since will not be relevant
-			const FAnimationAssetSampler SequenceBaseSampler(AnimationAsset, FTransform::Identity, SearchIndexAsset.GetBlendParameters(), FAnimationAssetSampler::DefaultRootTransformSamplingRate, false);
+			const FAnimationAssetSampler SequenceBaseSampler(AnimationAsset, FTransform::Identity, SearchIndexAsset.GetBlendParameters(), FAnimationAssetSampler::DefaultRootTransformSamplingRate, false, false);
 
 			bool bDone = false;
 			SequenceBaseSampler.ExtractAnimNotifyStates(SampleTime, PreAllocatedNotifyContext, [&UpdatedContinuingPoseCostBias, &bDone](const UAnimNotifyState* AnimNotifyState)
diff --git a/Engine/Plugins/Animation/PoseSearch/Source/Runtime/Private/PoseSearchLibrary.cpp b/Engine/Plugins/Animation/PoseSearch/Source/Runtime/Private/PoseSearchLibrary.cpp
--- a/Engine/Plugins/Animation/PoseSearch/Source/Runtime/Private/PoseSearchLibrary.cpp
+++ b/Engine/Plugins/Animation/PoseSearch/Source/Runtime/Private/PoseSearchLibrary.cpp
@@ -558,7 +558,7 @@ void UPoseSearchLibrary::TraceMotionMatching(
 			if (UAnimationAsset* AnimationAsset = Cast<UAnimationAsset>(DatabaseAsset->GetAnimationAsset()))
 			{
 				// Simulate the time step to get accurate root motion prediction for this frame.
-				FAnimationAssetSampler Sampler(AnimationAsset);
+				FAnimationAssetSampler Sampler(AnimationAsset, FTransform::Identity,FVector::ZeroVector, FAnimationAssetSampler::DefaultRootTransformSamplingRate, true, false);
 
 				const float TimeStep = DeltaTime * WantedPlayRate;
 				const FTransform PrevRoot = Sampler.ExtractRootTransform(CurrentResult.AssetTime);
@@ -1049,7 +1049,7 @@ UE::PoseSearch::FSearchResult UPoseSearchLibrary::MotionMatch(
 						}
 					}
 
-					const FAnimationAssetSampler Sampler(AnimationAsset, FTransform::Identity, FVector::ZeroVector, FAnimationAssetSampler::DefaultRootTransformSamplingRate, false);
+					const FAnimationAssetSampler Sampler(AnimationAsset, FTransform::Identity, FVector::ZeroVector, FAnimationAssetSampler::DefaultRootTransformSamplingRate, false, false);
 					for (int32 i = 0; i < 2; ++i)
 					{
 						const float FuturePoseExtractionTime = FutureAnimationTime + (i - 1) * FiniteDelta;
diff --git a/Engine/Plugins/Animation/PoseSearch/Source/Runtime/Public/PoseSearch/PoseSearchAssetSampler.h b/Engine/Plugins/Animation/PoseSearch/Source/Runtime/Public/PoseSearch/PoseSearchAssetSampler.h
--- a/Engine/Plugins/Animation/PoseSearch/Source/Runtime/Public/PoseSearch/PoseSearchAssetSampler.h
+++ b/Engine/Plugins/Animation/PoseSearch/Source/Runtime/Public/PoseSearch/PoseSearchAssetSampler.h
@@ -21,8 +21,8 @@ namespace UE::PoseSearch
 struct FAnimationAssetSampler
 {
 	enum { DefaultRootTransformSamplingRate = 30 };
-	POSESEARCH_API FAnimationAssetSampler(const UAnimationAsset* InAnimationAsset = nullptr, const FTransform& InRootTransformOrigin = FTransform::Identity, const FVector& InBlendParameters = FVector::ZeroVector, int32 InRootTransformSamplingRate = DefaultRootTransformSamplingRate, bool bPreProcessRootTransform = true);
-	POSESEARCH_API void Init(const UAnimationAsset* InAnimationAsset, const FTransform& InRootTransformOrigin = FTransform::Identity, const FVector& InBlendParameters = FVector::ZeroVector, int32 InRootTransformSamplingRate = DefaultRootTransformSamplingRate, bool bPreProcessRootTransform = true);
+	POSESEARCH_API FAnimationAssetSampler(const UAnimationAsset* InAnimationAsset = nullptr, const FTransform& InRootTransformOrigin = FTransform::Identity, const FVector& InBlendParameters = FVector::ZeroVector, int32 InRootTransformSamplingRate = DefaultRootTransformSamplingRate, bool bPreProcessRootTransform = true, bool bInEnforceCompressedDataSampling = true);
+	POSESEARCH_API void Init(const UAnimationAsset* InAnimationAsset, const FTransform& InRootTransformOrigin = FTransform::Identity, const FVector& InBlendParameters = FVector::ZeroVector, int32 InRootTransformSamplingRate = DefaultRootTransformSamplingRate, bool bPreProcessRootTransform = true, bool bInEnforceCompressedDataSampling = true);
 
 	POSESEARCH_API bool IsInitialized() const;
 	POSESEARCH_API float GetPlayLength() const;
@@ -74,7 +74,8 @@ protected:
 	const float ExtrapolationSampleTime = 1.f / 30.f;
 	const float ExtractionInterval = 1.0f / 120.0f;
 
-#if WITH_EDITOR
+	bool bEnforceCompressedDataSampling = true;
+#if WITH_EDITOR	
 	FIoHash PlatformHash;
 #endif // WITH_EDITOR
 };
