Ensure AssetSampler c-tor forwards enforce compressed data sampling to init

diff --git a/Engine/Plugins/Animation/PoseSearch/Source/Runtime/Private/PoseSearchAssetSampler.cpp b/Engine/Plugins/Animation/PoseSearch/Source/Runtime/Private/PoseSearchAssetSampler.cpp
--- a/Engine/Plugins/Animation/PoseSearch/Source/Runtime/Private/PoseSearchAssetSampler.cpp
+++ b/Engine/Plugins/Animation/PoseSearch/Source/Runtime/Private/PoseSearchAssetSampler.cpp
@@ -240,7 +240,7 @@ static int32 GetHighestWeightSample(const TArray<struct FBlendSampleData>& Sampl
 // FAssetSamplerBase
 FAnimationAssetSampler::FAnimationAssetSampler(const UAnimationAsset* InAnimationAsset, const FTransform& InRootTransformOrigin, const FVector& InBlendParameters, int32 InRootTransformSamplingRate, bool bPreProcessRootTransform, bool bInEnforceCompressedDataSampling)
 {
-	Init(InAnimationAsset, InRootTransformOrigin, InBlendParameters, InRootTransformSamplingRate, bPreProcessRootTransform);
+	Init(InAnimationAsset, InRootTransformOrigin, InBlendParameters, InRootTransformSamplingRate, bPreProcessRootTransform, bInEnforceCompressedDataSampling);
 }
 
 void FAnimationAssetSampler::Init(const UAnimationAsset* InAnimationAsset, const FTransform& InRootTransformOrigin, const FVector& InBlendParameters, int32 InRootTransformSamplingRate, bool bPreProcessRootTransform, bool bInEnforceCompressedDataSampling)
