PoseSearch - temporarily removing overzealous logging

#rnx

diff --git a/Engine/Plugins/Animation/PoseSearch/Source/Runtime/Private/PoseSearchHistory.cpp b/Engine/Plugins/Animation/PoseSearch/Source/Runtime/Private/PoseSearchHistory.cpp
--- a/Engine/Plugins/Animation/PoseSearch/Source/Runtime/Private/PoseSearchHistory.cpp
+++ b/Engine/Plugins/Animation/PoseSearch/Source/Runtime/Private/PoseSearchHistory.cpp
@@ -789,14 +789,15 @@ bool FPoseHistory::GetTransformAtTime(float Time, FTransform& OutBoneTransform,
 		}
 	}
 
-#if WITH_EDITOR
-	// @todo: logging errors only if pose history was properly initialized for now (since AnimNext version doesn't implement initialization with ref pose yet),
-	// and we ultimately want to catch errors of missing bone transforms only checking valid USkeleton(s) against valid PoseSearchSchema(s)
-	if (!bSuccess && BoneIndexSkeleton)
-	{
-		UE_LOG(LogPoseSearch, Error, TEXT("FPoseHistory::GetTransformAtTime - Couldn't find root bone transform at time %f for Skeleton %s!"), Time, *BoneIndexSkeleton->GetName());
-	}
-#endif // WITH_EDITOR
+	// @todo: re-enable this logging for editor only builds once the original issue has been fixed
+//#if WITH_EDITOR
+//	// @todo: logging errors only if pose history was properly initialized for now (since AnimNext version doesn't implement initialization with ref pose yet),
+//	// and we ultimately want to catch errors of missing bone transforms only checking valid USkeleton(s) against valid PoseSearchSchema(s)
+//	if (!bSuccess && BoneIndexSkeleton)
+//	{
+//		UE_LOG(LogPoseSearch, Error, TEXT("FPoseHistory::GetTransformAtTime - Couldn't find root bone transform at time %f for Skeleton %s!"), Time, *BoneIndexSkeleton->GetName());
+//	}
+//#endif // WITH_EDITOR
 
 	return bSuccess;
 }
diff --git a/Engine/Plugins/Animation/PoseSearch/Source/Runtime/Public/PoseSearch/PoseSearchInteractionSubsystem.h b/Engine/Plugins/Animation/PoseSearch/Source/Runtime/Public/PoseSearch/PoseSearchInteractionSubsystem.h
--- a/Engine/Plugins/Animation/PoseSearch/Source/Runtime/Public/PoseSearch/PoseSearchInteractionSubsystem.h
+++ b/Engine/Plugins/Animation/PoseSearch/Source/Runtime/Public/PoseSearch/PoseSearchInteractionSubsystem.h
@@ -56,7 +56,8 @@ struct FAvailabilitiesMap : protected TMap<TWeakObjectPtr<const UObject>, FInter
 	
 	void AddAvailabilities(const TArrayView<const FPoseSearchInteractionAvailability> Availabilities, const UObject* AnimContext, FName PoseHistoryName, const IPoseHistory* PoseHistory)
 	{
-#if WITH_EDITOR && DO_CHECK
+		// @todo: re-enable this logging for editor only builds once the original issue has been fixed
+//#if WITH_EDITOR && DO_CHECK
 		check(AnimContext);
 		const AActor* ContextOwningActor = GetContextOwningActor(AnimContext);
 		check(ContextOwningActor);
@@ -71,13 +72,13 @@ struct FAvailabilitiesMap : protected TMap<TWeakObjectPtr<const UObject>, FInter
 					// It's not currently supported! and if not filtered out it could potentially add interactions within the same actor...
 					if (ContextOwningActor == OtherContextOwningActor && AnimContext != OtherAnimContext)
 					{
-						UE_LOG(LogPoseSearch, Error, TEXT("FAvailabilitiesMap::AddAvailabilities from different AnimContexts %s, %s sharing the same Actor %s. This is not currently supported!"), *AnimContext->GetName(), *OtherAnimContext->GetName(), *ContextOwningActor->GetName());
+						//UE_LOG(LogPoseSearch, Error, TEXT("FAvailabilitiesMap::AddAvailabilities from different AnimContexts %s, %s sharing the same Actor %s. This is not currently supported!"), *AnimContext->GetName(), *OtherAnimContext->GetName(), *ContextOwningActor->GetName());
 						return;
 					}
 				}
 			}
 		}
-#endif // WITH_EDITOR && DO_CHECK
+//#endif // WITH_EDITOR && DO_CHECK
 
 		FInteractionAvailabilities& AnimContextAvailabilities = FindOrAdd(AnimContext);
 		for (const FPoseSearchInteractionAvailability& Availability : Availabilities)
