Pose Search Chooser Column

diff --git a/Engine/Plugins/Animation/PoseSearch/Source/Runtime/Private/Chooser/PoseSearchChooserColumn.cpp b/Engine/Plugins/Animation/PoseSearch/Source/Runtime/Private/Chooser/PoseSearchChooserColumn.cpp
new file mode 100644
--- /dev/null
+++ b/Engine/Plugins/Animation/PoseSearch/Source/Runtime/Private/Chooser/PoseSearchChooserColumn.cpp
@@ -0,0 +1,153 @@
+// Copyright Epic Games, Inc. All Rights Reserved.
+#include "PoseSearchChooserColumn.h"
+#include "ChooserIndexArray.h"
+#include "ChooserPropertyAccess.h"
+#include "IChooserParameterBool.h"
+#include "IChooserParameterFloat.h"
+#include "ChooserTrace.h"
+#include "PoseSearch/PoseSearchLibrary.h"
+#include "PoseSearch/PoseSearchDatabase.h"
+#if WITH_EDITOR
+#include "StructUtils/PropertyBag.h"
+#endif
+
+
+bool FPoseHistoryContextProperty::GetValue(FChooserEvaluationContext& Context, FPoseHistoryReference& OutResult) const
+{
+	const FPoseHistoryReference* Reference;
+	if (Binding.GetValuePtr(Context, Reference))
+	{
+		OutResult = *Reference;
+		return true;
+	}
+	return false;
+}
+
+
+FPoseSearchColumn::FPoseSearchColumn()
+{
+}
+
+#if WITH_EDITOR
+void FPoseSearchColumn::AutoPopulate(int32 RowIndex, UObject* OutputObject)
+{
+	if (RowValues.IsValidIndex(RowIndex))
+	{
+		RowValues[RowIndex].ResultAsset = Cast<UAnimationAsset>(OutputObject);
+	}
+}
+
+bool FPoseSearchColumn::EditorTestFilter(int32 RowIndex) const
+{
+	// todo: not sure how to do editor display for pose search weights (might need to make a copy of the pose history struct)
+	return true;
+}
+
+
+float FPoseSearchColumn::EditorTestCost(int32 RowIndex) const
+{
+	// todo: not sure how to do editor display for pose search weights (might need to make a copy of the pose history struct)
+	return 0.0f;
+}
+
+#endif
+
+bool FPoseSearchColumn::HasFilters() const
+{
+	return true;
+}
+
+void FPoseSearchColumn::Filter(FChooserEvaluationContext& Context, const FChooserIndexArray& IndexListIn, FChooserIndexArray& IndexListOut) const
+{
+
+	TArray<const UObject*, TInlineAllocator<128>> AssetsToSearch;
+	for(uint32 Index : IndexListIn)
+	{
+		AssetsToSearch.Add(RowValues[Index].ResultAsset);
+	}
+
+
+	FPoseHistoryReference PoseHistoryReference;
+	if (const FChooserParameterPoseHistoryBase* PoseHistoryParameter = InputValue.GetPtr<FChooserParameterPoseHistoryBase>())
+	{
+		if (PoseHistoryParameter->GetValue(Context, PoseHistoryReference))
+		{
+			if (const UE::PoseSearch::IPoseHistory* PoseHistory = PoseHistoryReference.PoseHistory.Get())
+			{
+				FPoseSearchContinuingProperties ContinuingProperties;
+				FPoseSearchFutureProperties Future;
+				const UObject* ContextObject;
+				for(const FStructView& Param : Context.Params)
+				{
+					if (FChooserEvaluationInputObject* ObjectParam = Param.GetPtr<FChooserEvaluationInputObject>())
+					{
+						ContextObject = ObjectParam->Object;
+						break;
+					}
+				}
+
+				const UE::PoseSearch::FSearchResult SearchResult = UPoseSearchLibrary::MotionMatch(
+					MakeArrayView(&ContextObject, 1),
+					MakeArrayView(&UE::PoseSearch::DefaultRole, 1),
+					MakeArrayView(&PoseHistory, 1),
+					AssetsToSearch,
+					ContinuingProperties,
+					Future,
+					FPoseSearchEvent());
+
+				 if (const UE::PoseSearch::FSearchIndexAsset* SearchIndexAsset = SearchResult.GetSearchIndexAsset())
+				 {
+				 	const UPoseSearchDatabase* Database = SearchResult.Database.Get();
+				 	if (const FPoseSearchDatabaseAnimationAssetBase* DatabaseAnimationAssetBase = Database->GetDatabaseAnimationAsset<FPoseSearchDatabaseAnimationAssetBase>(*SearchIndexAsset))
+				 	{
+				 		const UObject* Result = DatabaseAnimationAssetBase->GetAnimationAsset();
+				 		int Index;
+				 		if (AssetsToSearch.Find(Result, Index))
+				 		{
+				 			IndexListOut.Push(IndexListIn[Index]);
+
+				 			if (const FChooserParameterFloatBase* StartTime = OutputStartTime.GetPtr<FChooserParameterFloatBase>())
+				 			{
+				 				StartTime->SetValue(Context, SearchResult.AssetTime);
+				 			}
+				 			
+				 			if (const FChooserParameterFloatBase* Cost = OutputCost.GetPtr<FChooserParameterFloatBase>())
+                            {
+								Cost->SetValue(Context, SearchResult.PoseCost);
+                            }
+
+				 			if (const FChooserParameterBoolBase* Mirror = OutputMirror.GetPtr<FChooserParameterBoolBase>())
+				 			{
+				 				Mirror->SetValue(Context, SearchIndexAsset->IsMirrored());
+				 			}
+				 		}
+				 	}
+				 }
+			}
+		}
+	}
+}
+
+#if WITH_EDITOR
+	void FPoseSearchColumn::AddToDetails(FInstancedPropertyBag& PropertyBag, int32 ColumnIndex, int32 RowIndex)
+	{
+		FText DisplayName = NSLOCTEXT("PoseSearchColumn","Pose Search", "Pose Search");
+		FName PropertyName("RowData", ColumnIndex);
+		FPropertyBagPropertyDesc PropertyDesc(PropertyName,  EPropertyBagPropertyType::Struct, FChooserPoseSearchRowData::StaticStruct());
+		PropertyDesc.MetaData.Add(FPropertyBagPropertyDescMetaData("DisplayName", DisplayName.ToString()));
+		PropertyBag.AddProperties({PropertyDesc});
+		
+		PropertyBag.SetValueStruct(PropertyName, RowValues[RowIndex]);
+	}
+
+	void FPoseSearchColumn::SetFromDetails(FInstancedPropertyBag& PropertyBag, int32 ColumnIndex, int32 RowIndex)
+	{
+		FName PropertyName("RowData", ColumnIndex);
+		
+   		TValueOrError<FStructView, EPropertyBagResult> Result = PropertyBag.GetValueStruct(PropertyName, FChooserPoseSearchRowData::StaticStruct());
+		if (FStructView* StructView = Result.TryGetValue())
+		{
+			RowValues[RowIndex] = StructView->Get<FChooserPoseSearchRowData>();
+		}
+	}
+#endif
\ No newline at end of file
diff --git a/Engine/Plugins/Animation/PoseSearch/Source/Runtime/Private/Chooser/PoseSearchChooserColumn.h b/Engine/Plugins/Animation/PoseSearch/Source/Runtime/Private/Chooser/PoseSearchChooserColumn.h
new file mode 100644
--- /dev/null
+++ b/Engine/Plugins/Animation/PoseSearch/Source/Runtime/Private/Chooser/PoseSearchChooserColumn.h
@@ -0,0 +1,99 @@
+// Copyright Epic Games, Inc. All Rights Reserved.
+#pragma once
+
+#include <IHasContext.h>
+
+#include "CoreMinimal.h"
+#include "IChooserColumn.h"
+#include "ChooserPropertyAccess.h"
+#include "PoseSearch/Chooser/ChooserParameterPoseHistoryBase.h"
+#include "StructUtils/InstancedStruct.h"
+#include "Serialization/MemoryReader.h"
+#include "Templates/SubclassOf.h"
+#include "PoseSearch/PoseSearchHistory.h"
+#include "PoseSearchChooserColumn.generated.h"
+
+USTRUCT(DisplayName = "Pose History Property Binding")
+struct POSESEARCH_API FPoseHistoryContextProperty :  public FChooserParameterPoseHistoryBase
+{
+	GENERATED_BODY()
+public:
+
+	UPROPERTY(EditAnywhere, Meta = (BindingType = "FPoseHistoryReference", BindingAllowFunctions = "false", BindingColor = "StructPinTypeColor"), Category = "Binding")
+	FChooserPropertyBinding Binding;
+	
+	virtual bool GetValue(FChooserEvaluationContext& Context, FPoseHistoryReference& OutResult) const override;
+	virtual bool IsBound() const override
+	{
+		return Binding.IsBoundToRoot || !Binding.PropertyBindingChain.IsEmpty();
+	}
+
+	CHOOSER_PARAMETER_BOILERPLATE();
+};
+
+
+USTRUCT()
+struct FChooserPoseSearchRowData
+{
+	GENERATED_BODY()
+
+	UPROPERTY(VisibleAnywhere, Category = "Data", Meta = (Tooltip = "The result animation asset for this row (filled by autopopulate from the result column)"));
+	TObjectPtr<UAnimationAsset> ResultAsset;
+};
+
+USTRUCT(DisplayName = "Pose Match", Meta = (Experimental, Category = "Experimental", Tooltip = "This column filters out all assets except the one which is selected by motion matching query.  Results must be AnimationAssets with a PoseSearchBranchIn notify state.  It also outputs OutputStartTime to specify the frame which matched pose best.  To work as intended it must be placed last (furthest right) in the Chooser so that other filters are applied first."))
+struct POSESEARCH_API FPoseSearchColumn : public FChooserColumnBase
+{
+	GENERATED_BODY()
+	public:
+	FPoseSearchColumn();
+	
+	// Pose History
+	UPROPERTY(EditAnywhere, DisplayName = "Pose History", NoClear, Meta = (ExcludeBaseStruct, BaseStruct = "/Script/PoseSearch.ChooserParameterPoseHistoryBase"), Category = "Data")
+	FInstancedStruct InputValue;
+	
+	// Float output for the start time with the best matching pose
+	UPROPERTY(EditAnywhere, NoClear, Meta = (ExcludeBaseStruct, BaseStruct = "/Script/Chooser.ChooserParameterFloatBase"), Category = "Data")
+	FInstancedStruct OutputStartTime;
+	
+	// Bool output for if asset should be mirrored
+	UPROPERTY(EditAnywhere, NoClear, Meta = (ExcludeBaseStruct, BaseStruct = "/Script/Chooser.ChooserParameterBoolBase"), Category = "Data")
+	FInstancedStruct OutputMirror;
+
+	// Float output for the cost of the selected pose
+	UPROPERTY(EditAnywhere, NoClear, Meta = (ExcludeBaseStruct, BaseStruct = "/Script/Chooser.ChooserParameterFloatBase"), Category = "Data")
+	FInstancedStruct OutputCost;
+
+#if WITH_EDITORONLY_DATA
+	UPROPERTY(EditAnywhere, Category=Data, meta=(ToolTip="DefaultRowValue will be assigned to cells when new rows are created"));
+	FChooserPoseSearchRowData DefaultRowValue;
+#endif
+	
+	UPROPERTY()
+	// array of results (cells for this column for each row in the table)
+	// should match the length of the Results array 
+	TArray<FChooserPoseSearchRowData> RowValues;
+
+	virtual bool HasFilters() const override;
+	virtual void Filter(FChooserEvaluationContext& Context, const FChooserIndexArray& IndexListIn, FChooserIndexArray& IndexListOut) const override;
+
+#if WITH_EDITOR
+	mutable double TestValue = 0.0;
+	virtual bool EditorTestFilter(int32 RowIndex) const override;
+	virtual float EditorTestCost(int32 RowIndex) const override;
+	
+	virtual void SetTestValue(TArrayView<const uint8> Value) override
+	{
+		FMemoryReaderView Reader(Value);
+		Reader << TestValue;
+	}
+
+	virtual void AddToDetails(FInstancedPropertyBag& PropertyBag, int32 ColumnIndex, int32 RowIndex) override;
+	virtual void SetFromDetails(FInstancedPropertyBag& PropertyBag, int32 ColumnIndex, int32 RowIndex) override;
+
+	virtual bool AutoPopulates() const override { return true; }
+	virtual void AutoPopulate(int32 RowIndex, UObject* OutputObject) override;
+#endif
+	
+	CHOOSER_COLUMN_BOILERPLATE(FChooserParameterPoseHistoryBase);
+};
\ No newline at end of file
diff --git a/Engine/Plugins/Animation/PoseSearch/Source/Runtime/Public/PoseSearch/Chooser/ChooserParameterPoseHistoryBase.h b/Engine/Plugins/Animation/PoseSearch/Source/Runtime/Public/PoseSearch/Chooser/ChooserParameterPoseHistoryBase.h
new file mode 100644
--- /dev/null
+++ b/Engine/Plugins/Animation/PoseSearch/Source/Runtime/Public/PoseSearch/Chooser/ChooserParameterPoseHistoryBase.h
@@ -0,0 +1,17 @@
+// Copyright Epic Games, Inc. All Rights Reserved.
+#pragma once
+
+#include <IHasContext.h>
+
+#include "CoreMinimal.h"
+#include "IChooserParameterBase.h"
+#include "PoseSearch/PoseSearchHistory.h"
+#include "ChooserParameterPoseHistoryBase.generated.h"
+
+USTRUCT()
+struct FChooserParameterPoseHistoryBase : public FChooserParameterBase
+{
+	GENERATED_BODY()
+	virtual bool GetValue(FChooserEvaluationContext& Context, FPoseHistoryReference& OutResult) const { return false; }
+	virtual bool IsBound() const { return false; }
+};
\ No newline at end of file
diff --git a/Engine/Plugins/Animation/PoseSearch/Source/Runtime/Public/PoseSearch/PoseSearchHistory.h b/Engine/Plugins/Animation/PoseSearch/Source/Runtime/Public/PoseSearch/PoseSearchHistory.h
--- a/Engine/Plugins/Animation/PoseSearch/Source/Runtime/Public/PoseSearch/PoseSearchHistory.h
+++ b/Engine/Plugins/Animation/PoseSearch/Source/Runtime/Public/PoseSearch/PoseSearchHistory.h
@@ -12,6 +12,8 @@
 #include "Animation/TrajectoryTypes.h"
 #include "UObject/ObjectKey.h"
 
+#include "PoseSearchHistory.generated.h"
+
 struct FAnimInstanceProxy;
 class USkeleton;
 class UWorld;
@@ -314,4 +316,10 @@ struct FPoseIndicesHistory
 
 } // namespace UE::PoseSearch
 
+USTRUCT(BlueprintType)
+struct FPoseHistoryReference
+{
+	GENERATED_BODY()
+	TSharedPtr<UE::PoseSearch::IPoseHistory> PoseHistory;
+};
 
diff --git a/Engine/Plugins/Experimental/AnimNextPoseSearch/Source/AnimNextPoseSearch/Private/HistoryCollectorTrait.cpp b/Engine/Plugins/Experimental/AnimNextPoseSearch/Source/AnimNextPoseSearch/Private/HistoryCollectorTrait.cpp
--- a/Engine/Plugins/Experimental/AnimNextPoseSearch/Source/AnimNextPoseSearch/Private/HistoryCollectorTrait.cpp
+++ b/Engine/Plugins/Experimental/AnimNextPoseSearch/Source/AnimNextPoseSearch/Private/HistoryCollectorTrait.cpp
@@ -6,7 +6,9 @@
 #include "AnimationRuntime.h"
 #include "EvaluationVM/EvaluationVM.h"
 #include "GenerationTools.h"
+#include "Module/AnimNextModuleInstance.h"
 #include "PoseSearch/PoseSearchHistoryAttribute.h"
+#include "TraitCore/NodeInstance.h"
 
 namespace UE::AnimNext
 {
@@ -222,6 +224,14 @@ namespace UE::AnimNext
 		}
 #endif // WITH_EDITOR
 
+		if (!SharedData->PoseHistoryReferenceVariable.IsNone())
+		{
+			FAnimNextModuleInstance* ModuleInstance = Binding.GetTraitPtr().GetNodeInstance()->GetOwner().GetModuleInstance();
+			FPoseHistoryReference Reference;
+			Reference.PoseHistory = InstanceData->PoseHistoryPtr;
+			ModuleInstance->SetVariable(SharedData->PoseHistoryReferenceVariable, Reference);
+		}
+		
 		check(InstanceData->PoseHistoryPtr.IsValid());
 		FAnimNextHistoryCollectorTask Task;
 		Task.PoseHistory = InstanceData->PoseHistoryPtr.Get();
diff --git a/Engine/Plugins/Experimental/AnimNextPoseSearch/Source/AnimNextPoseSearch/Private/HistoryCollectorTrait.h b/Engine/Plugins/Experimental/AnimNextPoseSearch/Source/AnimNextPoseSearch/Private/HistoryCollectorTrait.h
--- a/Engine/Plugins/Experimental/AnimNextPoseSearch/Source/AnimNextPoseSearch/Private/HistoryCollectorTrait.h
+++ b/Engine/Plugins/Experimental/AnimNextPoseSearch/Source/AnimNextPoseSearch/Private/HistoryCollectorTrait.h
@@ -83,6 +83,10 @@ struct FAnimNextHistoryCollectorTraitSharedData : public FAnimNextTraitSharedDat
 	UPROPERTY(EditAnywhere, Category = Experimental, meta = (ClampMin="0.001", EditCondition="bGenerateTrajectory", EditConditionHides))
 	float PredictionSamplingInterval = 0.4f;
 
+	// Property to store a reference to the PoseHistory struct in
+	UPROPERTY(EditAnywhere, Category = Experimental)
+	FName PoseHistoryReferenceVariable;
+
 	// if bGenerateTrajectory is true, TrajectoryData contains the tuning parameters to generate the trajectory
 	//UPROPERTY(EditAnywhere, Category = Experimental, meta=(EditCondition="bGenerateTrajectory", EditConditionHides))
 	//FPoseSearchTrajectoryData TrajectoryData;
diff --git a/Engine/Plugins/Experimental/AnimNextPoseSearch/Source/AnimNextPoseSearch/Private/Module.cpp b/Engine/Plugins/Experimental/AnimNextPoseSearch/Source/AnimNextPoseSearch/Private/Module.cpp
--- a/Engine/Plugins/Experimental/AnimNextPoseSearch/Source/AnimNextPoseSearch/Private/Module.cpp
+++ b/Engine/Plugins/Experimental/AnimNextPoseSearch/Source/AnimNextPoseSearch/Private/Module.cpp
@@ -39,6 +39,7 @@ public:
 			FAlphaBlendArgs::StaticStruct(),
 			FPoseSearchInteractionAssetItem::StaticStruct(),
 			FPoseSearchInteractionAvailability::StaticStruct(),
+			FPoseHistoryReference::StaticStruct(),
 		};
 
 		FRigVMRegistry& RigVMRegistry = FRigVMRegistry::Get();
diff --git a/Engine/Plugins/Experimental/AnimNextPoseSearch/Source/AnimNextPoseSearch/Private/PoseHistoryChooserParameter.cpp b/Engine/Plugins/Experimental/AnimNextPoseSearch/Source/AnimNextPoseSearch/Private/PoseHistoryChooserParameter.cpp
new file mode 100644
--- /dev/null
+++ b/Engine/Plugins/Experimental/AnimNextPoseSearch/Source/AnimNextPoseSearch/Private/PoseHistoryChooserParameter.cpp
@@ -0,0 +1,33 @@
+// Copyright Epic Games, Inc. All Rights Reserved.
+
+#include "PoseHistoryChooserParameter.h"
+#include "DataInterface/AnimNextDataInterfaceInstance.h"
+#include "StructUtils/PropertyBag.h"
+
+namespace UE::AnimNext::Private
+{
+
+// Helper function to find the first data interface instance in the context 
+FAnimNextDataInterfaceInstance* GetFirstDataInterfaceInstance(FChooserEvaluationContext& Context)
+{
+	for(const FStructView& Param : Context.Params)
+	{
+		if(Param.GetScriptStruct() == FAnimNextDataInterfaceInstance::StaticStruct())
+		{
+			return Param.GetPtr<FAnimNextDataInterfaceInstance>();
+		}
+	}
+
+	return nullptr;
+}
+
+}
+
+bool FPoseHistoryAnimProperty::GetValue(FChooserEvaluationContext& Context, FPoseHistoryReference& OutResult) const
+{
+	if(FAnimNextDataInterfaceInstance* Instance = UE::AnimNext::Private::GetFirstDataInterfaceInstance(Context))
+	{
+		return Instance->GetVariable(VariableName, OutResult) == EPropertyBagResult::Success;
+	}
+	return false;
+}
diff --git a/Engine/Plugins/Experimental/AnimNextPoseSearch/Source/AnimNextPoseSearch/Private/PoseHistoryChooserParameter.h b/Engine/Plugins/Experimental/AnimNextPoseSearch/Source/AnimNextPoseSearch/Private/PoseHistoryChooserParameter.h
new file mode 100644
--- /dev/null
+++ b/Engine/Plugins/Experimental/AnimNextPoseSearch/Source/AnimNextPoseSearch/Private/PoseHistoryChooserParameter.h
@@ -0,0 +1,22 @@
+// Copyright Epic Games, Inc. All Rights Reserved.
+#pragma once
+
+#include "CoreMinimal.h"
+#include "IChooserParameterBool.h"
+#include "IChooserParameterFloat.h"
+#include "IChooserParameterEnum.h"
+#include "PoseSearch/Chooser/ChooserParameterPoseHistoryBase.h"
+#include "PoseHistoryChooserParameter.generated.h"
+
+USTRUCT(DisplayName = "Pose History Anim Param")
+struct FPoseHistoryAnimProperty :  public FChooserParameterPoseHistoryBase
+{
+	GENERATED_BODY()
+public:
+
+	UPROPERTY(DisplayName = "Variable", EditAnywhere, Category="Variable")
+	FName VariableName;
+
+	bool GetValue(FChooserEvaluationContext& Context, FPoseHistoryReference& OutResult) const override;
+	virtual bool IsBound() const { return !VariableName.IsNone(); }
+};
\ No newline at end of file
