Fix for deprecation warning in PoseSearchChooserColumn.cpp

diff --git a/Engine/Plugins/Animation/PoseSearch/Source/Runtime/Private/Chooser/PoseSearchChooserColumn.cpp b/Engine/Plugins/Animation/PoseSearch/Source/Runtime/Private/Chooser/PoseSearchChooserColumn.cpp
--- a/Engine/Plugins/Animation/PoseSearch/Source/Runtime/Private/Chooser/PoseSearchChooserColumn.cpp
+++ b/Engine/Plugins/Animation/PoseSearch/Source/Runtime/Private/Chooser/PoseSearchChooserColumn.cpp
@@ -61,9 +61,9 @@ void FPoseSearchColumn::Filter(FChooserEvaluationContext& Context, const FChoose
 {
 
 	TArray<const UObject*, TInlineAllocator<128>> AssetsToSearch;
-	for(uint32 Index : IndexListIn)
+	for (const FChooserIndexArray::FIndexData& IndexData : IndexListIn)
 	{
-		AssetsToSearch.Add(RowValues[Index].ResultAsset);
+		AssetsToSearch.Add(RowValues[IndexData.Index].ResultAsset);
 	}
 
 
