Add max distance filtering to pose search distance channel

diff --git a/Engine/Plugins/Animation/PoseSearch/Source/Runtime/Private/PoseSearchFeatureChannel_Distance.cpp b/Engine/Plugins/Animation/PoseSearch/Source/Runtime/Private/PoseSearchFeatureChannel_Distance.cpp
--- a/Engine/Plugins/Animation/PoseSearch/Source/Runtime/Private/PoseSearchFeatureChannel_Distance.cpp
+++ b/Engine/Plugins/Animation/PoseSearch/Source/Runtime/Private/PoseSearchFeatureChannel_Distance.cpp
@@ -53,6 +53,17 @@ void UPoseSearchFeatureChannel_Distance::AddDependentChannels(UPoseSearchSchema*
 	}
 }
 
+bool UPoseSearchFeatureChannel_Distance::IsFilterValid(TConstArrayView<float> PoseValues, TConstArrayView<float> QueryValues, int32 PoseIdx, const UE::PoseSearch::FPoseMetadata& Metadata) const
+{
+	using namespace UE::PoseSearch;
+	const float Pose = FFeatureVectorHelper::DecodeFloat(PoseValues, ChannelDataOffset);
+	const float Query = FFeatureVectorHelper::DecodeFloat(QueryValues, ChannelDataOffset);
+
+	check(MaxDistance > 0.f);
+
+	return FMath::Abs(Pose - Query) <= MaxDistance;
+}
+
 void UPoseSearchFeatureChannel_Distance::BuildQuery(UE::PoseSearch::FSearchContext& SearchContext) const
 {
 	using namespace UE::PoseSearch;
diff --git a/Engine/Plugins/Animation/PoseSearch/Source/Runtime/Public/PoseSearch/PoseSearchFeatureChannel_Distance.h b/Engine/Plugins/Animation/PoseSearch/Source/Runtime/Public/PoseSearch/PoseSearchFeatureChannel_Distance.h
--- a/Engine/Plugins/Animation/PoseSearch/Source/Runtime/Public/PoseSearch/PoseSearchFeatureChannel_Distance.h
+++ b/Engine/Plugins/Animation/PoseSearch/Source/Runtime/Public/PoseSearch/PoseSearchFeatureChannel_Distance.h
@@ -30,6 +30,13 @@ public:
 	UPROPERTY(EditAnywhere, Category = "Settings")
 	FName OriginRole = UE::PoseSearch::DefaultRole;
 
+	// Experimental, this feature might be removed without warning, not for production use
+   	// During pose selection if the distance between query versus candidate poses for this 
+   	// distance channel is greater than MaxDistance the candidate will be discarded.
+   	// the filtering will be enabled only for MaxDistance > 0
+   	UPROPERTY(EditAnywhere, Category = "Experimental|Filter")
+   	float MaxDistance = 0.f;
+	
 #if WITH_EDITORONLY_DATA
 	UPROPERTY(EditAnywhere, Category = "Settings")
 	float Weight = 1.f;
@@ -91,6 +98,12 @@ public:
 
 	virtual EPermutationTimeType GetPermutationTimeType() const override { return PermutationTimeType; }
 	virtual void AddDependentChannels(UPoseSearchSchema* Schema) const override;
+	
+	// IPoseSearchFilter interface
+	// Experimental, this feature might be removed without warning, not for production use
+	virtual bool IsFilterActive() const override { return MaxDistance > 0.f; }
+	// Experimental, this feature might be removed without warning, not for production use
+	virtual bool IsFilterValid(TConstArrayView<float> PoseValues, TConstArrayView<float> QueryValues, int32 PoseIdx, const UE::PoseSearch::FPoseMetadata& Metadata) const override;
 
 #if ENABLE_DRAW_DEBUG
 	virtual void DebugDraw(const UE::PoseSearch::FDebugDrawParams& DrawParams, TConstArrayView<float> PoseVector) const override;
