fix pose search database updates from notify states

diff --git a/Engine/Plugins/Animation/PoseSearch/Source/Runtime/Private/PoseSearchDatabase.cpp b/Engine/Plugins/Animation/PoseSearch/Source/Runtime/Private/PoseSearchDatabase.cpp
--- a/Engine/Plugins/Animation/PoseSearch/Source/Runtime/Private/PoseSearchDatabase.cpp
+++ b/Engine/Plugins/Animation/PoseSearch/Source/Runtime/Private/PoseSearchDatabase.cpp
@@ -27,6 +27,7 @@
 #include "Editor/EditorEngine.h"
 #endif //WITH_EDITOR && WITH_ENGINE
 
+struct FPoseSearchDatabaseAnimationAssetBase;
 DECLARE_STATS_GROUP(TEXT("PoseSearch"), STATGROUP_PoseSearch, STATCAT_Advanced);
 DECLARE_CYCLE_STAT_EXTERN(TEXT("Search Brute Force"), STAT_PoseSearch_BruteForce, STATGROUP_PoseSearch, );
 DECLARE_CYCLE_STAT_EXTERN(TEXT("Search PCA/KNN"), STAT_PoseSearch_PCAKNN, STATGROUP_PoseSearch, );
@@ -397,6 +398,17 @@ FTransform FPoseSearchDatabaseAnimationAssetBase::GetRootTransformOriginForRole(
 }
 
 #if WITH_EDITORONLY_DATA
+
+bool FPoseSearchDatabaseAnimationAssetBase::UpdateFrom(const FPoseSearchDatabaseAnimationAssetBase& Source)
+{
+	if (BranchInId != 0 && BranchInId == Source.BranchInId)
+	{
+		SetSamplingRange(Source.GetSamplingRange());
+		return true;
+	}
+	return false;
+}
+
 int64 FPoseSearchDatabaseAnimationAssetBase::GetEditorMemSize() const
 {
 	FArchiveCountMem EditorMemCount(GetAnimationAsset());
@@ -1004,14 +1016,10 @@ void UPoseSearchDatabase::SynchronizeWithExternalDependencies(TConstArrayView<UA
 	TArray<FInstancedStruct> NewAnimationAssets;
 
 	// collecting all the database AnimationAsset(s) that don't require synchronization
-	TArray<bool> DisableReselection;
-	DisableReselection.Reserve(GetNumAnimationAssets());
-
+	
 	for (FInstancedStruct& AnimationAsset : AnimationAssets)
 	{
 		FPoseSearchDatabaseAnimationAssetBase& AnimationAssetBase = AnimationAsset.GetMutable<FPoseSearchDatabaseAnimationAssetBase>();
-		DisableReselection.Add(AnimationAssetBase.bDisableReselection);
-		AnimationAssetBase.bDisableReselection = false;
 
 		const bool bRequiresSynchronization = AnimationAssetBase.IsSynchronizedWithExternalDependency() && SequencesBase.Contains(AnimationAssetBase.GetAnimationAsset());
 		if (!bRequiresSynchronization)
@@ -1075,11 +1083,22 @@ void UPoseSearchDatabase::SynchronizeWithExternalDependencies(TConstArrayView<UA
 	bool bModified = false;
 	for (int32 AnimationAssetIndex = GetNumAnimationAssets() - 1; AnimationAssetIndex >= 0; --AnimationAssetIndex)
 	{
-		const int32 FoundIndex = NewAnimationAssets.Find(AnimationAssets[AnimationAssetIndex]);
+		FPoseSearchDatabaseAnimationAssetBase& AnimationAsset = AnimationAssets[AnimationAssetIndex].GetMutable<FPoseSearchDatabaseAnimationAssetBase>();
+		
+		int32 FoundIndex = -1;
+		for(int i=0; i < NewAnimationAssets.Num(); i++)
+		{
+			const FPoseSearchDatabaseAnimationAssetBase& NewAnimationAsset = NewAnimationAssets[i].Get<FPoseSearchDatabaseAnimationAssetBase>();
+			if (AnimationAsset.UpdateFrom(NewAnimationAsset))
+			{
+				FoundIndex = i;
+				break;
+			}
+			
+		}
+		
 		if (FoundIndex >= 0)
 		{
-			FPoseSearchDatabaseAnimationAssetBase& AnimationAssetBase = AnimationAssets[AnimationAssetIndex].GetMutable<FPoseSearchDatabaseAnimationAssetBase>();
-			AnimationAssetBase.bDisableReselection = DisableReselection[AnimationAssetIndex];
 			NewAnimationAssets.RemoveAt(FoundIndex); 
 		}
 		else
diff --git a/Engine/Plugins/Animation/PoseSearch/Source/Runtime/Public/PoseSearch/PoseSearchDatabase.h b/Engine/Plugins/Animation/PoseSearch/Source/Runtime/Public/PoseSearch/PoseSearchDatabase.h
--- a/Engine/Plugins/Animation/PoseSearch/Source/Runtime/Public/PoseSearch/PoseSearchDatabase.h
+++ b/Engine/Plugins/Animation/PoseSearch/Source/Runtime/Public/PoseSearch/PoseSearchDatabase.h
@@ -68,7 +68,6 @@ struct POSESEARCH_API FPoseSearchDatabaseAnimationAssetBase
 	virtual UE::PoseSearch::FRole GetRole(int32 RoleIndex) const { return UE::PoseSearch::DefaultRole; }
 	virtual UAnimationAsset* GetAnimationAssetForRole(const UE::PoseSearch::FRole& Role) const;
 	virtual FTransform GetRootTransformOriginForRole(const UE::PoseSearch::FRole& Role) const;
-
 #if WITH_EDITOR
 	virtual int32 GetFrameAtTime(float Time) const;
 	virtual bool IsSkeletonCompatible(TObjectPtr<const UPoseSearchSchema> InSchema) const;
@@ -78,6 +77,7 @@ struct POSESEARCH_API FPoseSearchDatabaseAnimationAssetBase
 #endif // WITH_EDITOR
 
 #if WITH_EDITORONLY_DATA
+    virtual bool UpdateFrom(const FPoseSearchDatabaseAnimationAssetBase& Source);
 	virtual bool IsDisableReselection() const { return bDisableReselection; }
 	virtual void SetDisableReselection(bool bValue) { bDisableReselection = bValue; }
 	virtual UClass* GetAnimationAssetStaticClass() const { return nullptr; }
@@ -90,6 +90,7 @@ struct POSESEARCH_API FPoseSearchDatabaseAnimationAssetBase
 
 	// [0, 0] represents the entire frame range of the original animation.
 	virtual FFloatInterval GetSamplingRange() const { return FFloatInterval(0.f, 0.f); }
+	virtual void SetSamplingRange(const FFloatInterval& NewRange) { }
 	virtual FFloatInterval GetEffectiveSamplingRange(const FVector& BlendParameters) const;
 	
 	UE_DEPRECATED(5.6, "Use other signatures")
@@ -161,6 +162,7 @@ struct POSESEARCH_API FPoseSearchDatabaseSequence : public FPoseSearchDatabaseAn
 	virtual bool IsLooping() const override;
 	virtual bool IsRootMotionEnabled() const override;
 	virtual FFloatInterval GetSamplingRange() const override { return SamplingRange; }
+	virtual void SetSamplingRange(const FFloatInterval& NewRange) override { SamplingRange = NewRange; }
 #endif // WITH_EDITORONLY_DATA
 	
 	virtual UObject* GetAnimationAsset() const override;
@@ -229,6 +231,7 @@ struct POSESEARCH_API FPoseSearchDatabaseBlendSpace : public FPoseSearchDatabase
 	virtual FFloatInterval GetSamplingRange() const override { return SamplingRange; }
 	virtual FFloatInterval GetEffectiveSamplingRange(const FVector& BlendParameters) const override;
 	virtual void IterateOverSamplingParameter(const TFunction<void(const FVector& BlendParameters)>& ProcessSamplingParameter) const override;
+	virtual void SetSamplingRange(const FFloatInterval& NewRange) override { SamplingRange = NewRange; }
 
 	UE_DEPRECATED(5.6, "Use IterateOverSamplingParameter to iterate over all the sampling paramters for the blend space")
 	void GetBlendSpaceParameterSampleRanges(int32& HorizontalBlendNum, int32& VerticalBlendNum) const;
@@ -287,6 +290,7 @@ struct POSESEARCH_API FPoseSearchDatabaseAnimComposite : public FPoseSearchDatab
 	virtual bool IsLooping() const override;
 	virtual bool IsRootMotionEnabled() const override;
 	virtual FFloatInterval GetSamplingRange() const override { return SamplingRange; }
+	virtual void SetSamplingRange(const FFloatInterval& NewRange) override { SamplingRange = NewRange; }
 #endif // WITH_EDITORONLY_DATA
 
 	virtual UObject* GetAnimationAsset() const override;
@@ -330,6 +334,7 @@ struct POSESEARCH_API FPoseSearchDatabaseAnimMontage : public FPoseSearchDatabas
 	virtual bool IsLooping() const override;
 	virtual bool IsRootMotionEnabled() const override;
 	virtual FFloatInterval GetSamplingRange() const override { return SamplingRange; }
+	virtual void SetSamplingRange(const FFloatInterval& NewRange) override { SamplingRange = NewRange; }
 #endif // WITH_EDITORONLY_DATA
 
 	virtual UObject* GetAnimationAsset() const override;
@@ -379,6 +384,7 @@ struct POSESEARCH_API FPoseSearchDatabaseMultiAnimAsset : public FPoseSearchData
 	virtual bool IsRootMotionEnabled() const override;
 	virtual FFloatInterval GetSamplingRange() const override { return SamplingRange; }
 	virtual void IterateOverSamplingParameter(const TFunction<void(const FVector& BlendParameters)>& ProcessSamplingParameter) const override;
+	virtual void SetSamplingRange(const FFloatInterval& NewRange) override { SamplingRange = NewRange; }
 #endif // WITH_EDITORONLY_DATA
 
 #if WITH_EDITOR
