PoseSearch - fix for motion matching interaction asset editor timeline

diff --git a/Engine/Plugins/Animation/PoseSearch/Source/Editor/Private/PoseSearchInteractionAssetEditor.cpp b/Engine/Plugins/Animation/PoseSearch/Source/Editor/Private/PoseSearchInteractionAssetEditor.cpp
--- a/Engine/Plugins/Animation/PoseSearch/Source/Editor/Private/PoseSearchInteractionAssetEditor.cpp
+++ b/Engine/Plugins/Animation/PoseSearch/Source/Editor/Private/PoseSearchInteractionAssetEditor.cpp
@@ -98,7 +98,7 @@ void FInteractionAssetPreviewActor::UpdatePreviewActor(const UPoseSearchInteract
 		bool bPlayTimeUpdated = false;
 
 		float NewCurrentTime = 0.f;
-		FAnimationRuntime::AdvanceTime(false, PlayTime, NewCurrentTime, PreviewAsset->GetPlayLength());
+		FAnimationRuntime::AdvanceTime(false, PlayTime, NewCurrentTime, Samplers[PreviewActor].GetPlayLength());
 
 		if (!FMath::IsNearlyEqual(CurrentTime, NewCurrentTime))
 		{
