PoseSearch - MM node cleanup

#ForUE5

diff --git a/Engine/Plugins/Animation/PoseSearch/Source/Runtime/Private/AnimNode_MotionMatching.cpp b/Engine/Plugins/Animation/PoseSearch/Source/Runtime/Private/AnimNode_MotionMatching.cpp
--- a/Engine/Plugins/Animation/PoseSearch/Source/Runtime/Private/AnimNode_MotionMatching.cpp
+++ b/Engine/Plugins/Animation/PoseSearch/Source/Runtime/Private/AnimNode_MotionMatching.cpp
@@ -70,28 +70,6 @@ void FAnimNode_MotionMatching::Initialize_AnyThread(const FAnimationInitializeCo
 	MotionMatchingState.Reset();
 }
 
-void FAnimNode_MotionMatching::Evaluate_AnyThread(FPoseContext& Output)
-{
-	DECLARE_SCOPE_HIERARCHICAL_COUNTER_ANIMNODE(Evaluate_AnyThread);
-	ANIM_MT_SCOPE_CYCLE_COUNTER_VERBOSE(MotionMatching, !IsInGameThread());
-
-	using namespace UE::PoseSearch;
-
-	FAnimNode_BlendStack_Standalone::Evaluate_AnyThread(Output);
-
-	const UE::Anim::IAnimRootMotionProvider* RootMotionProvider = UE::Anim::IAnimRootMotionProvider::Get();
-	FTransform RootMotionTransformDelta;
-	if (RootMotionProvider && RootMotionProvider->HasRootMotion(Output.CustomAttributes))
-	{
-		RootMotionProvider->ExtractRootMotion(Output.CustomAttributes, RootMotionTransformDelta);
-	}
-	else
-	{
-		RootMotionTransformDelta = FTransform::Identity;
-		RootMotionProvider = nullptr;
-	}
-}
-
 void FAnimNode_MotionMatching::UpdateAssetPlayer(const FAnimationUpdateContext& Context)
 {
 	DECLARE_SCOPE_HIERARCHICAL_COUNTER_ANIMNODE(UpdateAssetPlayer);
diff --git a/Engine/Plugins/Animation/PoseSearch/Source/Runtime/Public/PoseSearch/AnimNode_MotionMatching.h b/Engine/Plugins/Animation/PoseSearch/Source/Runtime/Public/PoseSearch/AnimNode_MotionMatching.h
--- a/Engine/Plugins/Animation/PoseSearch/Source/Runtime/Public/PoseSearch/AnimNode_MotionMatching.h
+++ b/Engine/Plugins/Animation/PoseSearch/Source/Runtime/Public/PoseSearch/AnimNode_MotionMatching.h
@@ -38,7 +38,6 @@ public:
 	// FAnimNode_Base interface
 	// @todo: implement CacheBones_AnyThread to rebind the schema bones
 	virtual void Initialize_AnyThread(const FAnimationInitializeContext& Context) override;
-	virtual void Evaluate_AnyThread(FPoseContext& Output) override;
 	virtual bool GetIgnoreForRelevancyTest() const override;
 	virtual bool SetIgnoreForRelevancyTest(bool bInIgnoreForRelevancyTest) override;
 	virtual FName GetGroupName() const override;
