Pose Search Column Editor

diff --git a/Engine/Plugins/Animation/PoseSearch/Source/Editor/PoseSearchEditor.Build.cs b/Engine/Plugins/Animation/PoseSearch/Source/Editor/PoseSearchEditor.Build.cs
--- a/Engine/Plugins/Animation/PoseSearch/Source/Editor/PoseSearchEditor.Build.cs
+++ b/Engine/Plugins/Animation/PoseSearch/Source/Editor/PoseSearchEditor.Build.cs
@@ -36,7 +36,8 @@ public class PoseSearchEditor : ModuleRules
 				"SlateCore",
 				"RewindDebuggerInterface",
 				"UnrealEd",
-				"InputCore"
+				"InputCore",
+				"ChooserEditor"
 			}
 		);
 
diff --git a/Engine/Plugins/Animation/PoseSearch/Source/Editor/Private/Chooser/PoseSearchColumnEditor.cpp b/Engine/Plugins/Animation/PoseSearch/Source/Editor/Private/Chooser/PoseSearchColumnEditor.cpp
new file mode 100644
--- /dev/null
+++ b/Engine/Plugins/Animation/PoseSearch/Source/Editor/Private/Chooser/PoseSearchColumnEditor.cpp
@@ -0,0 +1,64 @@
+// Copyright Epic Games, Inc. All Rights Reserved.
+
+#include "PoseSearchColumnEditor.h"
+
+#include <ChooserColumnHeader.h>
+
+#include "PoseSearch/Chooser/PoseSearchChooserColumn.h"
+#include "SPropertyAccessChainWidget.h"
+#include "ObjectChooserWidgetFactories.h"
+#include "Widgets/SBoxPanel.h"
+#include "Widgets/Layout/SBox.h"
+#include "Widgets/Layout/SSpacer.h"
+#include "Widgets/Input/SNumericEntryBox.h"
+#include "ScopedTransaction.h"
+
+#define LOCTEXT_NAMESPACE "PoseSearchColumnEditor"
+
+namespace UE::PoseSearchEditor
+{
+
+TSharedRef<SWidget> CreatePoseSearchColumnWidget(UChooserTable* Chooser, FChooserColumnBase* Column, int Row)
+{
+	FPoseSearchColumn* PoseSearchColumn = static_cast<FPoseSearchColumn*>(Column);
+	
+	if (Row == UE::ChooserEditor::ColumnWidget_SpecialIndex_Fallback)
+	{
+		return SNullWidget::NullWidget;
+	}
+	if (Row == UE::ChooserEditor::ColumnWidget_SpecialIndex_Header)
+	{
+		// create column header widget
+		const FSlateBrush* ColumnIcon = FAppStyle::Get().GetBrush("Icons.Search");
+		const FText ColumnTooltip = LOCTEXT("Pose Match Tooltip", "Pose Match: Selects a single result based on the animation with the best matching pose, and outputs the StartTime for the frame with that pose. Animation Assets must contain \"Pose Match Branch In\" Notify State. AutoPopulate will fill in Column data with result Animation Assets.");
+		const FText ColumnName = LOCTEXT("Pose Match","Pose Match");
+        		
+		TSharedPtr<SWidget> DebugWidget = nullptr;
+
+		return UE::ChooserEditor::MakeColumnHeaderWidget(Chooser, Column, ColumnName, ColumnTooltip, ColumnIcon, DebugWidget);
+	}
+
+	// create cell widget
+	return
+		SNew(STextBlock)
+			.Text_Lambda([Row, PoseSearchColumn]()
+			{
+				if (PoseSearchColumn->RowValues.IsValidIndex(Row))
+				{
+					if (UObject* ResultAsset = PoseSearchColumn->RowValues[Row].ResultAsset)
+					{
+						return FText::FromString(ResultAsset->GetName());
+					}
+				}
+				return FText();
+			});
+}
+	
+void RegisterPoseSearchChooserWidgets()
+{
+	UE::ChooserEditor::FObjectChooserWidgetFactories::RegisterColumnWidgetCreator(FPoseSearchColumn::StaticStruct(), CreatePoseSearchColumnWidget);
+}
+	
+}
+
+#undef LOCTEXT_NAMESPACE
diff --git a/Engine/Plugins/Animation/PoseSearch/Source/Editor/Private/PoseSearchEditor.cpp b/Engine/Plugins/Animation/PoseSearch/Source/Editor/Private/PoseSearchEditor.cpp
--- a/Engine/Plugins/Animation/PoseSearch/Source/Editor/Private/PoseSearchEditor.cpp
+++ b/Engine/Plugins/Animation/PoseSearch/Source/Editor/Private/PoseSearchEditor.cpp
@@ -19,6 +19,7 @@
 
 #include "Trace/PoseSearchTraceAnalyzer.h"
 #include "Trace/PoseSearchTraceModule.h"
+#include "PoseSearchColumnEditor.h"
 
 DEFINE_LOG_CATEGORY(LogPoseSearchEditor);
 
@@ -76,6 +77,8 @@ void FEditorModule::StartupModule()
 		// Register Ed Modes used by PoseSearchDatabase and PoseSearchInteractionAsset
 		FEditorModeRegistry::Get().RegisterMode<FDatabaseEdMode>(FDatabaseEdMode::EdModeId, LOCTEXT("PoseSearchDatabaseEdModeName", "PoseSearchDatabase"));
 		FEditorModeRegistry::Get().RegisterMode<FInteractionAssetEdMode>(FInteractionAssetEdMode::EdModeId, LOCTEXT("FPoseSearchInteractionAssetEdModeName", "PoseSearchInteractionAsset"));
+
+		UE::PoseSearchEditor::RegisterPoseSearchChooserWidgets();
 	}
 
 	RegisterPropertyTypeCustomizations();
diff --git a/Engine/Plugins/Animation/PoseSearch/Source/Editor/Public/PoseSearchColumnEditor.h b/Engine/Plugins/Animation/PoseSearch/Source/Editor/Public/PoseSearchColumnEditor.h
new file mode 100644
--- /dev/null
+++ b/Engine/Plugins/Animation/PoseSearch/Source/Editor/Public/PoseSearchColumnEditor.h
@@ -0,0 +1,10 @@
+// Copyright Epic Games, Inc. All Rights Reserved.
+
+#pragma once
+
+#include "CoreMinimal.h"
+
+namespace UE::PoseSearchEditor
+{
+	void RegisterPoseSearchChooserWidgets();
+}
\ No newline at end of file
diff --git a/Engine/Plugins/Animation/PoseSearch/Source/Runtime/Private/Chooser/PoseSearchChooserColumn.cpp b/Engine/Plugins/Animation/PoseSearch/Source/Runtime/Private/Chooser/PoseSearchChooserColumn.cpp
--- a/Engine/Plugins/Animation/PoseSearch/Source/Runtime/Private/Chooser/PoseSearchChooserColumn.cpp
+++ b/Engine/Plugins/Animation/PoseSearch/Source/Runtime/Private/Chooser/PoseSearchChooserColumn.cpp
@@ -1,5 +1,5 @@
 // Copyright Epic Games, Inc. All Rights Reserved.
-#include "PoseSearchChooserColumn.h"
+#include "PoseSearch/Chooser/PoseSearchChooserColumn.h"
 #include "ChooserIndexArray.h"
 #include "ChooserPropertyAccess.h"
 #include "IChooserParameterBool.h"
diff --git a/Engine/Plugins/Animation/PoseSearch/Source/Runtime/Private/Chooser/PoseSearchChooserColumn.h b/Engine/Plugins/Animation/PoseSearch/Source/Runtime/Private/Chooser/PoseSearchChooserColumn.h
deleted file mode 100644
diff --git a/Engine/Plugins/Animation/PoseSearch/Source/Runtime/Public/PoseSearch/Chooser/PoseSearchChooserColumn.h b/Engine/Plugins/Animation/PoseSearch/Source/Runtime/Public/PoseSearch/Chooser/PoseSearchChooserColumn.h
new file mode 100644
--- /dev/null
+++ b/Engine/Plugins/Animation/PoseSearch/Source/Runtime/Public/PoseSearch/Chooser/PoseSearchChooserColumn.h
@@ -0,0 +1,99 @@
+// Copyright Epic Games, Inc. All Rights Reserved.
+#pragma once
+
+#include <IHasContext.h>
+
+#include "CoreMinimal.h"
+#include "IChooserColumn.h"
+#include "ChooserPropertyAccess.h"
+#include "PoseSearch/Chooser/ChooserParameterPoseHistoryBase.h"
+#include "StructUtils/InstancedStruct.h"
+#include "Serialization/MemoryReader.h"
+#include "Templates/SubclassOf.h"
+#include "PoseSearch/PoseSearchHistory.h"
+#include "PoseSearchChooserColumn.generated.h"
+
+USTRUCT(DisplayName = "Pose History Property Binding")
+struct POSESEARCH_API FPoseHistoryContextProperty :  public FChooserParameterPoseHistoryBase
+{
+	GENERATED_BODY()
+public:
+
+	UPROPERTY(EditAnywhere, Meta = (BindingType = "FPoseHistoryReference", BindingAllowFunctions = "false", BindingColor = "StructPinTypeColor"), Category = "Binding")
+	FChooserPropertyBinding Binding;
+	
+	virtual bool GetValue(FChooserEvaluationContext& Context, FPoseHistoryReference& OutResult) const override;
+	virtual bool IsBound() const override
+	{
+		return Binding.IsBoundToRoot || !Binding.PropertyBindingChain.IsEmpty();
+	}
+
+	CHOOSER_PARAMETER_BOILERPLATE();
+};
+
+
+USTRUCT()
+struct FChooserPoseSearchRowData
+{
+	GENERATED_BODY()
+
+	UPROPERTY(VisibleAnywhere, Category = "Data", Meta = (Tooltip = "The result animation asset for this row (filled by autopopulate from the result column)"));
+	TObjectPtr<UAnimationAsset> ResultAsset;
+};
+
+USTRUCT(DisplayName = "Pose Match", Meta = (Experimental, Category = "Experimental", Tooltip = "This column filters out all assets except the one which is selected by motion matching query.  Results must be AnimationAssets with a PoseSearchBranchIn notify state.  It also outputs OutputStartTime to specify the frame which matched pose best.  To work as intended it must be placed last (furthest right) in the Chooser so that other filters are applied first."))
+struct POSESEARCH_API FPoseSearchColumn : public FChooserColumnBase
+{
+	GENERATED_BODY()
+	public:
+	FPoseSearchColumn();
+	
+	// Pose History
+	UPROPERTY(EditAnywhere, DisplayName = "Pose History", NoClear, Meta = (ExcludeBaseStruct, BaseStruct = "/Script/PoseSearch.ChooserParameterPoseHistoryBase"), Category = "Data")
+	FInstancedStruct InputValue;
+	
+	// Float output for the start time with the best matching pose
+	UPROPERTY(EditAnywhere, NoClear, Meta = (ExcludeBaseStruct, BaseStruct = "/Script/Chooser.ChooserParameterFloatBase"), Category = "Data")
+	FInstancedStruct OutputStartTime;
+	
+	// Bool output for if asset should be mirrored
+	UPROPERTY(EditAnywhere, NoClear, Meta = (ExcludeBaseStruct, BaseStruct = "/Script/Chooser.ChooserParameterBoolBase"), Category = "Data")
+	FInstancedStruct OutputMirror;
+
+	// Float output for the cost of the selected pose
+	UPROPERTY(EditAnywhere, NoClear, Meta = (ExcludeBaseStruct, BaseStruct = "/Script/Chooser.ChooserParameterFloatBase"), Category = "Data")
+	FInstancedStruct OutputCost;
+
+#if WITH_EDITORONLY_DATA
+	UPROPERTY(EditAnywhere, Category=Data, meta=(ToolTip="DefaultRowValue will be assigned to cells when new rows are created"));
+	FChooserPoseSearchRowData DefaultRowValue;
+#endif
+	
+	UPROPERTY()
+	// array of results (cells for this column for each row in the table)
+	// should match the length of the Results array 
+	TArray<FChooserPoseSearchRowData> RowValues;
+
+	virtual bool HasFilters() const override;
+	virtual void Filter(FChooserEvaluationContext& Context, const FChooserIndexArray& IndexListIn, FChooserIndexArray& IndexListOut) const override;
+
+#if WITH_EDITOR
+	mutable double TestValue = 0.0;
+	virtual bool EditorTestFilter(int32 RowIndex) const override;
+	virtual float EditorTestCost(int32 RowIndex) const override;
+	
+	virtual void SetTestValue(TArrayView<const uint8> Value) override
+	{
+		FMemoryReaderView Reader(Value);
+		Reader << TestValue;
+	}
+
+	virtual void AddToDetails(FInstancedPropertyBag& PropertyBag, int32 ColumnIndex, int32 RowIndex) override;
+	virtual void SetFromDetails(FInstancedPropertyBag& PropertyBag, int32 ColumnIndex, int32 RowIndex) override;
+
+	virtual bool AutoPopulates() const override { return true; }
+	virtual void AutoPopulate(int32 RowIndex, UObject* OutputObject) override;
+#endif
+	
+	CHOOSER_COLUMN_BOILERPLATE(FChooserParameterPoseHistoryBase);
+};
\ No newline at end of file
diff --git a/Engine/Plugins/Chooser/Source/ChooserEditor/Public/ChooserColumnHeader.h b/Engine/Plugins/Chooser/Source/ChooserEditor/Public/ChooserColumnHeader.h
--- a/Engine/Plugins/Chooser/Source/ChooserEditor/Public/ChooserColumnHeader.h
+++ b/Engine/Plugins/Chooser/Source/ChooserEditor/Public/ChooserColumnHeader.h
@@ -10,5 +10,5 @@ struct FChooserColumnBase;
 
 namespace UE::ChooserEditor
 {
-	TSharedRef<SWidget> MakeColumnHeaderWidget(UChooserTable* Chooser, FChooserColumnBase* Column, const FText& ColumnName,const FText& ColumnTooltip, const FSlateBrush* ColumnIcon, TSharedPtr<SWidget> DebugWidget, FChooserWidgetValueChanged ValueChanged = FChooserWidgetValueChanged());
+	CHOOSEREDITOR_API TSharedRef<SWidget> MakeColumnHeaderWidget(UChooserTable* Chooser, FChooserColumnBase* Column, const FText& ColumnName,const FText& ColumnTooltip, const FSlateBrush* ColumnIcon, TSharedPtr<SWidget> DebugWidget, FChooserWidgetValueChanged ValueChanged = FChooserWidgetValueChanged());
 }
\ No newline at end of file
