Update PoseSearch to share context with Chooser, for improved interoperability

diff --git a/Engine/Plugins/Animation/PoseSearch/Source/Editor/PoseSearchEditor.Build.cs b/Engine/Plugins/Animation/PoseSearch/Source/Editor/PoseSearchEditor.Build.cs
--- a/Engine/Plugins/Animation/PoseSearch/Source/Editor/PoseSearchEditor.Build.cs
+++ b/Engine/Plugins/Animation/PoseSearch/Source/Editor/PoseSearchEditor.Build.cs
@@ -22,6 +22,7 @@ public class PoseSearchEditor : ModuleRules
 				"CoreUObject",
 				"Engine",
 				"PoseSearch",
+				"Chooser",
 				
 				// Trace-related dependencies
 				"TraceLog",
diff --git a/Engine/Plugins/Animation/PoseSearch/Source/Editor/Private/PoseSearchDatabaseViewModel.cpp b/Engine/Plugins/Animation/PoseSearch/Source/Editor/Private/PoseSearchDatabaseViewModel.cpp
--- a/Engine/Plugins/Animation/PoseSearch/Source/Editor/Private/PoseSearchDatabaseViewModel.cpp
+++ b/Engine/Plugins/Animation/PoseSearch/Source/Editor/Private/PoseSearchDatabaseViewModel.cpp
@@ -249,7 +249,9 @@ bool FDatabasePreviewActor::DrawPreviewActors(TConstArrayView<FDatabasePreviewAc
 	int32 CommonIndexAssetIndex = INDEX_NONE;
 #endif // DO_CHECK
 
-	TArray<const UObject*, TInlineAllocator<PreallocatedRolesNum>> AnimContexts;
+	
+	TArray<FChooserEvaluationContext, TInlineAllocator<PreallocatedRolesNum>> AnimContextsData;
+	TArray<FChooserEvaluationContext*, TInlineAllocator<PreallocatedRolesNum>> AnimContexts;
 	FRoleToIndex RoleToIndex;
 	TArray<FArchivedPoseHistory, TInlineAllocator<PreallocatedRolesNum>> ArchivedPoseHistories;
 	TArray<const IPoseHistory*, TInlineAllocator<PreallocatedRolesNum>> PoseHistories;
@@ -263,6 +265,8 @@ bool FDatabasePreviewActor::DrawPreviewActors(TConstArrayView<FDatabasePreviewAc
 		ArchivedPoseHistories.Reserve(NumPreviewActors);
 		PoseHistories.Reserve(NumPreviewActors);
 	}
+	
+	AnimContextsData.SetNum(NumPreviewActors);
 
 	for (const FDatabasePreviewActor& PreviewActor : PreviewActors)
 	{
@@ -315,8 +319,10 @@ bool FDatabasePreviewActor::DrawPreviewActors(TConstArrayView<FDatabasePreviewAc
 		}
 #endif // DO_CHECK
 
+		int Index = AnimContexts.Num();
 		RoleToIndex.Add(PreviewActor.ActorRole) = AnimContexts.Num();
-		AnimContexts.Add(Mesh);
+		AnimContexts.Add(&AnimContextsData[Index]);
+		AnimContextsData[Index].AddObjectParam(const_cast<UDebugSkelMeshComponent*>(Mesh));
 
 		FArchivedPoseHistory& ArchivedPoseHistory = ArchivedPoseHistories.AddDefaulted_GetRef();
 		ArchivedPoseHistory.Trajectory = PreviewActor.Trajectory;
diff --git a/Engine/Plugins/Animation/PoseSearch/Source/Editor/Private/PoseSearchDebuggerView.cpp b/Engine/Plugins/Animation/PoseSearch/Source/Editor/Private/PoseSearchDebuggerView.cpp
--- a/Engine/Plugins/Animation/PoseSearch/Source/Editor/Private/PoseSearchDebuggerView.cpp
+++ b/Engine/Plugins/Animation/PoseSearch/Source/Editor/Private/PoseSearchDebuggerView.cpp
@@ -227,13 +227,15 @@ void SDebuggerView::Tick(const FGeometry& AllottedGeometry, const double InCurre
 	if (const FTraceMotionMatchingStateMessage* State = Model->GetMotionMatchingState())
 	{
 		FRoleToIndex RoleToIndex;
-		TArray<const UObject*> AnimContexts;
+		TArray<FChooserEvaluationContext> AnimContextsData;
+		TArray<FChooserEvaluationContext*> AnimContexts;
 		TArray<const IPoseHistory*> PoseHistories;
 		UWorld* World = nullptr;
 
 		const int32 NumRoles = State->Roles.Num();
 
 		RoleToIndex.Reserve(NumRoles);
+		AnimContextsData.SetNum(NumRoles);
 		AnimContexts.SetNum(NumRoles);
 		PoseHistories.SetNum(NumRoles);
 
@@ -252,7 +254,8 @@ void SDebuggerView::Tick(const FGeometry& AllottedGeometry, const double InCurre
 						{
 							World = PoseSearchMeshComponent->GetWorld();
 							RoleToIndex.Add(Role) = RoleIndex;
-							AnimContexts[RoleIndex] = PoseSearchMeshComponent;
+							AnimContextsData[RoleIndex].AddObjectParam(PoseSearchMeshComponent);
+							AnimContexts[RoleIndex] = &AnimContextsData[RoleIndex];
 							PoseHistories[RoleIndex] = &State->PoseHistories[RoleIndex];
 							break;
 						}
diff --git a/Engine/Plugins/Animation/PoseSearch/Source/Runtime/Private/AnimNode_MotionMatching.cpp b/Engine/Plugins/Animation/PoseSearch/Source/Runtime/Private/AnimNode_MotionMatching.cpp
--- a/Engine/Plugins/Animation/PoseSearch/Source/Runtime/Private/AnimNode_MotionMatching.cpp
+++ b/Engine/Plugins/Animation/PoseSearch/Source/Runtime/Private/AnimNode_MotionMatching.cpp
@@ -157,8 +157,10 @@ void FAnimNode_MotionMatching::UpdateAssetPlayer(const FAnimationUpdateContext&
 	PoseSearchPlayRate = bPoseSearchPlayRateEnabled ? PoseSearchPlayRate : FFloatInterval(1.f, 1.f);
 #endif // ENABLE_ANIM_DEBUG
 
+	FChooserEvaluationContext EvaluationContext(Context.AnimInstanceProxy->GetAnimInstanceObject());
+
 	UPoseSearchLibrary::UpdateMotionMatchingState(
-		Context.AnimInstanceProxy->GetAnimInstanceObject(),
+		&EvaluationContext,
 		PoseHistory,
 		DatabasesToSearch,
 		Context.GetDeltaTime(),
diff --git a/Engine/Plugins/Animation/PoseSearch/Source/Runtime/Private/Chooser/PoseSearchChooserColumn.cpp b/Engine/Plugins/Animation/PoseSearch/Source/Runtime/Private/Chooser/PoseSearchChooserColumn.cpp
--- a/Engine/Plugins/Animation/PoseSearch/Source/Runtime/Private/Chooser/PoseSearchChooserColumn.cpp
+++ b/Engine/Plugins/Animation/PoseSearch/Source/Runtime/Private/Chooser/PoseSearchChooserColumn.cpp
@@ -76,18 +76,9 @@ void FPoseSearchColumn::Filter(FChooserEvaluationContext& Context, const FChoose
 			{
 				FPoseSearchContinuingProperties ContinuingProperties;
 				FPoseSearchFutureProperties Future;
-				const UObject* ContextObject;
-				for(const FStructView& Param : Context.Params)
-				{
-					if (FChooserEvaluationInputObject* ObjectParam = Param.GetPtr<FChooserEvaluationInputObject>())
-					{
-						ContextObject = ObjectParam->Object;
-						break;
-					}
-				}
 
 				const UE::PoseSearch::FSearchResult SearchResult = UPoseSearchLibrary::MotionMatch(
-					MakeArrayView(&ContextObject, 1),
+					MakeArrayView(&Context, 1),
 					MakeArrayView(&UE::PoseSearch::DefaultRole, 1),
 					MakeArrayView(&PoseHistory, 1),
 					AssetsToSearch,
diff --git a/Engine/Plugins/Animation/PoseSearch/Source/Runtime/Private/PoseSearchContext.cpp b/Engine/Plugins/Animation/PoseSearch/Source/Runtime/Private/PoseSearchContext.cpp
--- a/Engine/Plugins/Animation/PoseSearch/Source/Runtime/Private/PoseSearchContext.cpp
+++ b/Engine/Plugins/Animation/PoseSearch/Source/Runtime/Private/PoseSearchContext.cpp
@@ -71,6 +71,19 @@ const USkeleton* GetContextSkeleton(const UObject* AnimContext)
 	unimplemented();
 	return nullptr;
 }
+	
+const USkeleton* GetContextSkeleton(FChooserEvaluationContext& Context)
+{
+	// @todo: make sure this is not called from worker threads!
+	// check(IsInGameThread());
+
+	if (Context.ObjectParams.Num() > 0)
+	{
+		return GetContextSkeleton(Context.ObjectParams[0].Object);
+	}
+
+	return nullptr;
+}
 
 const USkeletalMeshComponent* GetContextSkeletalMeshComponent(const UObject* AnimContext)
 {
@@ -158,25 +171,25 @@ public:
 	}
 };
 
-static FAnimInstanceProxy* GetAnimInstanceProxy(TConstArrayView<const UObject*> AnimContexts)
+static FAnimInstanceProxy* GetAnimInstanceProxy(TConstArrayView<FChooserEvaluationContext*> AnimContexts)
 {
 	if (!AnimContexts.IsEmpty())
 	{
-		return UAnimInstanceProxyProvider::GetAnimInstanceProxy(Cast<UAnimInstance>(AnimContexts[0]));
+		return UAnimInstanceProxyProvider::GetAnimInstanceProxy(Cast<UAnimInstance>(AnimContexts[0]->GetFirstObjectParam()));
 	}
 	return nullptr;
 }
 
-static const USkinnedMeshComponent* GetMesh(TConstArrayView<const UObject*> AnimContexts, int32 RoleIndex = 0)
+static const USkinnedMeshComponent* GetMesh(TConstArrayView<FChooserEvaluationContext*> AnimContexts, int32 RoleIndex = 0)
 {
 	if (AnimContexts.IsValidIndex(RoleIndex))
 	{
-		return Cast<USkinnedMeshComponent>(AnimContexts[RoleIndex]);
+		return Cast<USkinnedMeshComponent>(AnimContexts[RoleIndex]->GetFirstObjectParam());
 	}
 	return nullptr;
 }
 
-static const UWorld* GetWorld(TConstArrayView<const UObject*> AnimContexts)
+static const UWorld* GetWorld(TConstArrayView<FChooserEvaluationContext*> AnimContexts)
 {
 	if (const USkinnedMeshComponent* SkinnedMeshComponent = GetMesh(AnimContexts))
 	{
@@ -205,7 +218,7 @@ FDebugDrawParams::FDebugDrawParams(TArrayView<const USkinnedMeshComponent*> InMe
 }
 PRAGMA_ENABLE_DEPRECATION_WARNINGS
 
-FDebugDrawParams::FDebugDrawParams(TConstArrayView<const UObject*> InAnimContexts, TConstArrayView<const IPoseHistory*> InPoseHistories, const FRoleToIndex& InRoleToIndex, const UPoseSearchDatabase* InDatabase)
+FDebugDrawParams::FDebugDrawParams(TConstArrayView<FChooserEvaluationContext*> InAnimContexts, TConstArrayView<const IPoseHistory*> InPoseHistories, const FRoleToIndex& InRoleToIndex, const UPoseSearchDatabase* InDatabase)
 : AnimContexts(InAnimContexts)
 , PoseHistories(InPoseHistories)
 , RoleToIndex(InRoleToIndex)
@@ -731,7 +744,7 @@ FSearchContext::FSearchContext(float InDesiredPermutationTimeOffset, const FPose
 #endif // UE_POSE_SEARCH_TRACE_ENABLED
 }
 
-void FSearchContext::AddRole(const FRole& Role, const UObject* AnimContext, const IPoseHistory* PoseHistory)
+void FSearchContext::AddRole(const FRole& Role, FChooserEvaluationContext* AnimContext, const IPoseHistory* PoseHistory)
 {
 	check(RoleToIndex.Num() == AnimContexts.Num() && RoleToIndex.Num() == PoseHistories.Num());
 
@@ -922,7 +935,7 @@ const UAnimInstance* FSearchContext::GetAnimInstance(const FRole& Role) const
 {
 	if (const int32* RoleIndex = RoleToIndex.Find(Role))
 	{
-		return Cast<UAnimInstance>(AnimContexts[*RoleIndex]);
+		return Cast<UAnimInstance>(AnimContexts[*RoleIndex]->GetFirstObjectParam());
 	}
 
 	UE_LOG(LogPoseSearch, Error, TEXT("FSearchContext::GetAnimInstance - Role %s could not be found!"), *Role.ToString());
@@ -933,13 +946,35 @@ const UObject* FSearchContext::GetAnimContext(const FRole& Role) const
 {
 	if (const int32* RoleIndex = RoleToIndex.Find(Role))
 	{
-		return AnimContexts[*RoleIndex];
+		return AnimContexts[*RoleIndex]->GetFirstObjectParam();
 	}
 		
 	UE_LOG(LogPoseSearch, Error, TEXT("FSearchContext::GetAnimContext - Role %s could not be found!"), *Role.ToString());
 	return nullptr;
 }
 
+const FChooserEvaluationContext* FSearchContext::GetContext(const FRole& Role) const
+{
+	if (const int32* RoleIndex = RoleToIndex.Find(Role))
+	{
+		return AnimContexts[*RoleIndex];
+	}
+		
+	UE_LOG(LogPoseSearch, Error, TEXT("FSearchContext::GetAnimContext - Role %s could not be found!"), *Role.ToString());
+	return nullptr;	
+}
+	
+FChooserEvaluationContext* FSearchContext::GetContext(const FRole& Role)
+{
+	if (const int32* RoleIndex = RoleToIndex.Find(Role))
+	{
+		return AnimContexts[*RoleIndex];
+	}
+		
+	UE_LOG(LogPoseSearch, Error, TEXT("FSearchContext::GetAnimContext - Role %s could not be found!"), *Role.ToString());
+	return nullptr;	
+}
+
 FTransform FSearchContext::GetWorldBoneTransformAtTime(float SampleTime, const FRole& SampleRole, int8 SchemaBoneIdx)
 {
 	// CachedQueries.Last is the query we're building 
diff --git a/Engine/Plugins/Animation/PoseSearch/Source/Runtime/Private/PoseSearchDatabase.cpp b/Engine/Plugins/Animation/PoseSearch/Source/Runtime/Private/PoseSearchDatabase.cpp
--- a/Engine/Plugins/Animation/PoseSearch/Source/Runtime/Private/PoseSearchDatabase.cpp
+++ b/Engine/Plugins/Animation/PoseSearch/Source/Runtime/Private/PoseSearchDatabase.cpp
@@ -1363,7 +1363,7 @@ UE::PoseSearch::FSearchResult UPoseSearchDatabase::Search(UE::PoseSearch::FSearc
 	// (and restore them with the original one CurrentAssetToConsider at the end of the method)
 	TConstArrayView<const UObject*> CurrentAssetToConsider = SearchContext.GetAssetsToConsider();
 	TArray<const UObject*, TInlineAllocator<128>> ChooserAssetToConsider;
-	if (Chooser && !SearchContext.GetAnimContexts().IsEmpty())
+	if (Chooser && !SearchContext.GetContexts().IsEmpty())
 	{
 		if (!CurrentAssetToConsider.IsEmpty())
 		{
@@ -1371,8 +1371,7 @@ UE::PoseSearch::FSearchResult UPoseSearchDatabase::Search(UE::PoseSearch::FSearc
 		}
 
 		// @todo: add MM interaction support. Maybe iterate over all the SearchContext.GetAnimContexts and use the union of the ChosenAssets?
-		FChooserEvaluationContext ChooserContext(const_cast<UObject*>(SearchContext.GetAnimContexts()[0]));
-		UChooserTable::EvaluateChooser(ChooserContext, Chooser, FObjectChooserBase::FObjectChooserIteratorCallback::CreateLambda([this, &ChooserAssetToConsider](UObject* InResult)
+		UChooserTable::EvaluateChooser(*SearchContext.GetContexts()[0], Chooser, FObjectChooserBase::FObjectChooserIteratorCallback::CreateLambda([this, &ChooserAssetToConsider](UObject* InResult)
 			{
 				ChooserAssetToConsider.Add(InResult);
 				return FObjectChooserBase::EIteratorStatus::Continue;
diff --git a/Engine/Plugins/Animation/PoseSearch/Source/Runtime/Private/PoseSearchFeatureChannel_Curve.cpp b/Engine/Plugins/Animation/PoseSearch/Source/Runtime/Private/PoseSearchFeatureChannel_Curve.cpp
--- a/Engine/Plugins/Animation/PoseSearch/Source/Runtime/Private/PoseSearchFeatureChannel_Curve.cpp
+++ b/Engine/Plugins/Animation/PoseSearch/Source/Runtime/Private/PoseSearchFeatureChannel_Curve.cpp
@@ -35,7 +35,7 @@ void UPoseSearchFeatureChannel_Curve::BuildQuery(UE::PoseSearch::FSearchContext&
 	float CurveValue = 0.0f;
 	if (bUseBlueprintQueryOverride)
 	{
-		if (const UAnimInstance* AnimInstance = Cast<UAnimInstance>(SearchContext.GetAnimContext(SampleRole)))
+		if (const UAnimInstance* AnimInstance = Cast<UAnimInstance>(SearchContext.GetContext(SampleRole)->GetFirstObjectParam()))
 		{
 			CurveValue = BP_GetCurveValue(AnimInstance);
 			FFeatureVectorHelper::EncodeFloat(SearchContext.EditFeatureVector(), ChannelDataOffset, CurveValue);
diff --git a/Engine/Plugins/Animation/PoseSearch/Source/Runtime/Private/PoseSearchFeatureChannel_Distance.cpp b/Engine/Plugins/Animation/PoseSearch/Source/Runtime/Private/PoseSearchFeatureChannel_Distance.cpp
--- a/Engine/Plugins/Animation/PoseSearch/Source/Runtime/Private/PoseSearchFeatureChannel_Distance.cpp
+++ b/Engine/Plugins/Animation/PoseSearch/Source/Runtime/Private/PoseSearchFeatureChannel_Distance.cpp
@@ -70,7 +70,7 @@ void UPoseSearchFeatureChannel_Distance::BuildQuery(UE::PoseSearch::FSearchConte
 
 	if (bUseBlueprintQueryOverride)
 	{
-		if (const UAnimInstance* AnimInstance = Cast<UAnimInstance>(SearchContext.GetAnimContext(SampleRole)))
+		if (const UAnimInstance* AnimInstance = Cast<UAnimInstance>(SearchContext.GetContext(SampleRole)->GetFirstObjectParam()))
 		{
 			const float Distance = BP_GetDistance(AnimInstance);
   			FFeatureVectorHelper::EncodeFloat(SearchContext.EditFeatureVector(), ChannelDataOffset, Distance);
diff --git a/Engine/Plugins/Animation/PoseSearch/Source/Runtime/Private/PoseSearchFeatureChannel_Heading.cpp b/Engine/Plugins/Animation/PoseSearch/Source/Runtime/Private/PoseSearchFeatureChannel_Heading.cpp
--- a/Engine/Plugins/Animation/PoseSearch/Source/Runtime/Private/PoseSearchFeatureChannel_Heading.cpp
+++ b/Engine/Plugins/Animation/PoseSearch/Source/Runtime/Private/PoseSearchFeatureChannel_Heading.cpp
@@ -107,7 +107,7 @@ void UPoseSearchFeatureChannel_Heading::BuildQuery(UE::PoseSearch::FSearchContex
 	const bool bIsRootBone = SchemaBoneIdx == RootSchemaBoneIdx;
 	if (bUseBlueprintQueryOverride)
 	{
-		if (const UAnimInstance* AnimInstance = Cast<UAnimInstance>(SearchContext.GetAnimContext(SampleRole)))
+		if (const UAnimInstance* AnimInstance = Cast<UAnimInstance>(SearchContext.GetContext(SampleRole)->GetFirstObjectParam()))
 		{
 			const FQuat BoneRotationWorld = BP_GetWorldRotation(AnimInstance);
 			const FQuat BoneRotation = SearchContext.GetSampleRotation(SampleTimeOffset, OriginTimeOffset, SchemaBoneIdx, SchemaOriginBoneIdx, SampleRole, OriginRole, EPermutationTimeType::UseSampleTime, &BoneRotationWorld);
diff --git a/Engine/Plugins/Animation/PoseSearch/Source/Runtime/Private/PoseSearchFeatureChannel_Position.cpp b/Engine/Plugins/Animation/PoseSearch/Source/Runtime/Private/PoseSearchFeatureChannel_Position.cpp
--- a/Engine/Plugins/Animation/PoseSearch/Source/Runtime/Private/PoseSearchFeatureChannel_Position.cpp
+++ b/Engine/Plugins/Animation/PoseSearch/Source/Runtime/Private/PoseSearchFeatureChannel_Position.cpp
@@ -109,7 +109,7 @@ void UPoseSearchFeatureChannel_Position::BuildQuery(UE::PoseSearch::FSearchConte
 
 	if (bUseBlueprintQueryOverride)
 	{
-		if (const UAnimInstance* AnimInstance = Cast<UAnimInstance>(SearchContext.GetAnimContext(SampleRole)))
+		if (const UAnimInstance* AnimInstance = Cast<UAnimInstance>(SearchContext.GetContext(SampleRole)->GetFirstObjectParam()))
 		{
 			const FVector BonePositionWorld = BP_GetWorldPosition(AnimInstance);
 			const FVector BonePosition = SearchContext.GetSamplePosition(SampleTimeOffset, OriginTimeOffset, SchemaBoneIdx, SchemaOriginBoneIdx, SampleRole, OriginRole, PermutationTimeType, &BonePositionWorld);
diff --git a/Engine/Plugins/Animation/PoseSearch/Source/Runtime/Private/PoseSearchFeatureChannel_TimeToEvent.cpp b/Engine/Plugins/Animation/PoseSearch/Source/Runtime/Private/PoseSearchFeatureChannel_TimeToEvent.cpp
--- a/Engine/Plugins/Animation/PoseSearch/Source/Runtime/Private/PoseSearchFeatureChannel_TimeToEvent.cpp
+++ b/Engine/Plugins/Animation/PoseSearch/Source/Runtime/Private/PoseSearchFeatureChannel_TimeToEvent.cpp
@@ -30,12 +30,12 @@ void UPoseSearchFeatureChannel_TimeToEvent::BuildQuery(UE::PoseSearch::FSearchCo
 		return;
 	}
 
-	if (SearchContext.GetAnimContexts().IsEmpty())
+	if (SearchContext.GetContexts().IsEmpty())
 	{
 		UE_LOG(LogPoseSearch, Error, TEXT("UPoseSearchFeatureChannel_Position::BuildQuery - no provided anim contexts!"));
 	}
 	
-	if (const UAnimInstance* AnimInstance = Cast<UAnimInstance>(SearchContext.GetAnimContexts()[0]))
+	if (const UAnimInstance* AnimInstance = Cast<UAnimInstance>(SearchContext.GetContexts()[0]->GetFirstObjectParam()))
 	{
 		const float TimeToEvent = BP_GetTimeToEvent(AnimInstance);
 		FFeatureVectorHelper::EncodeFloat(SearchContext.EditFeatureVector(), ChannelDataOffset, TimeToEvent);
diff --git a/Engine/Plugins/Animation/PoseSearch/Source/Runtime/Private/PoseSearchFeatureChannel_Velocity.cpp b/Engine/Plugins/Animation/PoseSearch/Source/Runtime/Private/PoseSearchFeatureChannel_Velocity.cpp
--- a/Engine/Plugins/Animation/PoseSearch/Source/Runtime/Private/PoseSearchFeatureChannel_Velocity.cpp
+++ b/Engine/Plugins/Animation/PoseSearch/Source/Runtime/Private/PoseSearchFeatureChannel_Velocity.cpp
@@ -49,7 +49,7 @@ void UPoseSearchFeatureChannel_Velocity::BuildQuery(UE::PoseSearch::FSearchConte
 	const bool bIsRootBone = SchemaBoneIdx == RootSchemaBoneIdx;
 	if (bUseBlueprintQueryOverride)
 	{
-		if (const UAnimInstance* AnimInstance = Cast<UAnimInstance>(SearchContext.GetAnimContext(SampleRole)))
+		if (const UAnimInstance* AnimInstance = Cast<UAnimInstance>(SearchContext.GetContext(SampleRole)->GetFirstObjectParam()))
 		{
 			const FVector LinearVelocityWorld = BP_GetWorldVelocity(AnimInstance);
 			FVector LinearVelocity = SearchContext.GetSampleVelocity(SampleTimeOffset, OriginTimeOffset, SchemaBoneIdx, SchemaOriginBoneIdx, SampleRole, OriginRole, bUseCharacterSpaceVelocities, EPermutationTimeType::UseSampleTime, &LinearVelocityWorld);
diff --git a/Engine/Plugins/Animation/PoseSearch/Source/Runtime/Private/PoseSearchLibrary.cpp b/Engine/Plugins/Animation/PoseSearch/Source/Runtime/Private/PoseSearchLibrary.cpp
--- a/Engine/Plugins/Animation/PoseSearch/Source/Runtime/Private/PoseSearchLibrary.cpp
+++ b/Engine/Plugins/Animation/PoseSearch/Source/Runtime/Private/PoseSearchLibrary.cpp
@@ -24,6 +24,7 @@
 #include "PoseSearch/PoseSearchFeatureChannel_Trajectory.h"
 #include "PoseSearch/Trace/PoseSearchTraceLogger.h"
 #include "PoseSearch/PoseSearchFeatureChannel_PermutationTime.h"
+#include "UObject/FastReferenceCollector.h"
 
 #include UE_INLINE_GENERATED_CPP_BY_NAME(PoseSearchLibrary)
 
@@ -432,9 +433,12 @@ void UPoseSearchLibrary::TraceMotionMatching(
 	}
 
 	float RecordingTime = 0.f;
-	if (!SearchContext.GetAnimContexts().IsEmpty())
+	if (!SearchContext.GetContexts().IsEmpty())
 	{
-		RecordingTime = FObjectTrace::GetWorldElapsedTime(SearchContext.GetAnimContexts()[0]->GetWorld());
+		if (UObject* FirstObject = SearchContext.GetContexts()[0]->GetFirstObjectParam())
+		{
+			RecordingTime = FObjectTrace::GetWorldElapsedTime(FirstObject->GetWorld());
+		}
 	}
 
 	uint32 SearchId = 787;
@@ -443,19 +447,20 @@ void UPoseSearchLibrary::TraceMotionMatching(
 
 	TraceState.InterruptMode = InterruptMode;
 
-	const int32 AnimContextsNum = SearchContext.GetAnimContexts().Num();
+	const int32 AnimContextsNum = SearchContext.GetContexts().Num();
 	TraceState.SkeletalMeshComponentIds.SetNum(AnimContextsNum);
 
 	for (int32 AnimInstanceIndex = 0; AnimInstanceIndex < AnimContextsNum; ++AnimInstanceIndex)
 	{
-		if (const UObject* AnimContext = SearchContext.GetAnimContexts()[AnimInstanceIndex])
+		if (FChooserEvaluationContext* AnimContext = SearchContext.GetContexts()[AnimInstanceIndex])
 		{
+			UObject* FirstObject = AnimContext->GetFirstObjectParam();
 			const UObject* SkeletalMeshComponent = nullptr;
-			if (const UAnimInstance* AnimInstance = Cast<UAnimInstance>(AnimContext))
+			if (const UAnimInstance* AnimInstance = Cast<UAnimInstance>(AnimContext->GetFirstObjectParam()))
 			{
 				SkeletalMeshComponent = AnimInstance->GetOuter();
 			}
-			else if (const UActorComponent* ActorComponent = Cast<UActorComponent>(AnimContext))
+			else if (const UActorComponent* ActorComponent = Cast<UActorComponent>(AnimContext->GetFirstObjectParam()))
 			{
 				const AActor* Actor = ActorComponent->GetOwner();
 				check(Actor);
@@ -474,9 +479,12 @@ void UPoseSearchLibrary::TraceMotionMatching(
 
 	for (int32 AnimInstanceIndex = 0; AnimInstanceIndex < AnimContextsNum; ++AnimInstanceIndex)
 	{
-		const UObject* AnimContext = SearchContext.GetAnimContexts()[AnimInstanceIndex];
-		TRACE_OBJECT(AnimContext);
-		SearchId = HashCombineFast(SearchId, GetTypeHash(FObjectTrace::GetObjectId(AnimContext)));
+		FChooserEvaluationContext* Context = SearchContext.GetContexts()[AnimInstanceIndex];
+		if (UObject* Object = Context->GetFirstObjectParam())
+		{
+			TRACE_OBJECT(Object);
+			SearchId = HashCombineFast(SearchId, GetTypeHash(FObjectTrace::GetObjectId(Object)));
+		}
 	}
 
 	TraceState.Roles.SetNum(AnimContextsNum);
@@ -606,9 +614,9 @@ void UPoseSearchLibrary::TraceMotionMatching(
 	TraceState.Cycle = FPlatformTime::Cycles64();
 
 	// @todo: avoid publishing duplicated TraceState in ALL the AnimContexts! -currently necessary for multi character-
-	for (const UObject* AnimContext : SearchContext.GetAnimContexts())
+	for (FChooserEvaluationContext* Context : SearchContext.GetContexts())
 	{
-		TraceState.AnimInstanceId = FObjectTrace::GetObjectId(AnimContext);
+		TraceState.AnimInstanceId = FObjectTrace::GetObjectId(Context->GetFirstObjectParam());
 		TraceState.NodeId = SearchId;
 		TraceState.Output();
 	}
@@ -648,7 +656,8 @@ void UPoseSearchLibrary::UpdateMotionMatchingState(
 	}
 
 	check(Context.AnimInstanceProxy);
-	UpdateMotionMatchingState(Context.AnimInstanceProxy->GetAnimInstanceObject(), PoseHistory, Databases, Context.GetDeltaTime(),
+	FChooserEvaluationContext AnimContext(Context.AnimInstanceProxy->GetAnimInstanceObject());
+	UpdateMotionMatchingState(&AnimContext, PoseHistory, Databases, Context.GetDeltaTime(),
 		PoseJumpThresholdTime, PoseReselectHistory, bShouldSearch ? SearchThrottleTime : UE_BIG_NUMBER, PlayRate, InOutMotionMatchingState,
 		InterruptMode, bShouldUseCachedChannelData, bDebugDrawQuery, bDebugDrawCurResult);
 }
@@ -668,6 +677,37 @@ void UPoseSearchLibrary::UpdateMotionMatchingState(
 	bool bDebugDrawQuery,
 	bool bDebugDrawCurResult,
 	const FPoseSearchEvent& EventToSearch)
+{
+	FChooserEvaluationContext Context(const_cast<UObject*>(AnimContext));
+	UpdateMotionMatchingState(&Context,
+		PoseHistory,
+		Databases,
+		DeltaTime,
+		PoseJumpThresholdTime,
+		PoseReselectHistory,
+		SearchThrottleTime,
+		PlayRate,
+		InOutMotionMatchingState,
+		InterruptMode,
+		bShouldUseCachedChannelData,
+		bDebugDrawQuery, bDebugDrawCurResult);
+}
+
+void UPoseSearchLibrary::UpdateMotionMatchingState(
+	FChooserEvaluationContext* AnimContext,
+	const UE::PoseSearch::IPoseHistory* PoseHistory,
+	const TConstArrayView<TObjectPtr<const UPoseSearchDatabase>> Databases,
+	float DeltaTime,
+	const FFloatInterval& PoseJumpThresholdTime,
+	float PoseReselectHistory,
+	float SearchThrottleTime,
+	const FFloatInterval& PlayRate,
+	FMotionMatchingState& InOutMotionMatchingState,
+	EPoseSearchInterruptMode InterruptMode,
+	bool bShouldUseCachedChannelData,
+	bool bDebugDrawQuery,
+	bool bDebugDrawCurResult,
+	const FPoseSearchEvent& EventToSearch)
 {
 	QUICK_SCOPE_CYCLE_COUNTER(STAT_PoseSearch_Update);
 
@@ -841,7 +881,7 @@ PRAGMA_ENABLE_DEPRECATION_WARNINGS
 		if (EAsyncBuildIndexResult::Success == FAsyncPoseSearchDatabasesManagement::RequestAsyncBuildIndex(CurResultDatabase, ERequestAsyncBuildFlag::ContinueRequest))
 #endif // WITH_EDITOR
 		{
-			FDebugDrawParams DrawParams(SearchContext.GetAnimContexts(), SearchContext.GetPoseHistories(), SearchContext.GetRoleToIndex(), CurResultDatabase);
+			FDebugDrawParams DrawParams(SearchContext.GetContexts(), SearchContext.GetPoseHistories(), SearchContext.GetRoleToIndex(), CurResultDatabase);
 			if (bDebugDrawCurResult)
 			{
 				DrawParams.DrawFeatureVector(CurResult.PoseIdx);
@@ -1001,6 +1041,26 @@ UE::PoseSearch::FSearchResult UPoseSearchLibrary::MotionMatch(
 	const FPoseSearchContinuingProperties& ContinuingProperties,
 	const FPoseSearchFutureProperties& Future,
 	const FPoseSearchEvent& EventToSearch)
+{
+	TArray<FChooserEvaluationContext, TInlineAllocator<UE::PoseSearch::PreallocatedRolesNum>> Contexts;
+	int NumContexts = AnimContexts.Num(); 
+	Contexts.SetNum(NumContexts);
+	for(int i = 0; i <NumContexts; i++)
+	{
+		Contexts[i].AddObjectParam(const_cast<UObject*>(AnimContexts[i]));
+	}
+
+	return MotionMatch(Contexts, Roles, PoseHistories, AssetsToSearch, ContinuingProperties, Future, EventToSearch);
+}
+
+UE::PoseSearch::FSearchResult UPoseSearchLibrary::MotionMatch(
+	const TArrayView<FChooserEvaluationContext> AnimContexts,
+	const TArrayView<const UE::PoseSearch::FRole> Roles,
+	const TArrayView<const UE::PoseSearch::IPoseHistory*> PoseHistories, 
+	const TArrayView<const UObject*> AssetsToSearch,
+	const FPoseSearchContinuingProperties& ContinuingProperties,
+	const FPoseSearchFutureProperties& Future,
+	const FPoseSearchEvent& EventToSearch)
 {
 	check(!AnimContexts.IsEmpty() && AnimContexts.Num() == Roles.Num() && AnimContexts.Num() == PoseHistories.Num());
 
@@ -1111,16 +1171,19 @@ UE::PoseSearch::FSearchResult UPoseSearchLibrary::MotionMatch(
 
 	for (int32 RoleIndex = 0; RoleIndex < Roles.Num(); ++RoleIndex)
 	{
-		SearchContext.AddRole(Roles[RoleIndex], AnimContexts[RoleIndex], InternalPoseHistories[RoleIndex]);
+		SearchContext.AddRole(Roles[RoleIndex], &AnimContexts[RoleIndex], InternalPoseHistories[RoleIndex]);
 	}
 	
 	bool bAsyncBuildIndexInProgress = false;
 
 	// collecting all the possible continuing pose search (it could be multiple searches, but most likely only one)
 	float DeltaSeconds = FiniteDelta;
-	if (const UAnimInstance* AnimInstance = Cast<UAnimInstance>(AnimContexts[0]))
+	if (AnimContexts[0].ObjectParams.Num() > 0)
 	{
-		DeltaSeconds = AnimInstance->GetDeltaSeconds();
+		if (const UAnimInstance* AnimInstance = Cast<UAnimInstance>(AnimContexts[0].GetFirstObjectParam()))
+		{
+			DeltaSeconds = AnimInstance->GetDeltaSeconds();
+		}
 	}
 
 	// collecting all the databases searches in AssetsToSearchPerDatabaseMap
@@ -1244,7 +1307,7 @@ UE::PoseSearch::FSearchResult UPoseSearchLibrary::MotionMatch(
 
 		if (bDrawMatch || bDrawquery)
 		{
-			FDebugDrawParams DrawParams(SearchContext.GetAnimContexts(), SearchContext.GetPoseHistories(), SearchContext.GetRoleToIndex(), SearchResult.Database.Get());
+			FDebugDrawParams DrawParams(SearchContext.GetContexts(), SearchContext.GetPoseHistories(), SearchContext.GetRoleToIndex(), SearchResult.Database.Get());
 			if (bDrawMatch)
 			{
 				DrawParams.DrawFeatureVector(SearchResult.PoseIdx);
diff --git a/Engine/Plugins/Animation/PoseSearch/Source/Runtime/Public/PoseSearch/PoseSearchContext.h b/Engine/Plugins/Animation/PoseSearch/Source/Runtime/Public/PoseSearch/PoseSearchContext.h
--- a/Engine/Plugins/Animation/PoseSearch/Source/Runtime/Public/PoseSearch/PoseSearchContext.h
+++ b/Engine/Plugins/Animation/PoseSearch/Source/Runtime/Public/PoseSearch/PoseSearchContext.h
@@ -4,6 +4,7 @@
 
 #include "Animation/AnimInstance.h"
 #include "DrawDebugHelpers.h"
+#include "IObjectChooser.h"
 #include "PoseSearch/PoseSearchDefines.h"
 #include "PoseSearch/PoseSearchEvent.h"
 #include "PoseSearch/PoseSearchFeatureChannel.h"
@@ -16,10 +17,12 @@ class UPoseSearchFeatureChannel_Position;
 
 namespace UE::PoseSearch
 {
+	
 
 // Utils functions
 const FTransform& GetContextTransform(const UObject* AnimContext);
 const USkeleton* GetContextSkeleton(const UObject* AnimContext);
+const USkeleton* GetContextSkeleton(FChooserEvaluationContext& AnimContext);
 const AActor* GetContextOwningActor(const UObject* AnimContext);
 FVector GetContextLocation(const UObject* AnimContext);
 
@@ -73,7 +76,7 @@ struct POSESEARCH_API FDebugDrawParams
 	UE_DEPRECATED(5.6, "Use other FDebugDrawParams signatures instead")
 	FDebugDrawParams(TArrayView<const USkinnedMeshComponent*> InMeshes, TConstArrayView<const IPoseHistory*> InPoseHistories, const FRoleToIndex& InRoleToIndex, const UPoseSearchDatabase* InDatabase, EDebugDrawFlags InFlags = EDebugDrawFlags::None);
 	PRAGMA_ENABLE_DEPRECATION_WARNINGS
-	FDebugDrawParams(TConstArrayView<const UObject*> InAnimContexts, TConstArrayView<const IPoseHistory*> InPoseHistories, const FRoleToIndex& InRoleToIndex, const UPoseSearchDatabase* InDatabase);
+	FDebugDrawParams(TConstArrayView<FChooserEvaluationContext*> InAnimContexts, TConstArrayView<const IPoseHistory*> InPoseHistories, const FRoleToIndex& InRoleToIndex, const UPoseSearchDatabase* InDatabase);
 	
 	const FSearchIndex* GetSearchIndex() const;
 	const UPoseSearchSchema* GetSchema() const;
@@ -99,7 +102,7 @@ struct POSESEARCH_API FDebugDrawParams
 private:
 	bool CanDraw() const;
 
-	const TConstArrayView<const UObject*> AnimContexts;
+	const TConstArrayView<FChooserEvaluationContext*> AnimContexts;
 	const TConstArrayView<const IPoseHistory*> PoseHistories;
 
 	// NoTe: mapping Role to the index of the associated asset that this FDebugDrawParams is drawing.
@@ -154,7 +157,7 @@ struct POSESEARCH_API FSearchContext
 	FSearchContext& operator=(const FSearchContext& Other) = delete;
 	FSearchContext& operator=(FSearchContext&& Other) = delete;
 
-	void AddRole(const FRole& Role, const UObject* AnimContext, const IPoseHistory* PoseHistory);
+	void AddRole(const FRole& Role, FChooserEvaluationContext* AnimContext, const IPoseHistory* PoseHistory);
 
 	// Returns the curve value of name CurveName at an offset time of SampleTimeOffset.
 	// If the curve is not found, assume value of 0, which is consistent with curve behavior in the animation update.
@@ -199,21 +202,36 @@ struct POSESEARCH_API FSearchContext
 
 	float GetDesiredPermutationTimeOffset() const { return DesiredPermutationTimeOffset; }
 	
-	UE_DEPRECATED(5.6, "Use GetAnimContext instead")
+	UE_DEPRECATED(5.6, "Use GetContext instead")
 	const UAnimInstance* GetAnimInstance(const FRole& Role) const;
+	UE_DEPRECATED(5.6, "Use GetContext instead")
 	const UObject* GetAnimContext(const FRole& Role) const;
 
-	UE_DEPRECATED(5.6, "Use GetAnimContexts instead")
+	const FChooserEvaluationContext* GetContext(const FRole& Role) const;
+	FChooserEvaluationContext* GetContext(const FRole& Role);
+	UE_DEPRECATED(5.6, "Use GetContexts instead")
 	const TArray<const UAnimInstance*, TInlineAllocator<PreallocatedRolesNum, TMemStackAllocator<>>> GetAnimInstances() const 
+	{
+	 	TArray<const UAnimInstance*, TInlineAllocator<PreallocatedRolesNum, TMemStackAllocator<>>> AnimInstances;
+	 	for (const FChooserEvaluationContext* AnimContext : AnimContexts)
+	 	{
+	 		AnimInstances.Add(Cast<UAnimInstance>(AnimContext->GetFirstObjectParam()));
+	 	}
+	 	return AnimInstances;
+	}
+
+	UE_DEPRECATED(5.6, "Use GetContexts instead")
+	const TArray<const UAnimInstance*, TInlineAllocator<PreallocatedRolesNum, TMemStackAllocator<>>> GetAnimContexts() const 
 	{
 		TArray<const UAnimInstance*, TInlineAllocator<PreallocatedRolesNum, TMemStackAllocator<>>> AnimInstances;
-		for (const UObject* AnimContext : AnimContexts)
+		for (const FChooserEvaluationContext* AnimContext : AnimContexts)
 		{
-			AnimInstances.Add(Cast<UAnimInstance>(AnimContext));
+			AnimInstances.Add(Cast<UAnimInstance>(AnimContext->GetFirstObjectParam()));
 		}
 		return AnimInstances;
 	}
-	TConstArrayView<const UObject*> GetAnimContexts() const { return AnimContexts; }
+	
+	TArrayView<FChooserEvaluationContext*> GetContexts() { return AnimContexts; }
 
 	const FRoleToIndex& GetRoleToIndex() const { return RoleToIndex; }
 
@@ -245,7 +263,7 @@ private:
 	FQuat GetSampleRotationInternal(float SampleTime, float OriginTime, int8 SchemaSampleBoneIdx, int8 SchemaOriginBoneIdx, const FRole& SampleRole, const FRole& OriginRole, const FQuat* SampleBoneRotationWorldOverride = nullptr);
 	FTransform GetWorldRootBoneTransformAtTime(float SampleTime, const FRole& SampleRole) const;
 	
-	TArray<const UObject*, TInlineAllocator<PreallocatedRolesNum, TMemStackAllocator<>>> AnimContexts;
+	TArray<FChooserEvaluationContext*, TInlineAllocator<PreallocatedRolesNum, TMemStackAllocator<>>> AnimContexts;
 	TArray<const IPoseHistory*, TInlineAllocator<PreallocatedRolesNum, TMemStackAllocator<>>> PoseHistories;
 	FRoleToIndex RoleToIndex;
 	
diff --git a/Engine/Plugins/Animation/PoseSearch/Source/Runtime/Public/PoseSearch/PoseSearchLibrary.h b/Engine/Plugins/Animation/PoseSearch/Source/Runtime/Public/PoseSearch/PoseSearchLibrary.h
--- a/Engine/Plugins/Animation/PoseSearch/Source/Runtime/Public/PoseSearch/PoseSearchLibrary.h
+++ b/Engine/Plugins/Animation/PoseSearch/Source/Runtime/Public/PoseSearch/PoseSearchLibrary.h
@@ -2,6 +2,7 @@
 
 #pragma once
 
+#include "IObjectChooser.h"
 #include "Animation/AnimExecutionContext.h"
 #include "Animation/AnimNodeReference.h"
 #include "Kismet/BlueprintFunctionLibrary.h"
@@ -198,6 +199,23 @@ public:
 	* @param bDebugDrawQuery				Input draw the composed query if valid
 	* @param bDebugDrawCurResult			Input draw the current result if valid
 	*/
+	static void UpdateMotionMatchingState(
+		FChooserEvaluationContext* AnimContext,
+		const UE::PoseSearch::IPoseHistory* PoseHistory,
+		const TConstArrayView<TObjectPtr<const UPoseSearchDatabase>> Databases,
+		float DeltaTime,
+		const FFloatInterval& PoseJumpThresholdTime,
+		float PoseReselectHistory,
+		float SearchThrottleTime,
+		const FFloatInterval& PlayRate,
+		FMotionMatchingState& InOutMotionMatchingState,
+		EPoseSearchInterruptMode InterruptMode,
+		bool bShouldUseCachedChannelData,
+		bool bDebugDrawQuery = false,
+		bool bDebugDrawCurResult = false,
+		const FPoseSearchEvent& EventToSearch = FPoseSearchEvent());
+	
+	UE_DEPRECATED(5.6, "Use other UpdateMotionMatchingState signatures instead")
 	static void UpdateMotionMatchingState(
 		const UObject* AnimContext,
 		const UE::PoseSearch::IPoseHistory* PoseHistory,
@@ -213,7 +231,7 @@ public:
 		bool bDebugDrawQuery = false,
 		bool bDebugDrawCurResult = false,
 		const FPoseSearchEvent& EventToSearch = FPoseSearchEvent());
-
+	
 	UE_DEPRECATED(5.6, "Use other UpdateMotionMatchingState signatures instead")
 	static void UpdateMotionMatchingState(
 		const FAnimationUpdateContext& Context,
@@ -268,6 +286,16 @@ public:
 		const FPoseSearchContinuingProperties& ContinuingProperties,
 		const FPoseSearchFutureProperties& Future,
 		const FPoseSearchEvent& EventToSearch);
+	
+	static UE::PoseSearch::FSearchResult MotionMatch(
+			const TArrayView<FChooserEvaluationContext> Contexts,
+    		const TArrayView<const UE::PoseSearch::FRole> Roles,
+    		const TArrayView<const UE::PoseSearch::IPoseHistory*> PoseHistories, 
+    		const TArrayView<const UObject*> AssetsToSearch,
+    		const FPoseSearchContinuingProperties& ContinuingProperties,
+    		const FPoseSearchFutureProperties& Future,
+    		const FPoseSearchEvent& EventToSearch);
+	
 		
 	UE_DEPRECATED(5.5, "Use other MotionMatch signatures instead")
 	static UE::PoseSearch::FSearchResult MotionMatch(
diff --git a/Engine/Plugins/Chooser/Source/Chooser/Public/IObjectChooser.h b/Engine/Plugins/Chooser/Source/Chooser/Public/IObjectChooser.h
--- a/Engine/Plugins/Chooser/Source/Chooser/Public/IObjectChooser.h
+++ b/Engine/Plugins/Chooser/Source/Chooser/Public/IObjectChooser.h
@@ -64,6 +64,18 @@ struct CHOOSER_API FChooserEvaluationContext
 		AddStructParam(ObjectParams.Last());
 	}
 
+	// Helper to get the first object parameter if there is one
+	UObject* GetFirstObjectParam()
+	{
+		return ObjectParams.Num() > 0 ? ObjectParams[0].Object : nullptr;
+	}
+	
+	// Helper to get the first object parameter if there is one
+    const UObject* GetFirstObjectParam() const
+    {
+		return ObjectParams.Num() > 0 ? ObjectParams[0].Object : nullptr;
+    }
+
 	// Add a struct view Parameter to the Context
 	// the struct will be referred to by reference, and so must have a lifetime that is longer than this context
 	void AddStructViewParam(FStructView Param)
diff --git a/Engine/Plugins/Experimental/AnimNextPoseSearch/Source/AnimNextPoseSearch/Private/PoseSearchFeatureChannel_DistanceFromAnimNextVar.cpp b/Engine/Plugins/Experimental/AnimNextPoseSearch/Source/AnimNextPoseSearch/Private/PoseSearchFeatureChannel_DistanceFromAnimNextVar.cpp
--- a/Engine/Plugins/Experimental/AnimNextPoseSearch/Source/AnimNextPoseSearch/Private/PoseSearchFeatureChannel_DistanceFromAnimNextVar.cpp
+++ b/Engine/Plugins/Experimental/AnimNextPoseSearch/Source/AnimNextPoseSearch/Private/PoseSearchFeatureChannel_DistanceFromAnimNextVar.cpp
@@ -2,24 +2,39 @@
 
 #include "PoseSearchFeatureChannel_DistanceFromAnimNextVar.h"
 #include "Component/AnimNextComponent.h"
+#include "DataInterface/AnimNextDataInterfaceInstance.h"
 
 #include "PoseSearch/PoseSearchContext.h"
 #include "StructUtils/PropertyBag.h"
 
+
+namespace {
+// Helper function to find the first data interface instance in the context 
+FAnimNextDataInterfaceInstance* GetFirstDataInterfaceInstance(FChooserEvaluationContext& Context)
+{
+	for(const FStructView& Param : Context.Params)
+	{
+		if(Param.GetScriptStruct() == FAnimNextDataInterfaceInstance::StaticStruct())
+		{
+			return Param.GetPtr<FAnimNextDataInterfaceInstance>();
+		}
+	}
+
+	return nullptr;
+}
+}
+
+
 void UPoseSearchFeatureChannel_DistanceFromAnimNextVar::BuildQuery(UE::PoseSearch::FSearchContext& SearchContext) const
 {
 	using namespace UE::PoseSearch;
-	const UObject* Owner = SearchContext.GetAnimContext(SampleRole);
+	FChooserEvaluationContext* Context = SearchContext.GetContext(SampleRole);
 	
 	float Distance = 0.0f;
 
-	if (const UAnimNextComponent* Component = Cast<UAnimNextComponent>(Owner))
+	if (FAnimNextDataInterfaceInstance* Instance = GetFirstDataInterfaceInstance(*Context))
 	{
-		TOptional<float> Result = Component->GetVariableFloat(DistanceVariableName);
-		if (Result.IsSet())
-		{
-			Distance = Result.GetValue();
-		}
+		Instance->GetVariable(DistanceVariableName, Distance);
 	}
 	
 	FFeatureVectorHelper::EncodeFloat(SearchContext.EditFeatureVector(), ChannelDataOffset, Distance);
