Updated PoseSearchDatabase to warn about the requirements when using blendspace assets.

#rnx

diff --git a/Engine/Plugins/Animation/PoseSearch/Source/Editor/Private/PoseSearchDatabaseAssetListItem.cpp b/Engine/Plugins/Animation/PoseSearch/Source/Editor/Private/PoseSearchDatabaseAssetListItem.cpp
--- a/Engine/Plugins/Animation/PoseSearch/Source/Editor/Private/PoseSearchDatabaseAssetListItem.cpp
+++ b/Engine/Plugins/Animation/PoseSearch/Source/Editor/Private/PoseSearchDatabaseAssetListItem.cpp
@@ -40,6 +40,62 @@ namespace UE::PoseSearch
 {
 	static constexpr FLinearColor DisabledColor = FLinearColor(1.f, 1.f, 1.f, 0.25f);
 
+	static FText GetWarningsForDatabaseAsset(const FPoseSearchDatabaseAnimationAssetBase & InDatabaseAsset, const UPoseSearchDatabase & InDatabase)
+	{
+		if (InDatabaseAsset.GetAnimationAsset() == nullptr)
+		{
+			return LOCTEXT("ErrorNoAsset", "No asset has been selected.");
+		}
+		else if (!InDatabaseAsset.IsSkeletonCompatible(InDatabase.Schema))
+		{
+			return FText::Format(LOCTEXT("ErrorIncompatibleSkeleton", "{0}'s skeleton is not compatible with the schema's skeleton(s)."), FText::FromString(InDatabaseAsset.GetName()));
+		}
+		else if (const UBlendSpace* BlendSpace = Cast<UBlendSpace>(InDatabaseAsset.GetAnimationAsset()))
+		{
+			if (!BlendSpace->bShouldMatchSyncPhases)
+			{
+				return FText::Format(LOCTEXT("ErrorMissingSyncPhaseMatch", "{0}'s bShouldMatchSyncPhases flag is not enabled. This is required for properly pose matching blendspaces."), FText::FromString(InDatabaseAsset.GetName()));
+			}
+			else
+			{
+				const TArray<FBlendSample>& BlendSamples = BlendSpace->GetBlendSamples();
+				for (int i = 0; i < BlendSamples.Num() - 1; ++i)
+				{
+					const FBlendSample& CurrSample = BlendSamples[i];
+					const FBlendSample& NextSample = BlendSamples[i + 1];
+					
+					if (CurrSample.Animation && NextSample.Animation)
+					{
+						bool bWarning = false;
+						
+						if (CurrSample.Animation->AuthoredSyncMarkers.Num() == NextSample.Animation->AuthoredSyncMarkers.Num())
+						{
+							for (int j = 0; j < BlendSamples.Num(); ++j)
+							{
+								if (CurrSample.Animation->AuthoredSyncMarkers[j].MarkerName != NextSample.Animation->AuthoredSyncMarkers[j].MarkerName)
+								{
+									bWarning = true;
+									break;
+								}
+							}
+						}
+						else
+						{
+							bWarning = true;
+						}
+						
+						if (bWarning)
+						{
+							return FText::Format(LOCTEXT("ErrorDifferentNumOfSyncMarkers", "{0}'s samples don't share the same layout of sync markers. This is required for properly pose matching blendspaces."), FText::FromString(InDatabaseAsset.GetName()));
+						}
+					}
+				}
+			}
+		}
+
+		return FText::GetEmpty();
+	}
+	
 	/* We need a custom widget to be able to consume the "DoubleClick" event so we can cycle through the mirror options but not open the asset. */
 	class SMirrorTypeWidget : public SCompoundWidget
 	{
@@ -729,7 +785,7 @@ namespace UE::PoseSearch
 				{
 					if (DatabaseAnimationAssetBase->IsEnabled())
 					{
-						if (DatabaseAnimationAssetBase->GetAnimationAsset() == nullptr || !DatabaseAnimationAssetBase->IsSkeletonCompatible(Database->Schema))
+						if (!GetWarningsForDatabaseAsset(*DatabaseAnimationAssetBase, *Database).IsEmpty())
 						{
 							return EVisibility::Visible;
 						}
@@ -753,14 +809,7 @@ namespace UE::PoseSearch
 				{
 					if (DatabaseAnimationAssetBase->IsEnabled())
 					{
-						if (DatabaseAnimationAssetBase->GetAnimationAsset() == nullptr)
-						{
-							return LOCTEXT("ErrorNoAsset", "No asset has been selected.");
-						}
-						else if (!DatabaseAnimationAssetBase->IsSkeletonCompatible(Database->Schema))
-						{
-							return LOCTEXT("ErrorIncompatibleSkeleton", "This asset's skeleton is not compatible with the schema's skeleton(s).");
-						}
+						return GetWarningsForDatabaseAsset(*DatabaseAnimationAssetBase, *Database);
 					}
 				}
 			}
@@ -864,7 +913,7 @@ namespace UE::PoseSearch
 				{
 					if (DatabaseAnimationAssetBase->IsEnabled())
 					{
-						if (DatabaseAnimationAssetBase->GetAnimationAsset() == nullptr || !DatabaseAnimationAssetBase->IsSkeletonCompatible(Database->Schema))
+						if (!GetWarningsForDatabaseAsset(*DatabaseAnimationAssetBase, *Database).IsEmpty())
 						{
 							return FColor::Red;
 						}
diff --git a/Engine/Plugins/Animation/PoseSearch/Source/Runtime/Private/PoseSearchAssetSampler.cpp b/Engine/Plugins/Animation/PoseSearch/Source/Runtime/Private/PoseSearchAssetSampler.cpp
--- a/Engine/Plugins/Animation/PoseSearch/Source/Runtime/Private/PoseSearchAssetSampler.cpp
+++ b/Engine/Plugins/Animation/PoseSearch/Source/Runtime/Private/PoseSearchAssetSampler.cpp
@@ -393,7 +393,7 @@ void FAnimationAssetSampler::ExtractPose(const FAnimExtractContext& ExtractionCt
 		int32 TriangulationIndex = 0;
 		if (BlendSpace->GetSamplesFromBlendInput(BlendParameters, BlendSamples, TriangulationIndex, true))
 		{
-			BlendSpace->ResetBlendSamples(BlendSamples, ToNormalizedTime(ExtractionCtx.CurrentTime), ExtractionCtx.bLooping, false);
+			BlendSpace->ResetBlendSamples(BlendSamples, ToNormalizedTime(ExtractionCtx.CurrentTime), ExtractionCtx.bLooping, true);
 			BlendSpace->GetAnimationPose(BlendSamples, ExtractionCtx, OutAnimPoseData);
 		}
 		else
diff --git a/Engine/Plugins/Animation/PoseSearch/Source/Runtime/Private/PoseSearchDerivedData.cpp b/Engine/Plugins/Animation/PoseSearch/Source/Runtime/Private/PoseSearchDerivedData.cpp
--- a/Engine/Plugins/Animation/PoseSearch/Source/Runtime/Private/PoseSearchDerivedData.cpp
+++ b/Engine/Plugins/Animation/PoseSearch/Source/Runtime/Private/PoseSearchDerivedData.cpp
@@ -1844,6 +1844,63 @@ void FPoseSearchDatabaseAsyncCacheTask::OnGetComplete(UE::DerivedData::FCacheGet
 						}
 					}
 
+					for (int32 AnimationAssetIndex = 0; AnimationAssetIndex < DependentDatabase->GetNumAnimationAssets(); ++AnimationAssetIndex)
+					{
+						if (const FPoseSearchDatabaseAnimationAssetBase* DatabaseAsset = DependentDatabase->GetDatabaseAnimationAsset<FPoseSearchDatabaseAnimationAssetBase>(AnimationAssetIndex))
+						{
+							if (DatabaseAsset->GetAnimationAsset() == nullptr)
+							{
+								UE_LOG(LogPoseSearch, Warning, TEXT("OnGetComplete - - No asset has been selected."));
+							}
+							else if (!DatabaseAsset->IsSkeletonCompatible(Database->Schema))
+							{
+								UE_LOG(LogPoseSearch, Warning, TEXT("OnGetComplete - %s's skeleton is not compatible with the schema's skeleton(s)."), *DatabaseAsset->GetName());
+							}
+							else if (const UBlendSpace* BlendSpace = Cast<UBlendSpace>(DatabaseAsset->GetAnimationAsset()))
+							{
+								if (!BlendSpace->bShouldMatchSyncPhases)
+								{
+									UE_LOG(LogPoseSearch, Warning, TEXT("OnGetComplete - %s's bShouldMatchSyncPhases flag is not enabled. This is required for properly pose matching blendspaces."), *DatabaseAsset->GetName());
+								}
+								else
+								{
+									const TArray<FBlendSample>& BlendSamples = BlendSpace->GetBlendSamples();
+									for (int i = 0; i < BlendSamples.Num() - 1; ++i)
+									{
+										const FBlendSample& CurrSample = BlendSamples[i];
+										const FBlendSample& NextSample = BlendSamples[i + 1];
+										
+										if (CurrSample.Animation && NextSample.Animation)
+										{
+											bool bWarning = false;
+											
+											if (CurrSample.Animation->AuthoredSyncMarkers.Num() == NextSample.Animation->AuthoredSyncMarkers.Num())
+											{
+												for (int j = 0; j < BlendSamples.Num(); ++j)
+												{
+													if (CurrSample.Animation->AuthoredSyncMarkers[j].MarkerName != NextSample.Animation->AuthoredSyncMarkers[j].MarkerName)
+													{
+														bWarning = true;
+														break;
+													}
+												}
+											}
+											else
+											{
+												bWarning = true;
+											}
+											
+											if (bWarning)
+											{
+												UE_LOG(LogPoseSearch, Warning, TEXT("OnGetComplete - %s's samples don't share the same layout of sync markers. This is required for properly pose matching blendspaces."), *DatabaseAsset->GetName());
+											}
+										}
+									}
+								}
+							}
+						}
+					}
+					
 					if (Owner.IsCanceled())
 					{
 						UE_LOG(LogPoseSearch, Log, TEXT("%s - %s BuildIndex Cancelled"), *LexToString(FullIndexKey.Hash), *MainDatabaseName);
diff --git a/Engine/Plugins/Animation/PoseSearch/Source/Runtime/Public/PoseSearch/PoseSearchDefines.h b/Engine/Plugins/Animation/PoseSearch/Source/Runtime/Public/PoseSearch/PoseSearchDefines.h
--- a/Engine/Plugins/Animation/PoseSearch/Source/Runtime/Public/PoseSearch/PoseSearchDefines.h
+++ b/Engine/Plugins/Animation/PoseSearch/Source/Runtime/Public/PoseSearch/PoseSearchDefines.h
@@ -42,7 +42,7 @@ static constexpr int32 PreallocatedSearchesNum = 4;
 static constexpr float FiniteDelta = 1 / 60.0f; // Time delta used for computing pose derivatives
 static constexpr int32 MaxNumberOfCollectedPoseCandidatesPerDatabase = 256;
 static constexpr int32 InvalidSearchId = INDEX_NONE; // @todo: it should be 0, since the search id is an hash!
-static constexpr int32 DatabaseIndexDerivedDataCacheKeyVersion = 2688; // bump this value to force reindex ALL the UPoseSearchDatabases, without having to change the global FDevSystemGuids::POSESEARCHDB_DERIVEDDATA_VER
+static constexpr int32 DatabaseIndexDerivedDataCacheKeyVersion = 2691; // bump this value to force reindex ALL the UPoseSearchDatabases, without having to change the global FDevSystemGuids::POSESEARCHDB_DERIVEDDATA_VER
 static constexpr int32 DefaultTickPriority = 100;
 
 template<uint32 Alignment = DEFAULT_ALIGNMENT> class TMemStackSetAllocator : public TSetAllocator<TSparseArrayAllocator<TMemStackAllocator<Alignment>, TMemStackAllocator<Alignment>>, TMemStackAllocator<Alignment>> {};
