Fixed crash while building PSD due to trying to load/access UObject during GC. Also, fixed crash due to array out of bounds for both warnings for PSD during building of indexes and during database editor UI construction.

#rnx

diff --git a/Engine/Plugins/Animation/PoseSearch/Source/Editor/Private/PoseSearchDatabaseAssetListItem.cpp b/Engine/Plugins/Animation/PoseSearch/Source/Editor/Private/PoseSearchDatabaseAssetListItem.cpp
--- a/Engine/Plugins/Animation/PoseSearch/Source/Editor/Private/PoseSearchDatabaseAssetListItem.cpp
+++ b/Engine/Plugins/Animation/PoseSearch/Source/Editor/Private/PoseSearchDatabaseAssetListItem.cpp
@@ -70,7 +70,7 @@ namespace UE::PoseSearch
 						
 						if (CurrSample.Animation->AuthoredSyncMarkers.Num() == NextSample.Animation->AuthoredSyncMarkers.Num())
 						{
-							for (int j = 0; j < BlendSamples.Num(); ++j)
+							for (int j = 0; j < CurrSample.Animation->AuthoredSyncMarkers.Num(); ++j)
 							{
 								if (CurrSample.Animation->AuthoredSyncMarkers[j].MarkerName != NextSample.Animation->AuthoredSyncMarkers[j].MarkerName)
 								{
diff --git a/Engine/Plugins/Animation/PoseSearch/Source/Runtime/Private/PoseSearchDerivedData.cpp b/Engine/Plugins/Animation/PoseSearch/Source/Runtime/Private/PoseSearchDerivedData.cpp
--- a/Engine/Plugins/Animation/PoseSearch/Source/Runtime/Private/PoseSearchDerivedData.cpp
+++ b/Engine/Plugins/Animation/PoseSearch/Source/Runtime/Private/PoseSearchDerivedData.cpp
@@ -1843,7 +1843,7 @@ void FPoseSearchDatabaseAsyncCacheTask::OnGetComplete(UE::DerivedData::FCacheGet
 							}
 						}
 					}
-
+					
 					for (int32 AnimationAssetIndex = 0; AnimationAssetIndex < DependentDatabase->GetNumAnimationAssets(); ++AnimationAssetIndex)
 					{
 						if (const FPoseSearchDatabaseAnimationAssetBase* DatabaseAsset = DependentDatabase->GetDatabaseAnimationAsset<FPoseSearchDatabaseAnimationAssetBase>(AnimationAssetIndex))
@@ -1852,10 +1852,11 @@ void FPoseSearchDatabaseAsyncCacheTask::OnGetComplete(UE::DerivedData::FCacheGet
 							{
 								UE_LOG(LogPoseSearch, Warning, TEXT("OnGetComplete - - No asset has been selected."));
 							}
-							else if (!DatabaseAsset->IsSkeletonCompatible(Database->Schema))
+							// @todo: Find a way to prevent accessing UObjects while GC since this func is ran async. Commenting out for now to prevent crashes.
+							/*else if (!DatabaseAsset->IsSkeletonCompatible(Database->Schema))
 							{
 								UE_LOG(LogPoseSearch, Warning, TEXT("OnGetComplete - %s's skeleton is not compatible with the schema's skeleton(s)."), *DatabaseAsset->GetName());
-							}
+							}*/
 							else if (const UBlendSpace* BlendSpace = Cast<UBlendSpace>(DatabaseAsset->GetAnimationAsset()))
 							{
 								if (!BlendSpace->bShouldMatchSyncPhases)
@@ -1876,7 +1877,7 @@ void FPoseSearchDatabaseAsyncCacheTask::OnGetComplete(UE::DerivedData::FCacheGet
 											
 											if (CurrSample.Animation->AuthoredSyncMarkers.Num() == NextSample.Animation->AuthoredSyncMarkers.Num())
 											{
-												for (int j = 0; j < BlendSamples.Num(); ++j)
+												for (int j = 0; j < CurrSample.Animation->AuthoredSyncMarkers.Num(); ++j)
 												{
 													if (CurrSample.Animation->AuthoredSyncMarkers[j].MarkerName != NextSample.Animation->AuthoredSyncMarkers[j].MarkerName)
 													{
