PoseSearch - fix for MM events going out of bound if placed into a do not choose section

diff --git a/Engine/Plugins/Animation/PoseSearch/Source/Runtime/Private/PoseSearchAssetIndexer.cpp b/Engine/Plugins/Animation/PoseSearch/Source/Runtime/Private/PoseSearchAssetIndexer.cpp
--- a/Engine/Plugins/Animation/PoseSearch/Source/Runtime/Private/PoseSearchAssetIndexer.cpp
+++ b/Engine/Plugins/Animation/PoseSearch/Source/Runtime/Private/PoseSearchAssetIndexer.cpp
@@ -289,7 +289,10 @@ void FAssetIndexer::Process(int32 AssetIdx)
 					if (EventNotify->EventTag.IsValid())
 					{
 						const int32 PoseIdx = SearchIndexAsset.GetPoseIndexFromTime(AnimNotifyEvent.GetTime(), Schema.SampleRate);
-						EventDataCollector.Emplace(EventNotify->EventTag, PoseIdx);
+						if (PoseIdx != INDEX_NONE)
+						{
+							EventDataCollector.Emplace(EventNotify->EventTag, PoseIdx);
+						}
 					}
 				}
 			}
diff --git a/Engine/Plugins/Animation/PoseSearch/Source/Runtime/Private/PoseSearchIndex.cpp b/Engine/Plugins/Animation/PoseSearch/Source/Runtime/Private/PoseSearchIndex.cpp
--- a/Engine/Plugins/Animation/PoseSearch/Source/Runtime/Private/PoseSearchIndex.cpp
+++ b/Engine/Plugins/Animation/PoseSearch/Source/Runtime/Private/PoseSearchIndex.cpp
@@ -316,6 +316,7 @@ FArchive& operator<<(FArchive& Ar, FSearchStats& Stats)
 #if WITH_EDITOR
 void FEventDataCollector::Emplace(const FGameplayTag& EventTag, int32 PoseIdx)
 {
+	check(PoseIdx >= 0);
 	Data.FindOrAdd(EventTag).Add(PoseIdx);
 }
 
@@ -328,6 +329,7 @@ void FEventDataCollector::MergeWith(const FEventDataCollector& Other)
 		// adding all the missing poses of OtherTagToPoseIndexes to PoseIndexes
 		for (int32 OtherPoseIdx : OtherTagToPoseIndexes.Value)
 		{
+			check(OtherPoseIdx >= 0);
 			PoseIndexes.Add(OtherPoseIdx);
 		}
 	}
