PoseSearch - fix for UPoseSearchAssetSamplerLibrary::GetTransformByName wrongly spamming UPoseSearchAssetSamplerLibrary::GetTransformByName invalid FCompactPoseBoneIndex for BoneName...

diff --git a/Engine/Plugins/Animation/PoseSearch/Source/Runtime/Private/PoseSearchAssetSamplerLibrary.cpp b/Engine/Plugins/Animation/PoseSearch/Source/Runtime/Private/PoseSearchAssetSamplerLibrary.cpp
--- a/Engine/Plugins/Animation/PoseSearch/Source/Runtime/Private/PoseSearchAssetSamplerLibrary.cpp
+++ b/Engine/Plugins/Animation/PoseSearch/Source/Runtime/Private/PoseSearchAssetSamplerLibrary.cpp
@@ -78,7 +78,7 @@ FTransform UPoseSearchAssetSamplerLibrary::GetTransformByName(UPARAM(ref) FPoseS
 	}
 
 	const FCompactPoseBoneIndex CompactPoseBoneIndex = BoneContainer.GetCompactPoseIndexFromSkeletonPoseIndex(FSkeletonPoseBoneIndex(BoneReference.BoneIndex));
-	if (CompactPoseBoneIndex.IsValid())
+	if (!CompactPoseBoneIndex.IsValid())
 	{
 		UE_LOG(LogPoseSearch, Error, TEXT("UPoseSearchAssetSamplerLibrary::GetTransformByName invalid FCompactPoseBoneIndex for BoneName %s for Skeleton %s"), *BoneName.ToString(), *GetNameSafe(Skeleton));
 		return FTransform::Identity;
