PoseSearch - UPoseSearchInteractionSubsystem::AddAvailabilities improvements to avoid adding invalid FInteractionAvailability to AvailabilitiesMap unnecessarily slowing down UPoseSearchInteractionSubsystem::Query_AnyThread

diff --git a/Engine/Plugins/Animation/PoseSearch/Source/Runtime/Private/PoseSearchInteractionSubsystem.cpp b/Engine/Plugins/Animation/PoseSearch/Source/Runtime/Private/PoseSearchInteractionSubsystem.cpp
--- a/Engine/Plugins/Animation/PoseSearch/Source/Runtime/Private/PoseSearchInteractionSubsystem.cpp
+++ b/Engine/Plugins/Animation/PoseSearch/Source/Runtime/Private/PoseSearchInteractionSubsystem.cpp
@@ -494,30 +494,36 @@ void UPoseSearchInteractionSubsystem::Deinitialize()
 	Super::Deinitialize();
 }
 
-bool UPoseSearchInteractionSubsystem::IsThereAnyAvailability() const
-{
-	check(IsInGameThread());
-	return !AvailabilitiesMap.IsEmpty();
-}
-
 void UPoseSearchInteractionSubsystem::AddAvailabilities(const TArrayView<const FPoseSearchInteractionAvailability> Availabilities, const UObject* AnimContext, FName PoseHistoryName, const UE::PoseSearch::IPoseHistory* PoseHistory)
 {
-	FScopeLock Lock(&AvailabilitiesMapMutex);
-	AvailabilitiesMap.AddAvailabilities(Availabilities, AnimContext, PoseHistoryName, PoseHistory);
-}
+	using namespace UE::PoseSearch;
 
-#if ENABLE_DRAW_DEBUG && ENABLE_ANIM_DEBUG
-const UE::PoseSearch::FInteractionAvailabilities* UPoseSearchInteractionSubsystem::FindAvailabilities(const UObject* AnimContext) const
-{
-	check(IsInGameThread());
-	return AvailabilitiesMap.Find(AnimContext);
-}
-#endif // ENABLE_DRAW_DEBUG && ENABLE_ANIM_DEBUG
+	// collecting valid availabilities indexes here to minimize the time spent on the lock and avoid locking at all 
+	// if there're actually no valid availabilities (as well as avoiding polluting AvailabilitiesMap with an empty entry).
+	// This is unfortunately a common setup, and that's the main reason we optimize against it
+	TArray<int32, TInlineAllocator<64>> ValidAvailabilityIndexes;
+	for (int32 AvailabilityIndex = 0; AvailabilityIndex < Availabilities.Num(); ++AvailabilityIndex)
+	{
+		const FPoseSearchInteractionAvailability& Availability = Availabilities[AvailabilityIndex];
+		if (Availability.Database || Availability.IsTagValid())
+		{
+			ValidAvailabilityIndexes.Add(AvailabilityIndex);
+		}
+	}
 
-void UPoseSearchInteractionSubsystem::ResetAvailabilities()
-{
-	check(IsInGameThread());
-	AvailabilitiesMap.Reset();
+	if (!ValidAvailabilityIndexes.IsEmpty())
+	{
+		FScopeLock Lock(&AvailabilitiesMapMutex);
+		FInteractionAvailabilities& InteractionAvailabilities = AvailabilitiesMap.FindOrAdd(AnimContext);
+		for (const int32 AvailabilityIndex : ValidAvailabilityIndexes)
+		{
+			// Avoiding adding trivial duplicates. FInteractionAvailability could not be fully specified to understand if it's an actual duplicate in case the pose
+			// history is passed by name or the Availability.Database is null and supposed to be resolved using other availabilities Database(s) with the same Availability.Tag.
+			// The duplicated availabilities are excluded when creating the combinations of possible interactions during FAnimContextInfoVisitor when 
+			// FRoledAnimContextInfos.AddRoledAnimContextInfos calls FRoledAnimContextInfos.AddUnique
+			InteractionAvailabilities.AddUnique(FInteractionAvailability(Availabilities[AvailabilityIndex], PoseHistoryName, PoseHistory));
+		}
+	}
 }
 
 void UPoseSearchInteractionSubsystem::GenerateAnimContextInfosAndTagToDatabases(UE::PoseSearch::FAnimContextInfos& AnimContextInfos, UE::PoseSearch::FTagToDatabases& TagToDatabases) const
@@ -1112,7 +1118,7 @@ void UPoseSearchInteractionSubsystem::Tick(float DeltaSeconds)
 	FMemMark Mark(FMemStack::Get());
 	UpdateValidInteractionSearches();
 
-	if (!IsThereAnyAvailability())
+	if (AvailabilitiesMap.IsEmpty())
 	{
 		bool bAllUninjected = true;
 		for (FInteractionIsland* Island : Islands)
@@ -1153,7 +1159,7 @@ void UPoseSearchInteractionSubsystem::Tick(float DeltaSeconds)
 	RegenerateAllIslands(SearchContexts);
 
 	// clearing up AvailabilitiesMap for the next frame
-	ResetAvailabilities();
+	AvailabilitiesMap.Reset();
 
 #if DO_CHECK
 	check(ValidateAllIslands());
@@ -1209,10 +1215,7 @@ void UPoseSearchInteractionSubsystem::Query_AnyThread(const TArrayView<const FPo
 	}
 
 	// queuing the availabilities for the next frame Query_AnyThread
-	if (!Availabilities.IsEmpty())
-	{
-		AddAvailabilities(Availabilities, AnimContext, PoseHistoryName, PoseHistory);
-	}
+	AddAvailabilities(Availabilities, AnimContext, PoseHistoryName, PoseHistory);
 }
 
 void UPoseSearchInteractionSubsystem::GetResult_AnyThread(const UObject* AnimContext, FPoseSearchBlueprintResult& Result, bool bCompareOwningActors)
@@ -1272,7 +1275,7 @@ void UPoseSearchInteractionSubsystem::DebugDrawIslands() const
 				{
 					if (const UObject* AnimContext = SearchContext.AnimContexts[Index].Get())
 					{
-						if (const FInteractionAvailabilities* Availabilities = FindAvailabilities(AnimContext))
+						if (const FInteractionAvailabilities* Availabilities = AvailabilitiesMap.Find(AnimContext))
 						{
 							float MaxBroadPhaseRadius = 0.f;
 							for (const FInteractionAvailability& Availability : *Availabilities)
diff --git a/Engine/Plugins/Animation/PoseSearch/Source/Runtime/Public/PoseSearch/PoseSearchInteractionSubsystem.h b/Engine/Plugins/Animation/PoseSearch/Source/Runtime/Public/PoseSearch/PoseSearchInteractionSubsystem.h
--- a/Engine/Plugins/Animation/PoseSearch/Source/Runtime/Public/PoseSearch/PoseSearchInteractionSubsystem.h
+++ b/Engine/Plugins/Animation/PoseSearch/Source/Runtime/Public/PoseSearch/PoseSearchInteractionSubsystem.h
@@ -26,13 +26,6 @@ struct FInteractionAvailability : public FPoseSearchInteractionAvailability
 		check(!PoseHistory || PoseHistory->AsWeak().IsValid());
 	}
 
-	bool operator==(const FInteractionAvailability& Other) const
-	{
-		return static_cast<const FPoseSearchInteractionAvailability&>(*this) == static_cast<const FPoseSearchInteractionAvailability&>(Other)
-			&& PoseHistoryName == Other.PoseHistoryName
-			&& PoseHistory == Other.PoseHistory;
-	}
-
 	FString GetPoseHistoryName() const;
 	const IPoseHistory* GetPoseHistory(const UObject* AnimContext) const;
 
@@ -44,56 +37,7 @@ private:
 // Experimental, this feature might be removed without warning, not for production use
 typedef TArray<FInteractionAvailability> FInteractionAvailabilities;
 // Experimental, this feature might be removed without warning, not for production use
-struct FAvailabilitiesMap : protected TMap<TWeakObjectPtr<const UObject>, FInteractionAvailabilities>
-{
-	using Super = TMap<TWeakObjectPtr<const UObject>, FInteractionAvailabilities>;
-	using Super::TConstIterator;
-
-	using Super::CreateConstIterator;
-	using Super::IsEmpty;
-	using Super::Find;
-	using Super::Reset;
-	
-	void AddAvailabilities(const TArrayView<const FPoseSearchInteractionAvailability> Availabilities, const UObject* AnimContext, FName PoseHistoryName, const IPoseHistory* PoseHistory)
-	{
-		// @todo: re-enable this logging for editor only builds once the original issue has been fixed
-//#if WITH_EDITOR && DO_CHECK
-		check(AnimContext);
-		const AActor* ContextOwningActor = GetContextOwningActor(AnimContext);
-		check(ContextOwningActor);
-
-		for (const TPair<TWeakObjectPtr<const UObject>, FInteractionAvailabilities>& AvailabilitiesMapPair : *this)
-		{
-			if (const UObject* OtherAnimContext = AvailabilitiesMapPair.Key.Get())
-			{
-				if (const AActor* OtherContextOwningActor = GetContextOwningActor(OtherAnimContext))
-				{
-					// same actor but different AnimContext? Do we have multiple AddAvailabilities from different components?
-					// It's not currently supported! and if not filtered out it could potentially add interactions within the same actor...
-					if (ContextOwningActor == OtherContextOwningActor && AnimContext != OtherAnimContext)
-					{
-						//UE_LOG(LogPoseSearch, Error, TEXT("FAvailabilitiesMap::AddAvailabilities from different AnimContexts %s, %s sharing the same Actor %s. This is not currently supported!"), *AnimContext->GetName(), *OtherAnimContext->GetName(), *ContextOwningActor->GetName());
-						return;
-					}
-				}
-			}
-		}
-//#endif // WITH_EDITOR && DO_CHECK
-
-		FInteractionAvailabilities& AnimContextAvailabilities = FindOrAdd(AnimContext);
-		for (const FPoseSearchInteractionAvailability& Availability : Availabilities)
-		{
-			if (Availability.Database || Availability.IsTagValid())
-			{
-				// Avoiding adding trivial duplicates. FInteractionAvailability could not be fully specified to understand if it's an actual duplicate in case the pose
-				// history is passed by name or the Availability.Database is null and supposed to be resolved using other availabilities Database(s) with the same Availability.Tag.
-				// The duplicated availabilities are excluded when creating the combinations of possible interactions during FAnimContextInfoVisitor when 
-				// FRoledAnimContextInfos.AddRoledAnimContextInfos calls FRoledAnimContextInfos.AddUnique
-				AnimContextAvailabilities.AddUnique(FInteractionAvailability(Availability, PoseHistoryName, PoseHistory));
-			}
-		}
-	}
-};
+typedef TMap<TWeakObjectPtr<const UObject>, FInteractionAvailabilities> FAvailabilitiesMap;
 
 // Experimental, this feature might be removed without warning, not for production use
 typedef TArray<TPair<TWeakObjectPtr<AActor>, TWeakObjectPtr<AActor>>> FDisabledCollisions;
@@ -205,16 +149,13 @@ private:
 	
 #if ENABLE_DRAW_DEBUG && ENABLE_ANIM_DEBUG
 	void DebugDrawIslands() const;
-	const UE::PoseSearch::FInteractionAvailabilities* FindAvailabilities(const UObject* AnimContext) const;
 #endif // ENABLE_DRAW_DEBUG && ENABLE_ANIM_DEBUG
 	
 #if ENABLE_ANIM_DEBUG
 	void DebugLogTickDependencies() const;
 #endif // ENABLE_ANIM_DEBUG
 
-	bool IsThereAnyAvailability() const;
 	void AddAvailabilities(const TArrayView<const FPoseSearchInteractionAvailability> Availabilities, const UObject* AnimContext, FName PoseHistoryName, const UE::PoseSearch::IPoseHistory* PoseHistory);
-	void ResetAvailabilities();
 	void GenerateAnimContextInfosAndTagToDatabases(UE::PoseSearch::FAnimContextInfos& AnimContextInfos, UE::PoseSearch::FTagToDatabases& TagToDatabases) const;
 	void GenerateSearchContexts(UE::PoseSearch::FInteractionSearchContexts& SearchContexts) const;
 
@@ -224,7 +165,7 @@ private:
 	static void OnInteractionEnd(UE::PoseSearch::FValidInteractionSearch& ValidInteractionSearch);
 
 	// @todo: use a thread local storage data structure to make publishing availabilities lockless to avoid the need of AvailabilitiesMapMutex
-	// it could contain duplicated avaialbilities
+	// it could contain duplicated availabilities
 	UE::PoseSearch::FAvailabilitiesMap AvailabilitiesMap;
 	FCriticalSection AvailabilitiesMapMutex;
 
