Fallback codepath for pose search chooser column in case no assets pass pose search filters

diff --git a/Engine/Plugins/Animation/PoseSearch/Source/Runtime/Private/Chooser/PoseSearchChooserColumn.cpp b/Engine/Plugins/Animation/PoseSearch/Source/Runtime/Private/Chooser/PoseSearchChooserColumn.cpp
--- a/Engine/Plugins/Animation/PoseSearch/Source/Runtime/Private/Chooser/PoseSearchChooserColumn.cpp
+++ b/Engine/Plugins/Animation/PoseSearch/Source/Runtime/Private/Chooser/PoseSearchChooserColumn.cpp
@@ -117,6 +117,29 @@ void FPoseSearchColumn::Filter(FChooserEvaluationContext& Context, const FChoose
 			}
 		}
 	}
+
+	if (IndexListOut.IsEmpty())
+	{
+		// if nothing passed the pose match, or if PoseHistory was not found, then ignore pose matching and just pass through everything
+		IndexListOut = IndexListIn;
+
+		// set some reasonable defaults for our other outputs in the passthrough  case
+		if (const FChooserParameterFloatBase* StartTime = OutputStartTime.GetPtr<FChooserParameterFloatBase>())
+		{
+			StartTime->SetValue(Context, 0.0);
+		}
+		
+		if (const FChooserParameterFloatBase* Cost = OutputCost.GetPtr<FChooserParameterFloatBase>())
+		{
+			// set an arbitrary "high" cost value so that cost threshold implementations can still wait for the pose search to find a good match
+			Cost->SetValue(Context, 100.0);
+		}
+
+		if (const FChooserParameterBoolBase* Mirror = OutputMirror.GetPtr<FChooserParameterBoolBase>())
+		{
+			Mirror->SetValue(Context, false);
+		}
+	}
 }
 
 #if WITH_EDITOR
