PoseSearch - converting TMap<TWeakObjectPtr<const UObject> ...> to faster TMap<FObjectKey, ...>

diff --git a/Engine/Plugins/Animation/PoseSearch/Source/Runtime/Private/PoseSearchDatabase.cpp b/Engine/Plugins/Animation/PoseSearch/Source/Runtime/Private/PoseSearchDatabase.cpp
--- a/Engine/Plugins/Animation/PoseSearch/Source/Runtime/Private/PoseSearchDatabase.cpp
+++ b/Engine/Plugins/Animation/PoseSearch/Source/Runtime/Private/PoseSearchDatabase.cpp
@@ -782,7 +782,7 @@ void UPoseSearchDatabase::UpdateCachedProperties()
 		}
 	}
 
-	for (TPair<TWeakObjectPtr<UObject>, TArray<int32>>& CachedAssetMapPair : CachedAssetMap)
+	for (TPair<FObjectKey, TArray<int32>>& CachedAssetMapPair : CachedAssetMap)
 	{
 		CachedAssetMapPair.Value.Sort();
 	}
diff --git a/Engine/Plugins/Animation/PoseSearch/Source/Runtime/Private/PoseSearchInteractionSubsystem.cpp b/Engine/Plugins/Animation/PoseSearch/Source/Runtime/Private/PoseSearchInteractionSubsystem.cpp
--- a/Engine/Plugins/Animation/PoseSearch/Source/Runtime/Private/PoseSearchInteractionSubsystem.cpp
+++ b/Engine/Plugins/Animation/PoseSearch/Source/Runtime/Private/PoseSearchInteractionSubsystem.cpp
@@ -542,7 +542,7 @@ void UPoseSearchInteractionSubsystem::GenerateAnimContextInfosAndTagToDatabases(
 	for (FAvailabilitiesMap::TConstIterator Iter = AvailabilitiesMap.CreateConstIterator(); Iter; ++Iter)
 	{
 		// making sure the weak pointer is still valid
-		const UObject* AnimContext = Iter->Key.Get();
+		const UObject* AnimContext = Iter->Key.ResolveObjectPtr();
 		if (!AnimContext)
 		{
 			UE_LOG(LogPoseSearch, Log, TEXT("UPoseSearchInteractionSubsystem::GenerateAnimContextInfosAndTagToDatabaseMap null anim context. The associated character got removed from the world"));
diff --git a/Engine/Plugins/Animation/PoseSearch/Source/Runtime/Public/PoseSearch/PoseSearchDatabase.h b/Engine/Plugins/Animation/PoseSearch/Source/Runtime/Public/PoseSearch/PoseSearchDatabase.h
--- a/Engine/Plugins/Animation/PoseSearch/Source/Runtime/Public/PoseSearch/PoseSearchDatabase.h
+++ b/Engine/Plugins/Animation/PoseSearch/Source/Runtime/Public/PoseSearch/PoseSearchDatabase.h
@@ -521,7 +521,7 @@ private:
 	UE::PoseSearch::FSearchIndex SearchIndexPrivate;
 	
 	// CachedAssetMap is NOT serialized in operator<< but recalculated by UpdateCachedProperties every time SearchIndexPrivate changes
-	TMap<TWeakObjectPtr<UObject>, TArray<int32>> CachedAssetMap;
+	TMap<FObjectKey, TArray<int32>> CachedAssetMap;
 
 #if WITH_EDITOR
 	DECLARE_MULTICAST_DELEGATE(FOnDerivedDataRebuildMulticaster);
diff --git a/Engine/Plugins/Animation/PoseSearch/Source/Runtime/Public/PoseSearch/PoseSearchInteractionSubsystem.h b/Engine/Plugins/Animation/PoseSearch/Source/Runtime/Public/PoseSearch/PoseSearchInteractionSubsystem.h
--- a/Engine/Plugins/Animation/PoseSearch/Source/Runtime/Public/PoseSearch/PoseSearchInteractionSubsystem.h
+++ b/Engine/Plugins/Animation/PoseSearch/Source/Runtime/Public/PoseSearch/PoseSearchInteractionSubsystem.h
@@ -37,7 +37,7 @@ private:
 // Experimental, this feature might be removed without warning, not for production use
 typedef TArray<FInteractionAvailability> FInteractionAvailabilities;
 // Experimental, this feature might be removed without warning, not for production use
-typedef TMap<TWeakObjectPtr<const UObject>, FInteractionAvailabilities> FAvailabilitiesMap;
+typedef TMap<FObjectKey, FInteractionAvailabilities> FAvailabilitiesMap;
 
 // Experimental, this feature might be removed without warning, not for production use
 typedef TArray<TPair<TWeakObjectPtr<AActor>, TWeakObjectPtr<AActor>>> FDisabledCollisions;
