Fixed for root space in stitch track in sequencer since most anims don't mark Root Lock

diff --git a/Engine/Plugins/Animation/PoseSearch/Source/Runtime/Private/PoseSearchHistory.cpp b/Engine/Plugins/Animation/PoseSearch/Source/Runtime/Private/PoseSearchHistory.cpp
--- a/Engine/Plugins/Animation/PoseSearch/Source/Runtime/Private/PoseSearchHistory.cpp
+++ b/Engine/Plugins/Animation/PoseSearch/Source/Runtime/Private/PoseSearchHistory.cpp
@@ -1199,7 +1199,7 @@ void FMemStackPoseHistory::AddFuturePose(float Time, IComponentSpacePoseProvider
 	FutureEntries.InsertDefaulted_GetRef(LowerBoundIdx).Update(Time, ComponentSpacePoseProvider, GetBoneToTransformMap(), true, Curves, MakeConstArrayView(GetCollectedCurves()));
 }
 
-void FMemStackPoseHistory::ExtractAndAddFuturePoses(const UAnimationAsset* AnimationAsset, float AnimationTime, float FiniteDeltaTime, const FVector& BlendParame1ters, float IntervalTime, const USkeleton* OverrideSkeleton)
+void FMemStackPoseHistory::ExtractAndAddFuturePoses(const UAnimationAsset* AnimationAsset, float AnimationTime, float FiniteDeltaTime, const FVector& BlendParame1ters, float IntervalTime, const USkeleton* OverrideSkeleton, bool bUseRefPoseRootBone)
 {
 	FMemMark Mark(FMemStack::Get());
 
@@ -1258,6 +1258,12 @@ void FMemStackPoseHistory::ExtractAndAddFuturePoses(const UAnimationAsset* Anima
 		const float FuturePoseAnimationTime = IntervalTime + (i - 1) * FiniteDeltaTime;
 
 		Sampler.ExtractPose(FuturePoseExtractionTime, Pose, Curves);
+
+		if (bUseRefPoseRootBone)
+		{
+			Pose[FCompactPoseBoneIndex(RootBoneIndexType)] = Skeleton->GetReferenceSkeleton().GetRefBonePose()[RootBoneIndexType];
+		}
+
 		ComponentSpacePose.InitPose(Pose);
 		FComponentSpacePoseProvider ComponentSpacePoseProvider(ComponentSpacePose);
 		AddFuturePose(FuturePoseAnimationTime, ComponentSpacePoseProvider);
diff --git a/Engine/Plugins/Animation/PoseSearch/Source/Runtime/Public/PoseSearch/PoseSearchHistory.h b/Engine/Plugins/Animation/PoseSearch/Source/Runtime/Public/PoseSearch/PoseSearchHistory.h
--- a/Engine/Plugins/Animation/PoseSearch/Source/Runtime/Public/PoseSearch/PoseSearchHistory.h
+++ b/Engine/Plugins/Animation/PoseSearch/Source/Runtime/Public/PoseSearch/PoseSearchHistory.h
@@ -285,7 +285,7 @@ struct POSESEARCH_API FMemStackPoseHistory : public IPoseHistory
 	// Experimental, this feature might be removed without warning, not for production use
 	// extracting a pose from AnimationAsset at AnimationTime and  (using BlendParameters in case it's a blend space) and placing it in the FMemStackPoseHistory at IntervalTime seconds
 	// if FiniteDeltaTime is > FiniteDeltaTime a second pose will be extracted at AnimationTime - FiniteDeltaTime and placed at IntervalTime - FiniteDeltaTime (for MM to be able to calculate velocities)
-	void ExtractAndAddFuturePoses(const UAnimationAsset* AnimationAsset, float AnimationTime, float FiniteDeltaTime, const FVector& BlendParameters, float IntervalTime, const USkeleton* OverrideSkeleton = nullptr);
+	void ExtractAndAddFuturePoses(const UAnimationAsset* AnimationAsset, float AnimationTime, float FiniteDeltaTime, const FVector& BlendParameters, float IntervalTime, const USkeleton* OverrideSkeleton = nullptr, bool bUseRefPoseRootBone = false);
 
 	UE_DEPRECATED(5.6, "Use other AddFuturePose signatures instead")
 	void AddFuturePose(float Time, FCSPose<FCompactPose>& ComponentSpacePose, const FBlendedCurve& Curves = FBlendedCurve());
diff --git a/Engine/Plugins/Experimental/MovieScenePoseSearchTracks/Source/MovieScenePoseSearchTracks/Private/Sections/MovieSceneStitchAnimSection.cpp b/Engine/Plugins/Experimental/MovieScenePoseSearchTracks/Source/MovieScenePoseSearchTracks/Private/Sections/MovieSceneStitchAnimSection.cpp
--- a/Engine/Plugins/Experimental/MovieScenePoseSearchTracks/Source/MovieScenePoseSearchTracks/Private/Sections/MovieSceneStitchAnimSection.cpp
+++ b/Engine/Plugins/Experimental/MovieScenePoseSearchTracks/Source/MovieScenePoseSearchTracks/Private/Sections/MovieSceneStitchAnimSection.cpp
@@ -108,7 +108,7 @@ void FMovieSceneStitchAnimEvaluationTask::Execute(UE::AnimNext::FEvaluationVM& V
 
 			// extracting two future poses from StitchData.TargetPoseAsset at StitchData.TargetAnimationTimeSeconds in the future and StitchData.TargetAnimationTimeSeconds - FiniteDelta
 			// to have motion matching been able to compute velocities (if required by the database schema channels)
-			OverridePoseHistory.ExtractAndAddFuturePoses(StitchData.TargetPoseAsset, StitchData.TargetAnimationTimeSeconds, FiniteDelta, FVector::ZeroVector, TimeToTarget);
+			OverridePoseHistory.ExtractAndAddFuturePoses(StitchData.TargetPoseAsset, StitchData.TargetAnimationTimeSeconds, FiniteDelta, FVector::ZeroVector, TimeToTarget, nullptr, true);
 
 			PoseHistory = &OverridePoseHistory;
 
