PoseSearch - continuing interaction bias support

diff --git a/Engine/Plugins/Animation/PoseSearch/Source/Editor/Private/PoseSearchDebuggerDatabaseColumns.h b/Engine/Plugins/Animation/PoseSearch/Source/Editor/Private/PoseSearchDebuggerDatabaseColumns.h
--- a/Engine/Plugins/Animation/PoseSearch/Source/Editor/Private/PoseSearchDebuggerDatabaseColumns.h
+++ b/Engine/Plugins/Animation/PoseSearch/Source/Editor/Private/PoseSearchDebuggerDatabaseColumns.h
@@ -557,6 +557,82 @@ struct FCostModifier : ITextColumn
 	}
 };
 
+struct FContinuingPoseCost : ITextColumn
+{
+	using ITextColumn::ITextColumn;
+
+	virtual FText GetLabel() const override
+	{
+		return LOCTEXT("ColumnLabelContinuingPoseCost", "ContinuingPoseCost");
+	}
+
+	virtual FText GetLabelTooltip() const override
+	{
+		return LOCTEXT("ColumnLabelTooltipContinuingPoseCost", "Continuing Pose Cost Bias contribution");
+	}
+
+	virtual FSortPredicate GetSortPredicate() const override
+	{
+		return [](const FRowDataRef& Row0, const FRowDataRef& Row1) -> bool { return Row0->PoseCost.GetContinuingPoseCostAddend() < Row1->PoseCost.GetContinuingPoseCostAddend(); };
+	}
+
+	virtual FText GetRowText(const FRowDataRef& Row) const override
+	{
+		return FText::AsNumber(Row->PoseCost.GetContinuingPoseCostAddend());
+	}
+};
+
+struct FContinuingInteractionCost : ITextColumn
+{
+	using ITextColumn::ITextColumn;
+
+	virtual FText GetLabel() const override
+	{
+		return LOCTEXT("ColumnLabelContinuingInteractionCost", "ContinuingInteractionCost");
+	}
+
+	virtual FText GetLabelTooltip() const override
+	{
+		return LOCTEXT("ColumnLabelTooltipContinuingInteractionCost", "Continuing Interaction Cost Bias contribution");
+	}
+
+	virtual FSortPredicate GetSortPredicate() const override
+	{
+		return [](const FRowDataRef& Row0, const FRowDataRef& Row1) -> bool { return Row0->PoseCost.GetContinuingInteractionCostAddend() < Row1->PoseCost.GetContinuingInteractionCostAddend(); };
+	}
+
+	virtual FText GetRowText(const FRowDataRef& Row) const override
+	{
+		return FText::AsNumber(Row->PoseCost.GetContinuingInteractionCostAddend());
+	}
+};
+
+struct FNotifyCost : ITextColumn
+{
+	using ITextColumn::ITextColumn;
+
+	virtual FText GetLabel() const override
+	{
+		return LOCTEXT("ColumnLabelNotifyCost", "NotifyCost");
+	}
+
+	virtual FText GetLabelTooltip() const override
+	{
+		return LOCTEXT("ColumnLabelTooltipNotifyCost", "Notify Cost Bias contribution");
+	}
+
+	virtual FSortPredicate GetSortPredicate() const override
+	{
+		return [](const FRowDataRef& Row0, const FRowDataRef& Row1) -> bool { return Row0->PoseCost.GetNotifyCostAddend() < Row1->PoseCost.GetNotifyCostAddend(); };
+	}
+
+	virtual FText GetRowText(const FRowDataRef& Row) const override
+	{
+		return FText::AsNumber(Row->PoseCost.GetNotifyCostAddend());
+	}
+};
+
+
 struct FMirrored : ITextColumn
 {
 	using ITextColumn::ITextColumn;
diff --git a/Engine/Plugins/Animation/PoseSearch/Source/Editor/Private/PoseSearchDebuggerDatabaseView.cpp b/Engine/Plugins/Animation/PoseSearch/Source/Editor/Private/PoseSearchDebuggerDatabaseView.cpp
--- a/Engine/Plugins/Animation/PoseSearch/Source/Editor/Private/PoseSearchDebuggerDatabaseView.cpp
+++ b/Engine/Plugins/Animation/PoseSearch/Source/Editor/Private/PoseSearchDebuggerDatabaseView.cpp
@@ -28,6 +28,17 @@
 namespace UE::PoseSearch
 {
 
+static FORCEINLINE float ArraySum(TConstArrayView<float> View, int32 StartIndex, int32 Offset)
+{
+	float Sum = 0.f;
+	const int32 EndIndex = StartIndex + Offset;
+	for (int32 i = StartIndex; i < EndIndex; ++i)
+	{
+		Sum += View[i];
+	}
+	return Sum;
+}
+
 class SCostBreakDownData
 {
 public:
@@ -138,16 +149,7 @@ private:
 		}
 	}
 
-	static float ArraySum(TConstArrayView<float> View, int32 StartIndex, int32 Offset)
-	{
-		float Sum = 0.f;
-		const int32 EndIndex = StartIndex + Offset;
-		for (int32 i = StartIndex; i < EndIndex; ++i)
-		{
-			Sum += View[i];
-		}
-		return Sum;
-	}
+	
 
 	struct FLabelToChannels
 	{
@@ -188,22 +190,13 @@ static void AddUnfilteredDatabaseRow(const UPoseSearchDatabase* Database, TConst
 		if (Cost.IsValid())
 		{
 			Row->PoseCost = Cost;
-
-			// @todo: perhaps reuse this code to recalculate Row->PoseCost instead of using Cost
-			//const FPoseSearchCost RecalculatedCost = SearchIndex.ComparePoses(DbPoseIdx, 0.f, PoseValues, QueryVector);
-			//if (Cost != RecalculatedCost)
-			//{
-			//	UE_LOG(LogPoseSearchEditor, Error, TEXT("RecalculatedCost (%f) differs from Cost (%f) for Pose '%d'"), RecalculatedCost, Cost, DbPoseIdx);
-			//}
-		}
-		// is PoseValues padded at 16 bytes (and 16 bytes aligned by construction)?
-		else if (PoseValues.Num() % 4 == 0)
-		{
-			Row->PoseCost = SearchIndex.CompareAlignedPoses(DbPoseIdx, 0.f, PoseValues, QueryVector, DynamicWeightsSqrt);
 		}
 		else
 		{
-			Row->PoseCost = SearchIndex.ComparePoses(DbPoseIdx, 0.f, PoseValues, QueryVector, DynamicWeightsSqrt);
+			// @todo: gather ContinuingPoseCostAddend, ContinuingInteractionCostAddend properly to reconstruct the real FPoseSearchCost
+			const float ContinuingPoseCostAddend = 0.f;
+			const float ContinuingInteractionCostAddend = 0.f;
+			Row->PoseCost = FPoseSearchCost(ArraySum(Row->CostVector, 0, Database->Schema->SchemaCardinality), SearchIndex.PoseMetadata[DbPoseIdx].GetCostAddend(), ContinuingPoseCostAddend, ContinuingInteractionCostAddend);
 		}
 
 		if (!SharedData->PCAQueryVector.IsEmpty())
@@ -281,7 +274,7 @@ void SDebuggerDatabaseView::Update(const FTraceMotionMatchingStateMessage& State
 			{
 				SharedData->PCAQueryVector.SetNumZeroed(Database->GetNumberOfPrincipalComponents());
 				SearchIndex.PCAProject(DbEntry.QueryVector, SharedData->PCAQueryVector);
-				bAddPCACost = true;
+				bAddPCACost = bIsVerbose;
 			}
 
 			DynamicWeightsSqrtBuffer.SetNumUninitialized(SearchIndex.WeightsSqrt.Num());
@@ -401,9 +394,10 @@ void SDebuggerDatabaseView::Update(const FTraceMotionMatchingStateMessage& State
 		}
 	}
 
-	if (bHasPCACostColumn != bAddPCACost || !CostBreakDownData.AreLabelsEqualTo(OldLabels))
+	if (bHasPCACostColumn != bAddPCACost || bWasVerbose != bIsVerbose || !CostBreakDownData.AreLabelsEqualTo(OldLabels))
 	{
 		bHasPCACostColumn = bAddPCACost;
+		bWasVerbose = bIsVerbose;
 
 		OldLabels = CostBreakDownData.GetLabels();
 
@@ -423,13 +417,23 @@ void SDebuggerDatabaseView::Update(const FTraceMotionMatchingStateMessage& State
 			AddColumn(MakeShared<FPCACost>(ColumnIdx++));
 		}
 
+		if (bIsVerbose)
+		{
+			AddColumn(MakeShared<FContinuingPoseCost>(ColumnIdx++));
+			AddColumn(MakeShared<FNotifyCost>(ColumnIdx++));
+			AddColumn(MakeShared<FContinuingInteractionCost>(ColumnIdx++));
+		}
+		else
+		{
+			AddColumn(MakeShared<FCostModifier>(ColumnIdx++));
+		}
+		
 		int32 LabelIdx = 0;
 		for (const FText& Label : CostBreakDownData.GetLabels())
 		{
 			AddColumn(MakeShared<FChannelBreakdownCostColumn>(ColumnIdx++, LabelIdx++, Label));
 		}
 
-		AddColumn(MakeShared<FCostModifier>(ColumnIdx++));
 		AddColumn(MakeShared<FFrame>(ColumnIdx++));
 		AddColumn(MakeShared<FTime>(ColumnIdx++));
 		AddColumn(MakeShared<FPercentage>(ColumnIdx++));
@@ -1230,6 +1234,7 @@ void SDebuggerDatabaseView::Construct(const FArguments& InArgs)
 	OldLabels.Reset();
 	Columns.Reset();
 	bHasPCACostColumn = false;
+	bWasVerbose = false;
 }
 
 } // namespace UE::PoseSearch
diff --git a/Engine/Plugins/Animation/PoseSearch/Source/Editor/Private/PoseSearchDebuggerDatabaseView.h b/Engine/Plugins/Animation/PoseSearch/Source/Editor/Private/PoseSearchDebuggerDatabaseView.h
--- a/Engine/Plugins/Animation/PoseSearch/Source/Editor/Private/PoseSearchDebuggerDatabaseView.h
+++ b/Engine/Plugins/Animation/PoseSearch/Source/Editor/Private/PoseSearchDebuggerDatabaseView.h
@@ -105,6 +105,7 @@ private:
 
 	TArray<FText> OldLabels;
 	bool bHasPCACostColumn = false;
+	bool bWasVerbose = false;
 
 	struct FTable
 	{
diff --git a/Engine/Plugins/Animation/PoseSearch/Source/Runtime/Private/AnimNode_MotionMatchingInteraction.cpp b/Engine/Plugins/Animation/PoseSearch/Source/Runtime/Private/AnimNode_MotionMatchingInteraction.cpp
--- a/Engine/Plugins/Animation/PoseSearch/Source/Runtime/Private/AnimNode_MotionMatchingInteraction.cpp
+++ b/Engine/Plugins/Animation/PoseSearch/Source/Runtime/Private/AnimNode_MotionMatchingInteraction.cpp
@@ -63,14 +63,7 @@ void FAnimNode_MotionMatchingInteraction::UpdateAssetPlayer(const FAnimationUpda
 	const float DeltaTime = Context.GetDeltaTime();
 	if (FPoseHistoryProvider* PoseHistoryProvider = Context.GetMessage<FPoseHistoryProvider>())
 	{
-		// @todo: retrieve the playing UMultiAnimAsset from the UAnimAsset blendspace is playing via GetAnimAsset() and blend parameters via GetBlendParameters() and mirror state
-		FPoseSearchContinuingProperties ContinuingProperties;
-		ContinuingProperties.PlayingAsset = CurrentResult.SelectedAnim;
-		ContinuingProperties.PlayingAssetAccumulatedTime = GetAccumulatedTime();
-		ContinuingProperties.bIsPlayingAssetMirrored = CurrentResult.bIsMirrored;
-		ContinuingProperties.PlayingAssetBlendParameters = CurrentResult.BlendParameters;
-
-		UPoseSearchInteractionLibrary::MotionMatchInteraction(CurrentResult, Availabilities, Context.AnimInstanceProxy->GetAnimInstanceObject(), ContinuingProperties, FName(),  & PoseHistoryProvider->GetPoseHistory(), bValidateResultAgainstAvailabilities);
+		UPoseSearchInteractionLibrary::MotionMatchInteraction(CurrentResult, Availabilities, Context.AnimInstanceProxy->GetAnimInstanceObject(), FName(),  & PoseHistoryProvider->GetPoseHistory(), bValidateResultAgainstAvailabilities);
 		check(CurrentResult.ActorRootTransforms.Num() == CurrentResult.ActorRootBoneTransforms.Num());
 
 		if (const UMultiAnimAsset* MultiAnimAsset = Cast<UMultiAnimAsset>(CurrentResult.SelectedAnim))
diff --git a/Engine/Plugins/Animation/PoseSearch/Source/Runtime/Private/PoseSearchDatabase.cpp b/Engine/Plugins/Animation/PoseSearch/Source/Runtime/Private/PoseSearchDatabase.cpp
--- a/Engine/Plugins/Animation/PoseSearch/Source/Runtime/Private/PoseSearchDatabase.cpp
+++ b/Engine/Plugins/Animation/PoseSearch/Source/Runtime/Private/PoseSearchDatabase.cpp
@@ -21,6 +21,7 @@
 #include "PoseSearch/PoseSearchFeatureChannel_Group.h"
 #include "PoseSearch/PoseSearchHistory.h"
 #include "PoseSearch/PoseSearchSchema.h"
+#include "PoseSearchIndex.inl"
 #include "Serialization/ArchiveCountMem.h"
 #include "UObject/ObjectSaveContext.h"
 
@@ -93,10 +94,9 @@ struct FSearchFilters
 	}
 
 	// @todo: template this with bAlignedAndPadded to be able to use faster ComparePoses
-	bool AreFiltersValid(const FSearchIndex& SearchIndex, TConstArrayView<float> PoseValues, TConstArrayView<float> QueryValues, 
-		TConstArrayView<float> DynamicWeightsSqrt, int32 PoseIdx
+	bool AreFiltersValid(const FSearchIndex& SearchIndex, TConstArrayView<float> PoseValues, TConstArrayView<float> QueryValues, TConstArrayView<float> DynamicWeightsSqrt, int32 PoseIdx
 #if UE_POSE_SEARCH_TRACE_ENABLED
-		, UE::PoseSearch::FSearchContext& SearchContext, const UPoseSearchDatabase* Database
+		, float ContinuingPoseCostAddend, float ContinuingInteractionCostAddend, UE::PoseSearch::FSearchContext& SearchContext, const UPoseSearchDatabase* Database
 #endif // UE_POSE_SEARCH_TRACE_ENABLED
 	) const
 	{
@@ -111,17 +111,17 @@ struct FSearchFilters
 				}
 				else if (Filter == &SelectableAssetIdxFilter)
 				{
-					const FPoseSearchCost PoseCost = SearchIndex.ComparePoses(PoseIdx, 0.f, PoseValues, QueryValues, DynamicWeightsSqrt);
+					const FPoseSearchCost PoseCost(CompareFeatureVectors<false>(PoseValues, QueryValues, DynamicWeightsSqrt), SearchIndex.PoseMetadata[PoseIdx].GetCostAddend(), ContinuingPoseCostAddend, ContinuingInteractionCostAddend);
 					SearchContext.Track(Database, PoseIdx, EPoseCandidateFlags::DiscardedBy_AssetIdxFilter, PoseCost);
 				}
 				else if (Filter == &BlockTransitionFilter)
 				{
-					const FPoseSearchCost PoseCost = SearchIndex.ComparePoses(PoseIdx, 0.f, PoseValues, QueryValues, DynamicWeightsSqrt);
+					const FPoseSearchCost PoseCost(CompareFeatureVectors<false>(PoseValues, QueryValues, DynamicWeightsSqrt), SearchIndex.PoseMetadata[PoseIdx].GetCostAddend(), ContinuingPoseCostAddend, ContinuingInteractionCostAddend);
 					SearchContext.Track(Database, PoseIdx, EPoseCandidateFlags::DiscardedBy_BlockTransition, PoseCost);
 				}
 				else
 				{
-					const FPoseSearchCost PoseCost = SearchIndex.ComparePoses(PoseIdx, 0.f, PoseValues, QueryValues, DynamicWeightsSqrt);
+					const FPoseSearchCost PoseCost(CompareFeatureVectors<false>(PoseValues, QueryValues, DynamicWeightsSqrt), SearchIndex.PoseMetadata[PoseIdx].GetCostAddend(), ContinuingPoseCostAddend, ContinuingInteractionCostAddend);
 					SearchContext.Track(Database, PoseIdx, EPoseCandidateFlags::DiscardedBy_PoseFilter, PoseCost);
 				}
 #endif // UE_POSE_SEARCH_TRACE_ENABLED
@@ -198,7 +198,8 @@ private:
 
 template<bool bReconstructPoseValues, bool bAlignedAndPadded>
 static inline void EvaluatePoseKernel(UE::PoseSearch::FSearchResult& Result, const UE::PoseSearch::FSearchIndex& SearchIndex, TConstArrayView<float> QueryValues, TArrayView<float> ReconstructedPoseValuesBuffer,
-	int32 PoseIdx, int32 EventPoseIdx, const UE::PoseSearch::FSearchFilters& SearchFilters, UE::PoseSearch::FSearchContext& SearchContext, const UPoseSearchDatabase* Database, TConstArrayView<float> DynamicWeightsSqrt, bool bUpdateBestCandidates, int32 ResultIndex = -1)
+	int32 PoseIdx, int32 EventPoseIdx, const UE::PoseSearch::FSearchFilters& SearchFilters, float ContinuingPoseCostAddend, float ContinuingInteractionCostAddend,
+	UE::PoseSearch::FSearchContext& SearchContext, const UPoseSearchDatabase* Database, TConstArrayView<float> DynamicWeightsSqrt, bool bUpdateBestCandidates, int32 ResultIndex = -1)
 {
 	using namespace UE::PoseSearch;
 
@@ -206,11 +207,11 @@ static inline void EvaluatePoseKernel(UE::PoseSearch::FSearchResult& Result, con
 
 	if (SearchFilters.AreFiltersValid(SearchIndex, PoseValues, QueryValues, DynamicWeightsSqrt, PoseIdx
 #if UE_POSE_SEARCH_TRACE_ENABLED
-		, SearchContext, Database
+		, ContinuingPoseCostAddend, ContinuingInteractionCostAddend, SearchContext, Database
 #endif // UE_POSE_SEARCH_TRACE_ENABLED
 	))
 	{
-		const FPoseSearchCost PoseCost = bAlignedAndPadded ? SearchIndex.CompareAlignedPoses(PoseIdx, 0.f, PoseValues, QueryValues, DynamicWeightsSqrt) : SearchIndex.ComparePoses(PoseIdx, 0.f, PoseValues, QueryValues, DynamicWeightsSqrt);
+		const FPoseSearchCost PoseCost(CompareFeatureVectors<bAlignedAndPadded>(PoseValues, QueryValues, DynamicWeightsSqrt), SearchIndex.PoseMetadata[PoseIdx].GetCostAddend(), ContinuingPoseCostAddend, ContinuingInteractionCostAddend);
 		if (PoseCost < Result.PoseCost)
 		{
 			Result.PoseCost = PoseCost;
@@ -1428,7 +1429,7 @@ UE::PoseSearch::FSearchResult UPoseSearchDatabase::Search(UE::PoseSearch::FSearc
 
 void UPoseSearchDatabase::PopulateNonSelectableIdx(FNonSelectableIdx& NonSelectableIdx, UE::PoseSearch::FSearchContext& SearchContext
 #if UE_POSE_SEARCH_TRACE_ENABLED
-	, TConstArrayView<float> QueryValues, TConstArrayView<float> DynamicWeightsSqrt
+	, float ContinuingPoseCostAddend, float ContinuingInteractionCostAddend, TConstArrayView<float> QueryValues, TConstArrayView<float> DynamicWeightsSqrt
 #endif //UE_POSE_SEARCH_TRACE_ENABLED
 ) const
 {
@@ -1463,7 +1464,7 @@ void UPoseSearchDatabase::PopulateNonSelectableIdx(FNonSelectableIdx& NonSelecta
 
 #if UE_POSE_SEARCH_TRACE_ENABLED
 						const TConstArrayView<float> PoseValues = SearchIndex.GetPoseValuesSafe(PoseIdx, BufferUsedForReconstruction);
-						const FPoseSearchCost PoseCost = SearchIndex.ComparePoses(PoseIdx, 0.f, PoseValues, QueryValues, DynamicWeightsSqrt);
+						const FPoseSearchCost PoseCost(CompareFeatureVectors<false>(PoseValues, QueryValues, DynamicWeightsSqrt), SearchIndex.PoseMetadata[PoseIdx].GetCostAddend(), ContinuingPoseCostAddend, ContinuingInteractionCostAddend);
 						SearchContext.Track(this, PoseIdx, EPoseCandidateFlags::DiscardedBy_AssetReselection, PoseCost);
 #endif // UE_POSE_SEARCH_TRACE_ENABLED
 					}
@@ -1490,7 +1491,7 @@ void UPoseSearchDatabase::PopulateNonSelectableIdx(FNonSelectableIdx& NonSelecta
 
 #if UE_POSE_SEARCH_TRACE_ENABLED
 						const TConstArrayView<float> PoseValues = SearchIndex.GetPoseValuesSafe(PoseIdx, BufferUsedForReconstruction);
-						const FPoseSearchCost PoseCost = SearchIndex.ComparePoses(PoseIdx, 0.f, PoseValues, QueryValues, DynamicWeightsSqrt);
+						const FPoseSearchCost PoseCost(CompareFeatureVectors<false>(PoseValues, QueryValues, DynamicWeightsSqrt), SearchIndex.PoseMetadata[PoseIdx].GetCostAddend(), ContinuingPoseCostAddend, ContinuingInteractionCostAddend);
 						SearchContext.Track(this, PoseIdx, EPoseCandidateFlags::DiscardedBy_PoseJumpThresholdTime, PoseCost);
 #endif // UE_POSE_SEARCH_TRACE_ENABLED
 					}
@@ -1506,7 +1507,7 @@ void UPoseSearchDatabase::PopulateNonSelectableIdx(FNonSelectableIdx& NonSelecta
 
 #if UE_POSE_SEARCH_TRACE_ENABLED
 						const TConstArrayView<float> PoseValues = SearchIndex.GetPoseValuesSafe(PoseIdx, BufferUsedForReconstruction);
-						const FPoseSearchCost PoseCost = SearchIndex.ComparePoses(PoseIdx, 0.f, PoseValues, QueryValues, DynamicWeightsSqrt);
+						const FPoseSearchCost PoseCost(CompareFeatureVectors<false>(PoseValues, QueryValues, DynamicWeightsSqrt), SearchIndex.PoseMetadata[PoseIdx].GetCostAddend(), ContinuingPoseCostAddend, ContinuingInteractionCostAddend);
 						SearchContext.Track(this, PoseIdx, EPoseCandidateFlags::DiscardedBy_PoseJumpThresholdTime, PoseCost);
 #endif // UE_POSE_SEARCH_TRACE_ENABLED
 					}
@@ -1531,7 +1532,7 @@ void UPoseSearchDatabase::PopulateNonSelectableIdx(FNonSelectableIdx& NonSelecta
 				// if we're editing the database and removing assets it's possible that the PoseIndicesHistory contains invalid pose indexes
 				if (HistoricalPoseIndex.PoseIndex < SearchIndex.GetNumPoses())
 				{
-					const FPoseSearchCost PoseCost = SearchIndex.ComparePoses(HistoricalPoseIndex.PoseIndex, 0.f, SearchIndex.GetPoseValuesSafe(HistoricalPoseIndex.PoseIndex, BufferUsedForReconstruction), QueryValues, DynamicWeightsSqrt);
+					const FPoseSearchCost PoseCost(CompareFeatureVectors<false>(SearchIndex.GetPoseValuesSafe(HistoricalPoseIndex.PoseIndex, BufferUsedForReconstruction), QueryValues, DynamicWeightsSqrt), SearchIndex.PoseMetadata[HistoricalPoseIndex.PoseIndex].GetCostAddend(), ContinuingPoseCostAddend, ContinuingInteractionCostAddend);
 					SearchContext.Track(this, HistoricalPoseIndex.PoseIndex, EPoseCandidateFlags::DiscardedBy_PoseReselectHistory, PoseCost);
 				}
 #endif // UE_POSE_SEARCH_TRACE_ENABLED
@@ -1594,7 +1595,7 @@ UE::PoseSearch::FSearchResult UPoseSearchDatabase::SearchContinuingPose(UE::Pose
 	const FPoseSearchDatabaseAnimationAssetBase* DatabaseAnimationAssetBase = GetDatabaseAnimationAsset<FPoseSearchDatabaseAnimationAssetBase>(SearchIndexAsset);
 	check(DatabaseAnimationAssetBase);
 
-	float UpdatedContinuingPoseCostBias = ContinuingPoseCostBias;
+	float ContinuingPoseCostAddend = ContinuingPoseCostBias;
 	const float SampleTime = GetRealAssetTime(PoseIdx);
 	FAnimNotifyContext PreAllocatedNotifyContext;
 	for (int32 RoleIndex = 0; RoleIndex < DatabaseAnimationAssetBase->GetNumRoles(); ++RoleIndex)
@@ -1605,11 +1606,11 @@ UE::PoseSearch::FSearchResult UPoseSearchDatabase::SearchContinuingPose(UE::Pose
 			const FAnimationAssetSampler SequenceBaseSampler(AnimationAsset, FTransform::Identity, SearchIndexAsset.GetBlendParameters(), FAnimationAssetSampler::DefaultRootTransformSamplingRate, false, false);
 
 			bool bDone = false;
-			SequenceBaseSampler.ExtractAnimNotifyStates(SampleTime, PreAllocatedNotifyContext, [&UpdatedContinuingPoseCostBias, &bDone](const UAnimNotifyState* AnimNotifyState)
+			SequenceBaseSampler.ExtractAnimNotifyStates(SampleTime, PreAllocatedNotifyContext, [&ContinuingPoseCostAddend, &bDone](const UAnimNotifyState* AnimNotifyState)
 				{
 					if (const UAnimNotifyState_PoseSearchOverrideContinuingPoseCostBias* NotifyStateContinuingPoseCostBias = Cast<const UAnimNotifyState_PoseSearchOverrideContinuingPoseCostBias>(AnimNotifyState))
 					{
-						UpdatedContinuingPoseCostBias = NotifyStateContinuingPoseCostBias->CostAddend;
+						ContinuingPoseCostAddend = NotifyStateContinuingPoseCostBias->CostAddend;
 						bDone = true;
 						return false;
 					}
@@ -1623,9 +1624,11 @@ UE::PoseSearch::FSearchResult UPoseSearchDatabase::SearchContinuingPose(UE::Pose
 		}
 	}
 
-	// since any PoseCost calculated here is at least SearchIndex.MinCostAddend + UpdatedContinuingPoseCostBias,
+	const float ContinuingInteractionCostAddend = SearchContext.IsContinuingInteraction() ? ContinuingInteractionCostBias : 0.f;
+
+	// since any PoseCost calculated here is at least SearchIndex.MinCostAddend + ContinuingPoseCostAddend + ContinuingInteractionCostAddend,
 	// there's no point in performing the search if CurrentBestTotalCost is already better than that
-	if (!GetSkipSearchIfPossible() || SearchContext.GetCurrentBestTotalCost() > SearchIndex.MinCostAddend + UpdatedContinuingPoseCostBias)
+	if (!GetSkipSearchIfPossible() || SearchContext.GetCurrentBestTotalCost() > (SearchIndex.MinCostAddend + ContinuingPoseCostAddend + ContinuingInteractionCostAddend))
 	{
 		const int32 NumDimensions = Schema->SchemaCardinality;
 		// FMemory_Alloca is forced 16 bytes aligned
@@ -1641,12 +1644,12 @@ UE::PoseSearch::FSearchResult UPoseSearchDatabase::SearchContinuingPose(UE::Pose
 		// is the data padded at 16 bytes (and 16 bytes aligned by construction)?
 		if (NumDimensions % 4 == 0)
 		{
-			Result.PoseCost = SearchIndex.CompareAlignedPoses(ContinuingPoseIdx, UpdatedContinuingPoseCostBias, PoseValues, SearchContext.GetOrBuildQuery(Schema), DynamicWeightsSqrt);
+			Result.PoseCost = FPoseSearchCost(CompareFeatureVectors<true>(PoseValues, SearchContext.GetOrBuildQuery(Schema), DynamicWeightsSqrt), SearchIndex.PoseMetadata[ContinuingPoseIdx].GetCostAddend(), ContinuingPoseCostAddend, ContinuingInteractionCostAddend);
 		}
 		// data is not 16 bytes padded
 		else
 		{
-			Result.PoseCost = SearchIndex.ComparePoses(ContinuingPoseIdx, UpdatedContinuingPoseCostBias, PoseValues, SearchContext.GetOrBuildQuery(Schema), DynamicWeightsSqrt);
+			Result.PoseCost = FPoseSearchCost(CompareFeatureVectors<false>(PoseValues, SearchContext.GetOrBuildQuery(Schema), DynamicWeightsSqrt), SearchIndex.PoseMetadata[ContinuingPoseIdx].GetCostAddend(), ContinuingPoseCostAddend, ContinuingInteractionCostAddend);
 		}
 
 		Result.AssetTime = SearchContext.GetCurrentResult().AssetTime;
@@ -1679,6 +1682,8 @@ UE::PoseSearch::FSearchResult UPoseSearchDatabase::SearchPCAKDTree(UE::PoseSearc
 
 	const int32 NumDimensions = Schema->SchemaCardinality;
 	const FSearchIndex& SearchIndex = GetSearchIndex();
+	const float ContinuingPoseCostAddend = 0.f;
+	const float ContinuingInteractionCostAddend = SearchContext.IsContinuingInteraction() ? ContinuingInteractionCostBias : 0.f;
 
 	// since any PoseCost calculated here is at least SearchIndex.MinCostAddend,
 	// there's no point in performing the search if CurrentBestTotalCost is already better than that
@@ -1704,7 +1709,7 @@ UE::PoseSearch::FSearchResult UPoseSearchDatabase::SearchPCAKDTree(UE::PoseSearc
 		FNonSelectableIdx NonSelectableIdx;
 		PopulateNonSelectableIdx(NonSelectableIdx, SearchContext
 #if UE_POSE_SEARCH_TRACE_ENABLED
-			, QueryValues, DynamicWeightsSqrt
+			, ContinuingPoseCostAddend, ContinuingInteractionCostAddend, QueryValues, DynamicWeightsSqrt
 #endif // UE_POSE_SEARCH_TRACE_ENABLED
 		);
 
@@ -1852,7 +1857,7 @@ UE::PoseSearch::FSearchResult UPoseSearchDatabase::SearchPCAKDTree(UE::PoseSearc
 					const TConstArrayView<int32> PoseIndexes = SearchIndex.PCAValuesVectorToPoseIndexes[Results[ResultIndex].Index];
 					for (int32 Index = 0; Index < PoseIndexes.Num() && NumEvaluatePoseKernelCalls < MaxNumEvaluatePoseKernelCalls; ++Index, ++NumEvaluatePoseKernelCalls)
 					{
-						EvaluatePoseKernel<false, true>(Result, SearchIndex, QueryValues, TArrayView<float>(), PoseIndexes[Index], INDEX_NONE, SearchFilters, SearchContext, this, DynamicWeightsSqrt, true
+						EvaluatePoseKernel<false, true>(Result, SearchIndex, QueryValues, TArrayView<float>(), PoseIndexes[Index], INDEX_NONE, SearchFilters, ContinuingPoseCostAddend, ContinuingInteractionCostAddend, SearchContext, this, DynamicWeightsSqrt, true
 #if WITH_EDITOR && ENABLE_ANIM_DEBUG
 							, SortedResultsIndexes[ResultIndex]
 #endif // WITH_EDITOR && ENABLE_ANIM_DEBUG
@@ -1868,7 +1873,7 @@ UE::PoseSearch::FSearchResult UPoseSearchDatabase::SearchPCAKDTree(UE::PoseSearc
 					const TConstArrayView<int32> PoseIndexes = SearchIndex.PCAValuesVectorToPoseIndexes[Results[ResultIndex].Index];
 					for (int32 Index = 0; Index < PoseIndexes.Num() && NumEvaluatePoseKernelCalls < MaxNumEvaluatePoseKernelCalls; ++Index, ++NumEvaluatePoseKernelCalls)
 					{
-						EvaluatePoseKernel<false, false>(Result, SearchIndex, QueryValues, TArrayView<float>(), PoseIndexes[Index], INDEX_NONE, SearchFilters, SearchContext, this, DynamicWeightsSqrt, true
+						EvaluatePoseKernel<false, false>(Result, SearchIndex, QueryValues, TArrayView<float>(), PoseIndexes[Index], INDEX_NONE, SearchFilters, ContinuingPoseCostAddend, ContinuingInteractionCostAddend, SearchContext, this, DynamicWeightsSqrt, true
 #if WITH_EDITOR && ENABLE_ANIM_DEBUG
 							, SortedResultsIndexes[ResultIndex]
 #endif // WITH_EDITOR && ENABLE_ANIM_DEBUG
@@ -1885,7 +1890,7 @@ UE::PoseSearch::FSearchResult UPoseSearchDatabase::SearchPCAKDTree(UE::PoseSearc
 			check(IsAligned(ReconstructedPoseValuesBuffer.GetData(), alignof(VectorRegister4Float)));
 			for (int32 ResultIndex = 0; ResultIndex < NumResults; ++ResultIndex)
 			{
-				EvaluatePoseKernel<true, false>(Result, SearchIndex, QueryValues, ReconstructedPoseValuesBuffer, Results[ResultIndex].Index, INDEX_NONE, SearchFilters, SearchContext, this, DynamicWeightsSqrt, true
+				EvaluatePoseKernel<true, false>(Result, SearchIndex, QueryValues, ReconstructedPoseValuesBuffer, Results[ResultIndex].Index, INDEX_NONE, SearchFilters, ContinuingPoseCostAddend, ContinuingInteractionCostAddend, SearchContext, this, DynamicWeightsSqrt, true
 #if WITH_EDITOR && ENABLE_ANIM_DEBUG
 					, SortedResultsIndexes[ResultIndex]
 #endif // WITH_EDITOR && ENABLE_ANIM_DEBUG
@@ -1897,7 +1902,7 @@ UE::PoseSearch::FSearchResult UPoseSearchDatabase::SearchPCAKDTree(UE::PoseSearc
 		{
 			for (int32 ResultIndex = 0; ResultIndex < NumResults; ++ResultIndex)
 			{
-				EvaluatePoseKernel<false, true>(Result, SearchIndex, QueryValues, TArrayView<float>(), Results[ResultIndex].Index, INDEX_NONE, SearchFilters, SearchContext, this, DynamicWeightsSqrt, true
+				EvaluatePoseKernel<false, true>(Result, SearchIndex, QueryValues, TArrayView<float>(), Results[ResultIndex].Index, INDEX_NONE, SearchFilters, ContinuingPoseCostAddend, ContinuingInteractionCostAddend, SearchContext, this, DynamicWeightsSqrt, true
 #if WITH_EDITOR && ENABLE_ANIM_DEBUG
 					, SortedResultsIndexes[ResultIndex]
 #endif // WITH_EDITOR && ENABLE_ANIM_DEBUG
@@ -1909,7 +1914,7 @@ UE::PoseSearch::FSearchResult UPoseSearchDatabase::SearchPCAKDTree(UE::PoseSearc
 		{
 			for (int32 ResultIndex = 0; ResultIndex < NumResults; ++ResultIndex)
 			{
-				EvaluatePoseKernel<false, false>(Result, SearchIndex, QueryValues, TArrayView<float>(), Results[ResultIndex].Index, INDEX_NONE, SearchFilters, SearchContext, this, DynamicWeightsSqrt, true
+				EvaluatePoseKernel<false, false>(Result, SearchIndex, QueryValues, TArrayView<float>(), Results[ResultIndex].Index, INDEX_NONE, SearchFilters, ContinuingPoseCostAddend, ContinuingInteractionCostAddend, SearchContext, this, DynamicWeightsSqrt, true
 #if WITH_EDITOR && ENABLE_ANIM_DEBUG
 					, SortedResultsIndexes[ResultIndex]
 #endif // WITH_EDITOR && ENABLE_ANIM_DEBUG
@@ -1927,7 +1932,7 @@ UE::PoseSearch::FSearchResult UPoseSearchDatabase::SearchPCAKDTree(UE::PoseSearc
 		// calling just for reporting non selectable poses
 		TConstArrayView<float> QueryValues = SearchContext.GetOrBuildQuery(Schema);
 		FNonSelectableIdx NonSelectableIdx;
-		PopulateNonSelectableIdx(NonSelectableIdx, SearchContext, QueryValues, DynamicWeightsSqrt);
+		PopulateNonSelectableIdx(NonSelectableIdx, SearchContext, ContinuingPoseCostAddend, ContinuingInteractionCostAddend, QueryValues, DynamicWeightsSqrt);
 #endif // UE_POSE_SEARCH_TRACE_ENABLED
 	}
 
@@ -1951,6 +1956,8 @@ UE::PoseSearch::FSearchResult UPoseSearchDatabase::SearchVPTree(UE::PoseSearch::
 	FSearchResult Result;
 
 	const FSearchIndex& SearchIndex = GetSearchIndex();
+	const float ContinuingPoseCostAddend = 0.f;
+	const float ContinuingInteractionCostAddend = SearchContext.IsContinuingInteraction() ? ContinuingInteractionCostBias : 0.f;
 
 	// since any PoseCost calculated here is at least SearchIndex.MinCostAddend,
 	// there's no point in performing the search if CurrentBestTotalCost is already better than that
@@ -1971,7 +1978,7 @@ UE::PoseSearch::FSearchResult UPoseSearchDatabase::SearchVPTree(UE::PoseSearch::
 		FNonSelectableIdx NonSelectableIdx;
 		PopulateNonSelectableIdx(NonSelectableIdx, SearchContext
 #if UE_POSE_SEARCH_TRACE_ENABLED
-			, QueryValues, DynamicWeightsSqrt
+			, ContinuingPoseCostAddend, ContinuingInteractionCostAddend, QueryValues, DynamicWeightsSqrt
 #endif // UE_POSE_SEARCH_TRACE_ENABLED
 		);
 		
@@ -1997,7 +2004,7 @@ UE::PoseSearch::FSearchResult UPoseSearchDatabase::SearchVPTree(UE::PoseSearch::
 				const TConstArrayView<int32> PoseIndexes = SearchIndex.ValuesVectorToPoseIndexes[IndexDistance.Index];
 				for (int32 Index = 0; Index < PoseIndexes.Num() && NumEvaluatePoseKernelCalls < MaxNumEvaluatePoseKernelCalls; ++Index, ++NumEvaluatePoseKernelCalls)
 				{
-					EvaluatePoseKernel<false, false>(Result, SearchIndex, QueryValues, TArrayView<float>(), PoseIndexes[Index], INDEX_NONE, SearchFilters, SearchContext, this, DynamicWeightsSqrt, true
+					EvaluatePoseKernel<false, false>(Result, SearchIndex, QueryValues, TArrayView<float>(), PoseIndexes[Index], INDEX_NONE, SearchFilters, ContinuingPoseCostAddend, ContinuingInteractionCostAddend, SearchContext, this, DynamicWeightsSqrt, true
 #if WITH_EDITOR && ENABLE_ANIM_DEBUG
 						, ResultIndex
 #endif // WITH_EDITOR && ENABLE_ANIM_DEBUG
@@ -2012,7 +2019,7 @@ UE::PoseSearch::FSearchResult UPoseSearchDatabase::SearchVPTree(UE::PoseSearch::
 				const FIndexDistance& IndexDistance = UnsortedResults[ResultIndex];
 
 				// @todo: IndexDistance.Distance is the Sqrt(DissimilarityCost), so there's no need to calculate it again in SearchIndex.ComparePoses
-				EvaluatePoseKernel<false, false>(Result, SearchIndex, QueryValues, TArrayView<float>(), IndexDistance.Index, INDEX_NONE, SearchFilters, SearchContext, this, DynamicWeightsSqrt, true
+				EvaluatePoseKernel<false, false>(Result, SearchIndex, QueryValues, TArrayView<float>(), IndexDistance.Index, INDEX_NONE, SearchFilters, ContinuingPoseCostAddend, ContinuingInteractionCostAddend, SearchContext, this, DynamicWeightsSqrt, true
 #if WITH_EDITOR && ENABLE_ANIM_DEBUG
 					, ResultIndex
 #endif // WITH_EDITOR && ENABLE_ANIM_DEBUG
@@ -2031,7 +2038,7 @@ UE::PoseSearch::FSearchResult UPoseSearchDatabase::SearchVPTree(UE::PoseSearch::
 		// calling just for reporting non selectable poses
 		TConstArrayView<float> QueryValues = SearchContext.GetOrBuildQuery(Schema);
 		FNonSelectableIdx NonSelectableIdx;
-		PopulateNonSelectableIdx(NonSelectableIdx, SearchContext, QueryValues, DynamicWeightsSqrt);
+		PopulateNonSelectableIdx(NonSelectableIdx, SearchContext, ContinuingPoseCostAddend, ContinuingInteractionCostAddend, QueryValues, DynamicWeightsSqrt);
 #endif // UE_POSE_SEARCH_TRACE_ENABLED
 	}
 
@@ -2055,6 +2062,8 @@ UE::PoseSearch::FSearchResult UPoseSearchDatabase::SearchBruteForce(UE::PoseSear
 	FSearchResult Result;
 
 	const FSearchIndex& SearchIndex = GetSearchIndex();
+	const float ContinuingPoseCostAddend = 0.f;
+	const float ContinuingInteractionCostAddend = SearchContext.IsContinuingInteraction() ? ContinuingInteractionCostBias : 0.f;
 
 	// since any PoseCost calculated here is at least SearchIndex.MinCostAddend,
 	// there's no point in performing the search if CurrentBestTotalCost is already better than that
@@ -2073,12 +2082,11 @@ UE::PoseSearch::FSearchResult UPoseSearchDatabase::SearchBruteForce(UE::PoseSear
 		FNonSelectableIdx NonSelectableIdx;
 		PopulateNonSelectableIdx(NonSelectableIdx, SearchContext
 #if UE_POSE_SEARCH_TRACE_ENABLED
-			, QueryValues, DynamicWeightsSqrt
+			, ContinuingPoseCostAddend, ContinuingInteractionCostAddend, QueryValues, DynamicWeightsSqrt
 #endif // UE_POSE_SEARCH_TRACE_ENABLED
 		);
 
 		const bool bUpdateBestCandidates = PoseSearchMode == EPoseSearchMode::BruteForce;
-
 		const FSearchFilters SearchFilters(Schema, NonSelectableIdx, FSelectableAssetIdx(), SearchIndex.bAnyBlockTransition);
 
 		if (SelectableAssetIdx.IsEmpty())
@@ -2091,7 +2099,7 @@ UE::PoseSearch::FSearchResult UPoseSearchDatabase::SearchBruteForce(UE::PoseSear
 				check(IsAligned(ReconstructedPoseValuesBuffer.GetData(), alignof(VectorRegister4Float)));
 				for (int32 PoseIdx = 0; PoseIdx < SearchIndex.GetNumPoses(); ++PoseIdx)
 				{
-					EvaluatePoseKernel<true, false>(Result, SearchIndex, QueryValues, ReconstructedPoseValuesBuffer, PoseIdx, INDEX_NONE, SearchFilters, SearchContext, this, DynamicWeightsSqrt, bUpdateBestCandidates
+					EvaluatePoseKernel<true, false>(Result, SearchIndex, QueryValues, ReconstructedPoseValuesBuffer, PoseIdx, INDEX_NONE, SearchFilters, ContinuingPoseCostAddend, ContinuingInteractionCostAddend, SearchContext, this, DynamicWeightsSqrt, bUpdateBestCandidates
 #if WITH_EDITOR && ENABLE_ANIM_DEBUG
 						, PoseIdx
 #endif // WITH_EDITOR && ENABLE_ANIM_DEBUG
@@ -2103,7 +2111,7 @@ UE::PoseSearch::FSearchResult UPoseSearchDatabase::SearchBruteForce(UE::PoseSear
 			{
 				for (int32 PoseIdx = 0; PoseIdx < SearchIndex.GetNumPoses(); ++PoseIdx)
 				{
-					EvaluatePoseKernel<false, true>(Result, SearchIndex, QueryValues, TArrayView<float>(), PoseIdx, INDEX_NONE, SearchFilters, SearchContext, this, DynamicWeightsSqrt, bUpdateBestCandidates
+					EvaluatePoseKernel<false, true>(Result, SearchIndex, QueryValues, TArrayView<float>(), PoseIdx, INDEX_NONE, SearchFilters, ContinuingPoseCostAddend, ContinuingInteractionCostAddend, SearchContext, this, DynamicWeightsSqrt, bUpdateBestCandidates
 #if WITH_EDITOR && ENABLE_ANIM_DEBUG
 						, PoseIdx
 #endif // WITH_EDITOR && ENABLE_ANIM_DEBUG
@@ -2115,7 +2123,7 @@ UE::PoseSearch::FSearchResult UPoseSearchDatabase::SearchBruteForce(UE::PoseSear
 			{
 				for (int32 PoseIdx = 0; PoseIdx < SearchIndex.GetNumPoses(); ++PoseIdx)
 				{
-					EvaluatePoseKernel<false, false>(Result, SearchIndex, QueryValues, TArrayView<float>(), PoseIdx, INDEX_NONE, SearchFilters, SearchContext, this, DynamicWeightsSqrt, bUpdateBestCandidates
+					EvaluatePoseKernel<false, false>(Result, SearchIndex, QueryValues, TArrayView<float>(), PoseIdx, INDEX_NONE, SearchFilters, ContinuingPoseCostAddend, ContinuingInteractionCostAddend, SearchContext, this, DynamicWeightsSqrt, bUpdateBestCandidates
 #if WITH_EDITOR && ENABLE_ANIM_DEBUG
 						, PoseIdx
 #endif // WITH_EDITOR && ENABLE_ANIM_DEBUG
@@ -2141,7 +2149,7 @@ UE::PoseSearch::FSearchResult UPoseSearchDatabase::SearchBruteForce(UE::PoseSear
 					const int32 LastPoseIdx = FirstPoseIdx + SearchIndexAsset.GetNumPoses();
 					for (int32 PoseIdx = FirstPoseIdx; PoseIdx < LastPoseIdx; ++PoseIdx)
 					{
-						EvaluatePoseKernel<true, false>(Result, SearchIndex, QueryValues, ReconstructedPoseValuesBuffer, PoseIdx, INDEX_NONE, SearchFilters, SearchContext, this, DynamicWeightsSqrt, bUpdateBestCandidates
+						EvaluatePoseKernel<true, false>(Result, SearchIndex, QueryValues, ReconstructedPoseValuesBuffer, PoseIdx, INDEX_NONE, SearchFilters, ContinuingPoseCostAddend, ContinuingInteractionCostAddend, SearchContext, this, DynamicWeightsSqrt, bUpdateBestCandidates
 #if WITH_EDITOR && ENABLE_ANIM_DEBUG
 							, ++ResultIndex
 #endif // WITH_EDITOR && ENABLE_ANIM_DEBUG
@@ -2159,7 +2167,7 @@ UE::PoseSearch::FSearchResult UPoseSearchDatabase::SearchBruteForce(UE::PoseSear
 					const int32 LastPoseIdx = FirstPoseIdx + SearchIndexAsset.GetNumPoses();
 					for (int32 PoseIdx = FirstPoseIdx; PoseIdx < LastPoseIdx; ++PoseIdx)
 					{
-						EvaluatePoseKernel<false, true>(Result, SearchIndex, QueryValues, TArrayView<float>(), PoseIdx, INDEX_NONE, SearchFilters, SearchContext, this, DynamicWeightsSqrt, bUpdateBestCandidates
+						EvaluatePoseKernel<false, true>(Result, SearchIndex, QueryValues, TArrayView<float>(), PoseIdx, INDEX_NONE, SearchFilters, ContinuingPoseCostAddend, ContinuingInteractionCostAddend, SearchContext, this, DynamicWeightsSqrt, bUpdateBestCandidates
 #if WITH_EDITOR && ENABLE_ANIM_DEBUG
 							, ++ResultIndex
 #endif // WITH_EDITOR && ENABLE_ANIM_DEBUG
@@ -2177,7 +2185,7 @@ UE::PoseSearch::FSearchResult UPoseSearchDatabase::SearchBruteForce(UE::PoseSear
 					const int32 LastPoseIdx = FirstPoseIdx + SearchIndexAsset.GetNumPoses();
 					for (int32 PoseIdx = FirstPoseIdx; PoseIdx < LastPoseIdx; ++PoseIdx)
 					{
-						EvaluatePoseKernel<false, false>(Result, SearchIndex, QueryValues, TArrayView<float>(), PoseIdx, INDEX_NONE, SearchFilters, SearchContext, this, DynamicWeightsSqrt, bUpdateBestCandidates
+						EvaluatePoseKernel<false, false>(Result, SearchIndex, QueryValues, TArrayView<float>(), PoseIdx, INDEX_NONE, SearchFilters, ContinuingPoseCostAddend, ContinuingInteractionCostAddend, SearchContext, this, DynamicWeightsSqrt, bUpdateBestCandidates
 #if WITH_EDITOR && ENABLE_ANIM_DEBUG
 							, ++ResultIndex
 #endif // WITH_EDITOR && ENABLE_ANIM_DEBUG
@@ -2198,7 +2206,7 @@ UE::PoseSearch::FSearchResult UPoseSearchDatabase::SearchBruteForce(UE::PoseSear
 		// calling just for reporting non selectable poses
 		TConstArrayView<float> QueryValues = SearchContext.GetOrBuildQuery(Schema);
 		FNonSelectableIdx NonSelectableIdx;
-		PopulateNonSelectableIdx(NonSelectableIdx, SearchContext, QueryValues, DynamicWeightsSqrt);
+		PopulateNonSelectableIdx(NonSelectableIdx, SearchContext, ContinuingPoseCostAddend, ContinuingInteractionCostAddend, QueryValues, DynamicWeightsSqrt);
 #endif // UE_POSE_SEARCH_TRACE_ENABLED
 	}
 
@@ -2232,6 +2240,9 @@ UE::PoseSearch::FSearchResult UPoseSearchDatabase::SearchEvent(UE::PoseSearch::F
 		const TConstArrayView<int32> PosesWithEvent = SearchIndex.EventData.GetPosesWithEvent(EventToSearch.EventTag);
 		if (!PosesWithEvent.IsEmpty())
 		{
+			const float ContinuingPoseCostAddend = 0.f;
+			const float ContinuingInteractionCostAddend = SearchContext.IsContinuingInteraction() ? ContinuingInteractionCostBias : 0.f;
+
 			TConstArrayView<float> QueryValues = SearchContext.GetOrBuildQuery(Schema);
 
 			const int32 NumDimensions = Schema->SchemaCardinality;
@@ -2246,7 +2257,7 @@ UE::PoseSearch::FSearchResult UPoseSearchDatabase::SearchEvent(UE::PoseSearch::F
 			{
 				PopulateNonSelectableIdx(NonSelectableIdx, SearchContext
 #if UE_POSE_SEARCH_TRACE_ENABLED
-					, QueryValues, DynamicWeightsSqrt
+					, ContinuingPoseCostAddend, ContinuingInteractionCostAddend, QueryValues, DynamicWeightsSqrt
 #endif // UE_POSE_SEARCH_TRACE_ENABLED
 				);
 			}
@@ -2257,7 +2268,7 @@ UE::PoseSearch::FSearchResult UPoseSearchDatabase::SearchEvent(UE::PoseSearch::F
 			{
 				for (int32 EventPoseIdx : PosesWithEvent)
 				{
-					EvaluatePoseKernel<false, false>(Result, SearchIndex, QueryValues, TArrayView<float>(), EventPoseIdx, EventPoseIdx, SearchFilters, SearchContext, this, DynamicWeightsSqrt, true
+					EvaluatePoseKernel<false, false>(Result, SearchIndex, QueryValues, TArrayView<float>(), EventPoseIdx, EventPoseIdx, SearchFilters, ContinuingPoseCostAddend, ContinuingInteractionCostAddend, SearchContext, this, DynamicWeightsSqrt, true
 #if WITH_EDITOR && ENABLE_ANIM_DEBUG
 						, EventPoseIdx
 #endif // WITH_EDITOR && ENABLE_ANIM_DEBUG
@@ -2274,9 +2285,9 @@ UE::PoseSearch::FSearchResult UPoseSearchDatabase::SearchEvent(UE::PoseSearch::F
 					const FSearchIndexAsset& SearchIndexAsset = SearchIndex.GetAssetForPose(EventPoseIdx);
 					const float EventTime = SearchIndexAsset.GetTimeFromPoseIndex(EventPoseIdx, Schema->SampleRate);
 					SearchIndexAsset.IteratePoseIndexesOverTime(EventTime - TimeToEventEnd, EventTime - TimeToEventStart, Schema->SampleRate,
-						[this, &Result, &SearchIndex, &QueryValues, EventPoseIdx, &SearchFilters, &SearchContext, &DynamicWeightsSqrt](int32 PoseIdx)
+						[this, &Result, &SearchIndex, &QueryValues, EventPoseIdx, &SearchFilters, ContinuingPoseCostAddend, ContinuingInteractionCostAddend, &SearchContext, &DynamicWeightsSqrt](int32 PoseIdx)
 						{
-							EvaluatePoseKernel<false, false>(Result, SearchIndex, QueryValues, TArrayView<float>(), PoseIdx, EventPoseIdx, SearchFilters, SearchContext, this, DynamicWeightsSqrt, true
+							EvaluatePoseKernel<false, false>(Result, SearchIndex, QueryValues, TArrayView<float>(), PoseIdx, EventPoseIdx, SearchFilters, ContinuingPoseCostAddend, ContinuingInteractionCostAddend, SearchContext, this, DynamicWeightsSqrt, true
 #if WITH_EDITOR && ENABLE_ANIM_DEBUG
 								, PoseIdx
 #endif // WITH_EDITOR && ENABLE_ANIM_DEBUG
diff --git a/Engine/Plugins/Animation/PoseSearch/Source/Runtime/Private/PoseSearchHistory.cpp b/Engine/Plugins/Animation/PoseSearch/Source/Runtime/Private/PoseSearchHistory.cpp
--- a/Engine/Plugins/Animation/PoseSearch/Source/Runtime/Private/PoseSearchHistory.cpp
+++ b/Engine/Plugins/Animation/PoseSearch/Source/Runtime/Private/PoseSearchHistory.cpp
@@ -393,32 +393,29 @@ void FPoseHistoryEntry::SetComponentSpaceTransform(int32 Index, const FTransform
 
 FTransform FPoseHistoryEntry::GetComponentSpaceTransform(int32 Index) const
 {
-#if WITH_EDITOR
-	if (Index < 0 || Index >= ComponentSpaceRotations.Num())
+	if (ComponentSpaceRotations.IsValidIndex(Index))
 	{
-		UE_LOG(LogPoseSearch, Error, TEXT("FPoseHistoryEntry::GetComponentSpaceTransform - Index %d out of bound [0, %d)"), Index, ComponentSpaceRotations.Num());
-		return FTransform::Identity;
-	}
-#endif // WITH_EDITOR
+		check(ComponentSpacePositions.Num() == ComponentSpaceRotations.Num());
+		check(ComponentSpaceScales.IsEmpty() || ComponentSpaceRotations.Num() == ComponentSpaceScales.Num());
 
-	check(ComponentSpaceScales.IsEmpty() || ComponentSpaceRotations.Num() == ComponentSpaceScales.Num());
+		const FQuat Quat(ComponentSpaceRotations[Index]);
+		const FVector Scale(ComponentSpaceScales.IsEmpty() ? FVector3f::OneVector : ComponentSpaceScales[Index]);
+		return FTransform(Quat, ComponentSpacePositions[Index], Scale);
+	}
 
-	const FQuat Quat(ComponentSpaceRotations[Index]);
-	const FVector Scale(ComponentSpaceScales.IsEmpty() ? FVector3f::OneVector : ComponentSpaceScales[Index]);
-	return FTransform(Quat, ComponentSpacePositions[Index], Scale);
+	UE_LOG(LogPoseSearch, Error, TEXT("FPoseHistoryEntry::GetComponentSpaceTransform - Index %d out of bound [0, %d)"), Index, ComponentSpaceRotations.Num());
+	return FTransform::Identity;
 }
 
 float FPoseHistoryEntry::GetCurveValue(int32 Index) const
 {
-#if WITH_EDITOR
-	if (Index < 0 || Index >= CurveValues.Num())
+	if (CurveValues.IsValidIndex(Index))
 	{
-		UE_LOG(LogPoseSearch, Error, TEXT("FPoseHistoryEntry::GetCurveValue - Index %d out of bound [0, %d)"), Index, CurveValues.Num());
-		return 0.0f;
+		return CurveValues[Index];
 	}
-#endif // WITH_EDITOR
-
-	return CurveValues[Index];
+	
+	UE_LOG(LogPoseSearch, Error, TEXT("FPoseHistoryEntry::GetCurveValue - Index %d out of bound [0, %d)"), Index, CurveValues.Num());
+	return 0.0f;
 }
 
 FArchive& operator<<(FArchive& Ar, FPoseHistoryEntry& Entry)
diff --git a/Engine/Plugins/Animation/PoseSearch/Source/Runtime/Private/PoseSearchIndex.cpp b/Engine/Plugins/Animation/PoseSearch/Source/Runtime/Private/PoseSearchIndex.cpp
--- a/Engine/Plugins/Animation/PoseSearch/Source/Runtime/Private/PoseSearchIndex.cpp
+++ b/Engine/Plugins/Animation/PoseSearch/Source/Runtime/Private/PoseSearchIndex.cpp
@@ -1,68 +1,10 @@
 // Copyright Epic Games, Inc. All Rights Reserved.
 
-#include "PoseSearch/PoseSearchIndex.h"
-#include "PoseSearchEigenHelper.h"
+#include "PoseSearchIndex.inl"
 
 namespace UE::PoseSearch
 {
 
-static FORCEINLINE float CompareFeatureVectors(TConstArrayView<float> A, TConstArrayView<float> B, TConstArrayView<float> WeightsSqrt)
-{
-	check(A.Num() == B.Num() && A.Num() == WeightsSqrt.Num());
-
-	Eigen::Map<const Eigen::ArrayXf> VA(A.GetData(), A.Num());
-	Eigen::Map<const Eigen::ArrayXf> VB(B.GetData(), B.Num());
-	Eigen::Map<const Eigen::ArrayXf> VW(WeightsSqrt.GetData(), WeightsSqrt.Num());
-
-	return ((VA - VB) * VW).square().sum();
-}
-
-static FORCEINLINE float CompareAlignedFeatureVectors(TConstArrayView<float> A, TConstArrayView<float> B, TConstArrayView<float> WeightsSqrt)
-{
-	check(A.Num() == B.Num() && A.Num() == WeightsSqrt.Num());
-	check(A.Num() % 4 == 0);
-	check(IsAligned(A.GetData(), alignof(VectorRegister4Float)));
-	check(IsAligned(B.GetData(), alignof(VectorRegister4Float)));
-	check(IsAligned(WeightsSqrt.GetData(), alignof(VectorRegister4Float)));
-	// sufficient condition to check for pointer overlapping
-	check(A.GetData() != B.GetData() && A.GetData() != WeightsSqrt.GetData());
-
-	const int32 NumVectors = A.Num() / 4;
-
-	const VectorRegister4Float* RESTRICT VA = reinterpret_cast<const VectorRegister4Float*>(A.GetData());
-	const VectorRegister4Float* RESTRICT VB = reinterpret_cast<const VectorRegister4Float*>(B.GetData());
-	const VectorRegister4Float* RESTRICT VW = reinterpret_cast<const VectorRegister4Float*>(WeightsSqrt.GetData());
-
-	VectorRegister4Float PartialCost = VectorZero();
-	for (int32 VectorIdx = 0; VectorIdx < NumVectors; ++VectorIdx, ++VA, ++VB, ++VW)
-	{
-		const VectorRegister4Float Diff = VectorSubtract(*VA, *VB);
-		const VectorRegister4Float WeightedDiff = VectorMultiply(Diff, *VW);
-		PartialCost = VectorMultiplyAdd(WeightedDiff, WeightedDiff, PartialCost);
-	}
-
-	// calculating PartialCost.X + PartialCost.Y + PartialCost.Z + PartialCost.W
-	VectorRegister4Float Swizzle = VectorSwizzle(PartialCost, 1, 0, 3, 2);	// (Y, X, W, Z) of PartialCost
-	PartialCost = VectorAdd(PartialCost, Swizzle);							// (X + Y, Y + X, Z + W, W + Z)
-	Swizzle = VectorSwizzle(PartialCost, 2, 3, 0, 1);						// (Z + W, W + Z, X + Y, Y + X)
-	PartialCost = VectorAdd(PartialCost, Swizzle);							// (X + Y + Z + W, Y + X + W + Z, Z + W + X + Y, W + Z + Y + X)
-	float Cost;
-	VectorStoreFloat1(PartialCost, &Cost);
-
-// keeping this debug code to validate CompareAlignedFeatureVectors against CompareFeatureVectors
-//#if DO_CHECK
-//	const float EigenCost = CompareFeatureVectors(A, B, WeightsSqrt);
-//
-//	if (!FMath::IsNearlyEqual(Cost, EigenCost))
-//	{
-//		const float RelativeDifference = Cost > EigenCost ? (Cost - EigenCost) / Cost : (EigenCost - Cost) / EigenCost;
-//		check(FMath::IsNearlyZero(RelativeDifference, UE_KINDA_SMALL_NUMBER));
-//	}
-//#endif //DO_CHECK
-	
-	return Cost;
-}
-
 void CompareFeatureVectors(TConstArrayView<float> A, TConstArrayView<float> B, TConstArrayView<float> WeightsSqrt, TArrayView<float> Result)
 {
 	check(A.Num() == B.Num() && A.Num() == WeightsSqrt.Num() && A.Num() == Result.Num());
@@ -984,32 +926,12 @@ TConstArrayView<float> FSearchIndex::GetPCAPoseValues(int32 PCAValuesVectorIdx)
 
 FPoseSearchCost FSearchIndex::ComparePoses(int32 PoseIdx, float ContinuingPoseCostBias, TConstArrayView<float> PoseValues, TConstArrayView<float> QueryValues) const
 {
-	return ComparePoses(PoseIdx, ContinuingPoseCostBias, PoseValues, QueryValues, WeightsSqrt);
+	return FPoseSearchCost(CompareFeatureVectors<false>(PoseValues, QueryValues, WeightsSqrt), PoseMetadata[PoseIdx].GetCostAddend(), ContinuingPoseCostBias, 0.f);
 }
 
 FPoseSearchCost FSearchIndex::CompareAlignedPoses(int32 PoseIdx, float ContinuingPoseCostBias, TConstArrayView<float> PoseValues, TConstArrayView<float> QueryValues) const
 {
-	return CompareAlignedPoses(PoseIdx, ContinuingPoseCostBias, PoseValues, QueryValues, WeightsSqrt);
-}
-
-FPoseSearchCost FSearchIndex::ComparePoses(int32 PoseIdx, float ContinuingPoseCostBias, TConstArrayView<float> PoseValues, TConstArrayView<float> QueryValues, TConstArrayView<float> DynamicWeightsSqrt) const
-{
-	// base dissimilarity cost representing how the associated PoseIdx differ, in a weighted way, from the query pose (QueryValues)
-	const float DissimilarityCost = CompareFeatureVectors(PoseValues, QueryValues, DynamicWeightsSqrt);
-
-	// cost addend associated to Schema->BaseCostBias or overriden by UAnimNotifyState_PoseSearchModifyCost
-	const float NotifyAddend = PoseMetadata[PoseIdx].GetCostAddend();
-	return FPoseSearchCost(DissimilarityCost, NotifyAddend, ContinuingPoseCostBias);
-}
-
-FPoseSearchCost FSearchIndex::CompareAlignedPoses(int32 PoseIdx, float ContinuingPoseCostBias, TConstArrayView<float> PoseValues, TConstArrayView<float> QueryValues, TConstArrayView<float> DynamicWeightsSqrt) const
-{
-	// base dissimilarity cost representing how the associated PoseIdx differ, in a weighted way, from the query pose (QueryValues)
-	const float DissimilarityCost = CompareAlignedFeatureVectors(PoseValues, QueryValues, DynamicWeightsSqrt);
-
-	// cost addend associated to Schema->BaseCostBias or overriden by UAnimNotifyState_PoseSearchModifyCost
-	const float NotifyAddend = PoseMetadata[PoseIdx].GetCostAddend();
-	return FPoseSearchCost(DissimilarityCost, NotifyAddend, ContinuingPoseCostBias);
+	return FPoseSearchCost(CompareFeatureVectors<true>(PoseValues, QueryValues, WeightsSqrt), PoseMetadata[PoseIdx].GetCostAddend(), ContinuingPoseCostBias, 0.f);
 }
 
 void FSearchIndex::GetPoseToPCAValuesVectorIndexes(TArray<uint32>& PoseToPCAValuesVectorIndexes) const
@@ -1082,4 +1004,11 @@ FArchive& operator<<(FArchive& Ar, FSearchIndex& Index)
 	return Ar;
 }
 
+    
+float FVPTreeDataSource::GetDistance(const TConstArrayView<float> A, const TConstArrayView<float> B)
+{
+	// Estracting the Sqrt to satisfy the triangle inequality metric space requirements, since a <= b+c doesn't imply a^2 <= b^2 + c^2
+	return FMath::Sqrt(CompareFeatureVectors(A, B));
+}
+
 } // namespace UE::PoseSearch
diff --git a/Engine/Plugins/Animation/PoseSearch/Source/Runtime/Private/PoseSearchIndex.inl b/Engine/Plugins/Animation/PoseSearch/Source/Runtime/Private/PoseSearchIndex.inl
new file mode 100644
--- /dev/null
+++ b/Engine/Plugins/Animation/PoseSearch/Source/Runtime/Private/PoseSearchIndex.inl
@@ -0,0 +1,71 @@
+// Copyright Epic Games, Inc. All Rights Reserved.
+
+#pragma once
+
+#include "PoseSearch/PoseSearchIndex.h"
+#include "PoseSearchEigenHelper.h"
+
+namespace UE::PoseSearch
+{
+
+template <bool bAlignedAndPadded>
+static FORCEINLINE float CompareFeatureVectors(TConstArrayView<float> A, TConstArrayView<float> B, TConstArrayView<float> WeightsSqrt)
+{
+	if (bAlignedAndPadded)
+	{
+		check(A.Num() == B.Num() && A.Num() == WeightsSqrt.Num());
+		check(A.Num() % 4 == 0);
+		check(IsAligned(A.GetData(), alignof(VectorRegister4Float)));
+		check(IsAligned(B.GetData(), alignof(VectorRegister4Float)));
+		check(IsAligned(WeightsSqrt.GetData(), alignof(VectorRegister4Float)));
+		// sufficient condition to check for pointer overlapping
+		check(A.GetData() != B.GetData() && A.GetData() != WeightsSqrt.GetData());
+
+		const int32 NumVectors = A.Num() / 4;
+
+		const VectorRegister4Float* RESTRICT VA = reinterpret_cast<const VectorRegister4Float*>(A.GetData());
+		const VectorRegister4Float* RESTRICT VB = reinterpret_cast<const VectorRegister4Float*>(B.GetData());
+		const VectorRegister4Float* RESTRICT VW = reinterpret_cast<const VectorRegister4Float*>(WeightsSqrt.GetData());
+
+		VectorRegister4Float PartialCost = VectorZero();
+		for (int32 VectorIdx = 0; VectorIdx < NumVectors; ++VectorIdx, ++VA, ++VB, ++VW)
+		{
+			const VectorRegister4Float Diff = VectorSubtract(*VA, *VB);
+			const VectorRegister4Float WeightedDiff = VectorMultiply(Diff, *VW);
+			PartialCost = VectorMultiplyAdd(WeightedDiff, WeightedDiff, PartialCost);
+		}
+
+		// calculating PartialCost.X + PartialCost.Y + PartialCost.Z + PartialCost.W
+		VectorRegister4Float Swizzle = VectorSwizzle(PartialCost, 1, 0, 3, 2);	// (Y, X, W, Z) of PartialCost
+		PartialCost = VectorAdd(PartialCost, Swizzle);							// (X + Y, Y + X, Z + W, W + Z)
+		Swizzle = VectorSwizzle(PartialCost, 2, 3, 0, 1);						// (Z + W, W + Z, X + Y, Y + X)
+		PartialCost = VectorAdd(PartialCost, Swizzle);							// (X + Y + Z + W, Y + X + W + Z, Z + W + X + Y, W + Z + Y + X)
+		float Cost;
+		VectorStoreFloat1(PartialCost, &Cost);
+
+		// keeping this debug code to validate CompareAlignedFeatureVectors against CompareFeatureVectors
+		//#if DO_CHECK
+		//	const float EigenCost = CompareFeatureVectors(A, B, WeightsSqrt);
+		//
+		//	if (!FMath::IsNearlyEqual(Cost, EigenCost))
+		//	{
+		//		const float RelativeDifference = Cost > EigenCost ? (Cost - EigenCost) / Cost : (EigenCost - Cost) / EigenCost;
+		//		check(FMath::IsNearlyZero(RelativeDifference, UE_KINDA_SMALL_NUMBER));
+		//	}
+		//#endif //DO_CHECK
+
+		return Cost;
+	}
+	else
+	{
+		check(A.Num() == B.Num() && A.Num() == WeightsSqrt.Num());
+
+		Eigen::Map<const Eigen::ArrayXf> VA(A.GetData(), A.Num());
+		Eigen::Map<const Eigen::ArrayXf> VB(B.GetData(), B.Num());
+		Eigen::Map<const Eigen::ArrayXf> VW(WeightsSqrt.GetData(), WeightsSqrt.Num());
+
+		return ((VA - VB) * VW).square().sum();
+	}
+}
+
+} // namespace UE::PoseSearch
diff --git a/Engine/Plugins/Animation/PoseSearch/Source/Runtime/Private/PoseSearchInteractionIsland.cpp b/Engine/Plugins/Animation/PoseSearch/Source/Runtime/Private/PoseSearchInteractionIsland.cpp
--- a/Engine/Plugins/Animation/PoseSearch/Source/Runtime/Private/PoseSearchInteractionIsland.cpp
+++ b/Engine/Plugins/Animation/PoseSearch/Source/Runtime/Private/PoseSearchInteractionIsland.cpp
@@ -971,7 +971,7 @@ void FInteractionIsland::LogTickDependencies() const
 }
 #endif // ENABLE_ANIM_DEBUG
 
-bool FInteractionIsland::DoSearch_AnyThread(const UObject* AnimContext, const FPoseSearchContinuingProperties& ContinuingProperties, FPoseSearchBlueprintResult& Result)
+bool FInteractionIsland::DoSearch_AnyThread(const UObject* AnimContext, const TConstArrayView<FValidInteractionSearch> ValidInteractionSearches, FPoseSearchBlueprintResult& Result)
 {
 	check(AnimContext);
 
@@ -1049,8 +1049,8 @@ bool FInteractionIsland::DoSearch_AnyThread(const UObject* AnimContext, const FP
 		FPoseSearchFutureProperties PoseSearchFutureProperties;
 
 		// @todo: we could perform multiple UPoseSearchLibrary::MotionMatch in parallel!
-		PoseSearchResults[SearchIndex] = UPoseSearchLibrary::MotionMatch(AnimContexts, SearchContext.Roles,
-			PoseHistories, AssetsToSearch, SearchContext.GetContinuingProperties(), PoseSearchFutureProperties, FPoseSearchEvent());
+		PoseSearchResults[SearchIndex] = UPoseSearchLibrary::MotionMatch(AnimContexts, SearchContext.Roles, PoseHistories, 
+			AssetsToSearch, SearchContext.GetContinuingProperties(ValidInteractionSearches), PoseSearchFutureProperties, FPoseSearchEvent());
 	}
 		
 	InitSearchResults(SearchResults, PoseSearchResults, SearchContexts);
diff --git a/Engine/Plugins/Animation/PoseSearch/Source/Runtime/Private/PoseSearchInteractionLibrary.cpp b/Engine/Plugins/Animation/PoseSearch/Source/Runtime/Private/PoseSearchInteractionLibrary.cpp
--- a/Engine/Plugins/Animation/PoseSearch/Source/Runtime/Private/PoseSearchInteractionLibrary.cpp
+++ b/Engine/Plugins/Animation/PoseSearch/Source/Runtime/Private/PoseSearchInteractionLibrary.cpp
@@ -6,25 +6,25 @@
 #include "PoseSearch/PoseSearchInteractionSubsystem.h"
 #include "Engine/World.h"
 
-FPoseSearchBlueprintResult UPoseSearchInteractionLibrary::MotionMatchInteraction_Pure(TArray<FPoseSearchInteractionAvailability> Availabilities, const UObject* AnimContext, FPoseSearchContinuingProperties ContinuingProperties, FName PoseHistoryName, bool bValidateResultAgainstAvailabilities)
+FPoseSearchBlueprintResult UPoseSearchInteractionLibrary::MotionMatchInteraction_Pure(TArray<FPoseSearchInteractionAvailability> Availabilities, const UObject* AnimContext, FName PoseHistoryName, bool bValidateResultAgainstAvailabilities)
 {
 	FPoseSearchBlueprintResult Result;
-	MotionMatchInteraction(Result, Availabilities, AnimContext, ContinuingProperties, PoseHistoryName, nullptr, bValidateResultAgainstAvailabilities);
+	MotionMatchInteraction(Result, Availabilities, AnimContext, PoseHistoryName, nullptr, bValidateResultAgainstAvailabilities);
 	return Result;
 }
 
-FPoseSearchBlueprintResult UPoseSearchInteractionLibrary::MotionMatchInteraction(TArray<FPoseSearchInteractionAvailability> Availabilities, const UObject* AnimContext, FPoseSearchContinuingProperties ContinuingProperties, FName PoseHistoryName, bool bValidateResultAgainstAvailabilities)
+FPoseSearchBlueprintResult UPoseSearchInteractionLibrary::MotionMatchInteraction(TArray<FPoseSearchInteractionAvailability> Availabilities, const UObject* AnimContext, FName PoseHistoryName, bool bValidateResultAgainstAvailabilities)
 {
 	FPoseSearchBlueprintResult Result;
-	MotionMatchInteraction(Result, Availabilities, AnimContext, ContinuingProperties, PoseHistoryName, nullptr, bValidateResultAgainstAvailabilities);
+	MotionMatchInteraction(Result, Availabilities, AnimContext, PoseHistoryName, nullptr, bValidateResultAgainstAvailabilities);
 	return Result;
 }
 
-void UPoseSearchInteractionLibrary::MotionMatchInteraction(FPoseSearchBlueprintResult& Result, const TArrayView<const FPoseSearchInteractionAvailability> Availabilities, const UObject* AnimContext, const FPoseSearchContinuingProperties& ContinuingProperties, FName PoseHistoryName, const UE::PoseSearch::IPoseHistory* PoseHistory, bool bValidateResultAgainstAvailabilities)
+void UPoseSearchInteractionLibrary::MotionMatchInteraction(FPoseSearchBlueprintResult& Result, const TArrayView<const FPoseSearchInteractionAvailability> Availabilities, const UObject* AnimContext, FName PoseHistoryName, const UE::PoseSearch::IPoseHistory* PoseHistory, bool bValidateResultAgainstAvailabilities)
 {
 	if (UPoseSearchInteractionSubsystem* InteractionSubsystem = UPoseSearchInteractionSubsystem::GetSubsystem_AnyThread(AnimContext))
 	{
-		InteractionSubsystem->Query_AnyThread(Availabilities, AnimContext, ContinuingProperties, Result, PoseHistoryName, PoseHistory, bValidateResultAgainstAvailabilities);
+		InteractionSubsystem->Query_AnyThread(Availabilities, AnimContext, Result, PoseHistoryName, PoseHistory, bValidateResultAgainstAvailabilities);
 	}
 	else
 	{
diff --git a/Engine/Plugins/Animation/PoseSearch/Source/Runtime/Private/PoseSearchInteractionSubsystem.cpp b/Engine/Plugins/Animation/PoseSearch/Source/Runtime/Private/PoseSearchInteractionSubsystem.cpp
--- a/Engine/Plugins/Animation/PoseSearch/Source/Runtime/Private/PoseSearchInteractionSubsystem.cpp
+++ b/Engine/Plugins/Animation/PoseSearch/Source/Runtime/Private/PoseSearchInteractionSubsystem.cpp
@@ -1172,8 +1172,7 @@ TStatId UPoseSearchInteractionSubsystem::GetStatId() const
 }
 
 void UPoseSearchInteractionSubsystem::Query_AnyThread(const TArrayView<const FPoseSearchInteractionAvailability> Availabilities, const UObject* AnimContext, 
-	const FPoseSearchContinuingProperties& ContinuingProperties, FPoseSearchBlueprintResult& Result,
-	FName PoseHistoryName, const UE::PoseSearch::IPoseHistory* PoseHistory, bool bValidateResultAgainstAvailabilities)
+	FPoseSearchBlueprintResult& Result, FName PoseHistoryName, const UE::PoseSearch::IPoseHistory* PoseHistory, bool bValidateResultAgainstAvailabilities)
 {
 	QUICK_SCOPE_CYCLE_COUNTER(STAT_UPoseSearchInteractionSubsystem_Query_AnyThread);
 
@@ -1191,7 +1190,7 @@ void UPoseSearchInteractionSubsystem::Query_AnyThread(const TArrayView<const FPo
 	// if we find AnimContext in an island, we perform ALL the Island motion matching searches.
 	if (FInteractionIsland* Island = FindIsland(AnimContext))
 	{
-		Island->DoSearch_AnyThread(AnimContext, ContinuingProperties, Result);
+		Island->DoSearch_AnyThread(AnimContext, ValidInteractionSearches, Result);
 
 		if (bValidateResultAgainstAvailabilities && Result.SelectedAnim)
 		{
diff --git a/Engine/Plugins/Animation/PoseSearch/Source/Runtime/Private/PoseSearchLibrary.cpp b/Engine/Plugins/Animation/PoseSearch/Source/Runtime/Private/PoseSearchLibrary.cpp
--- a/Engine/Plugins/Animation/PoseSearch/Source/Runtime/Private/PoseSearchLibrary.cpp
+++ b/Engine/Plugins/Animation/PoseSearch/Source/Runtime/Private/PoseSearchLibrary.cpp
@@ -1159,14 +1159,13 @@ UE::PoseSearch::FSearchResult UPoseSearchLibrary::MotionMatch(
 	FReconstructedPreviousSearchBestResultMap ReconstructedPreviousSearchBestResultMap;
 	FSearchResult ReconstructedPreviousSearchResult;
 	FSearchContext SearchContext(DesiredPermutationTimeOffset, nullptr, ReconstructedPreviousSearchResult, FFloatInterval(0.f, 0.f), EventToSearch);
+	SearchContext.SetIsContinuingInteraction(ContinuingProperties.bIsContinuingInteraction);
 
 	for (int32 RoleIndex = 0; RoleIndex < Roles.Num(); ++RoleIndex)
 	{
 		SearchContext.AddRole(Roles[RoleIndex], &Contexts[RoleIndex], PoseHistories[RoleIndex]);
 	}
 	
-	bool bAsyncBuildIndexInProgress = false;
-
 	// collecting all the possible continuing pose search (it could be multiple searches, but most likely only one)
 	float DeltaSeconds = FiniteDelta;
 	if (Contexts[0].ObjectParams.Num() > 0)
diff --git a/Engine/Plugins/Animation/PoseSearch/Source/Runtime/Private/PoseSearchResult.cpp b/Engine/Plugins/Animation/PoseSearch/Source/Runtime/Private/PoseSearchResult.cpp
--- a/Engine/Plugins/Animation/PoseSearch/Source/Runtime/Private/PoseSearchResult.cpp
+++ b/Engine/Plugins/Animation/PoseSearch/Source/Runtime/Private/PoseSearchResult.cpp
@@ -192,7 +192,7 @@ const UAnimationAsset* FSearchResult::GetCurrentResultAnimationAsset(const FRole
 
 void FSearchResult::InitFrom(const FPoseSearchBlueprintResult& BlueprintResult)
 {
-	PoseCost = FPoseSearchCost(BlueprintResult.SearchCost, 0.f, 0.f);
+	PoseCost = FPoseSearchCost(BlueprintResult.SearchCost, 0.f, 0.f, 0.f);
 	PoseIdx = BlueprintResult.SelectedDatabase ? BlueprintResult.SelectedDatabase->GetPoseIndex(BlueprintResult.SelectedAnim.Get(), BlueprintResult.SelectedTime, BlueprintResult.bIsMirrored, BlueprintResult.BlendParameters) : INDEX_NONE;
 	EventPoseIdx = INDEX_NONE;
 	Database = BlueprintResult.SelectedDatabase;
diff --git a/Engine/Plugins/Animation/PoseSearch/Source/Runtime/Public/PoseSearch/PoseSearchContext.h b/Engine/Plugins/Animation/PoseSearch/Source/Runtime/Public/PoseSearch/PoseSearchContext.h
--- a/Engine/Plugins/Animation/PoseSearch/Source/Runtime/Public/PoseSearch/PoseSearchContext.h
+++ b/Engine/Plugins/Animation/PoseSearch/Source/Runtime/Public/PoseSearch/PoseSearchContext.h
@@ -256,6 +256,11 @@ struct POSESEARCH_API FSearchContext
 	const UPoseSearchFeatureChannel* GetCachedChannelData(uint32 ChannelUniqueIdentifier, const UPoseSearchFeatureChannel* Channel, TConstArrayView<float>& CachedChannelData);
 	bool IsUseCachedChannelData() const { return bUseCachedChannelData; }
 	void SetUseCachedChannelData(bool bInUseCachedChannelData) { bUseCachedChannelData = bInUseCachedChannelData; }
+	
+	// Experimental, this feature might be removed without warning, not for production use
+	bool IsContinuingInteraction() const { return bIsContinuingInteraction; }
+	// Experimental, this feature might be removed without warning, not for production use
+	void SetIsContinuingInteraction(bool bInIsContinuingInteraction) { bIsContinuingInteraction = bInIsContinuingInteraction; }
 
 private:
 	float GetSampleCurveValueInternal(float SampleTime, const FName& CurveName, const FRole& SampleRole);
@@ -290,7 +295,11 @@ private:
 	TMap<uint32, FCachedChannel, TInlineSetAllocator<PreallocatedCachedChannelDataNum, TMemStackSetAllocator<>>> CachedChannels;
 
 	float CurrentBestTotalCost = MAX_flt;
-	
+
+	// @todo: add FPoseSearchContinuingProperties here and reconstuct the continuing pose FSearchResult for the previous frame from it
+	// Experimental, this feature might be removed without warning, not for production use
+	bool bIsContinuingInteraction = false;
+
 #if WITH_EDITOR
 	bool bAsyncBuildIndexInProgress = false;
 #endif // WITH_EDITOR
diff --git a/Engine/Plugins/Animation/PoseSearch/Source/Runtime/Public/PoseSearch/PoseSearchCost.h b/Engine/Plugins/Animation/PoseSearch/Source/Runtime/Public/PoseSearch/PoseSearchCost.h
--- a/Engine/Plugins/Animation/PoseSearch/Source/Runtime/Public/PoseSearch/PoseSearchCost.h
+++ b/Engine/Plugins/Animation/PoseSearch/Source/Runtime/Public/PoseSearch/PoseSearchCost.h
@@ -9,35 +9,71 @@ struct POSESEARCH_API FPoseSearchCost
 {
 	GENERATED_BODY()
 public:
+
+PRAGMA_DISABLE_DEPRECATION_WARNINGS
+	// Note: We need to explicitly disable warnings on these constructors/operators for clang to be happy with deprecated variables
+	// this is a requirement for clang to compile without warnings.
 	FPoseSearchCost() = default;
-	FPoseSearchCost(float DissimilarityCost, float NotifyCostAddend, float ContinuingPoseCostAddend)
-	: TotalCost(DissimilarityCost + NotifyCostAddend + ContinuingPoseCostAddend)
-	{
+	~FPoseSearchCost() = default;
+	FPoseSearchCost(const FPoseSearchCost&) = default;
+	FPoseSearchCost(FPoseSearchCost&&) = default;
+	FPoseSearchCost& operator=(const FPoseSearchCost&) = default;
+	FPoseSearchCost& operator=(FPoseSearchCost&&) = default;
+PRAGMA_ENABLE_DEPRECATION_WARNINGS
+
+	FPoseSearchCost(float DissimilarityCost, float InNotifyCostAddend, float InContinuingPoseCostAddend, float InContinuingInteractionCostAddend = 0.f)
+		: TotalCost(DissimilarityCost + InNotifyCostAddend + InContinuingPoseCostAddend + InContinuingInteractionCostAddend)
 #if WITH_EDITORONLY_DATA
-		CostAddend = NotifyCostAddend + ContinuingPoseCostAddend;
+		, NotifyCostAddend(InNotifyCostAddend)
+		, ContinuingPoseCostAddend(InContinuingPoseCostAddend)
+		, ContinuingInteractionCostAddend(InContinuingInteractionCostAddend)
+PRAGMA_DISABLE_DEPRECATION_WARNINGS
+		, CostAddend(InNotifyCostAddend + InContinuingPoseCostAddend + InContinuingInteractionCostAddend)
+PRAGMA_ENABLE_DEPRECATION_WARNINGS
 #endif // WITH_EDITORONLY_DATA
+	{
 	}
 
 	static bool IsCostValid(const float Cost) { return Cost != MAX_flt; }
-
 	bool IsValid() const { return IsCostValid(TotalCost); }
 	
 	UE_DEPRECATED(5.6, "Use operator float() instead.")
 	float GetTotalCost() const { return TotalCost; }
+
 	operator float() const { return TotalCost; }
 
+#if WITH_EDITORONLY_DATA
+	float GetCostAddend() const { return GetNotifyCostAddend() + GetContinuingPoseCostAddend() + GetContinuingInteractionCostAddend(); }
+
+	float GetNotifyCostAddend() const { return NotifyCostAddend; }
+	float GetContinuingPoseCostAddend() const { return ContinuingPoseCostAddend; }
+	float GetContinuingInteractionCostAddend() const { return ContinuingInteractionCostAddend; }
+#endif // WITH_EDITORONLY_DATA
+
 protected:
+	// TotalCost is the sum of all the Cost contributions (dissimilarity, notifies, continuing pose, continuing interaction costs)
 	UPROPERTY()
 	float TotalCost = MAX_flt;
 
 #if WITH_EDITORONLY_DATA
-public:
 
-	float GetCostAddend() const { return CostAddend; }
+	// Notify Cost Bias contribution
+	UPROPERTY()
+	float NotifyCostAddend = 0.f;
+	
+	// Continuing Pose Cost Bias contribution
+	UPROPERTY()
+	float ContinuingPoseCostAddend = 0.f;
 
-	// Contribution from ModifyCost anim notify, and ContinuingPoseCostAddend
+	// Experimental, this feature might be removed without warning, not for production use
+	// Continuing Interaction Cost Bias contribution
 	UPROPERTY()
-	float CostAddend = 0.f;
+	float ContinuingInteractionCostAddend = 0.f;
 
+public:
+
+	UE_DEPRECATED(5.6, "Use GetCostAddend instead")
+	UPROPERTY()
+	float CostAddend = 0.f;
 #endif // WITH_EDITORONLY_DATA
 };
\ No newline at end of file
diff --git a/Engine/Plugins/Animation/PoseSearch/Source/Runtime/Public/PoseSearch/PoseSearchDatabase.h b/Engine/Plugins/Animation/PoseSearch/Source/Runtime/Public/PoseSearch/PoseSearchDatabase.h
--- a/Engine/Plugins/Animation/PoseSearch/Source/Runtime/Public/PoseSearch/PoseSearchDatabase.h
+++ b/Engine/Plugins/Animation/PoseSearch/Source/Runtime/Public/PoseSearch/PoseSearchDatabase.h
@@ -442,6 +442,14 @@ public:
 	UPROPERTY(EditAnywhere, Category = "Database")
 	float LoopingCostBias = -0.005f;
 
+	// Cost added to poses from this database (if setup with UMultiAnimAsset assets in "interaction mode") as continuation of a previous interaction.
+	// This allows users to apply a cost bias (positive or negative) to poses following up an interaction.
+	// This is useful to help the system stay in one animation segment longer, or shorter depending on how you set this bias.
+	// Negative values make it more likely to be picked, or stayed in, positive values make it less likely to be picked or stay in.
+	// Note: excluded from DDC hash, since used only at runtime
+	UPROPERTY(EditAnywhere, Category = "Database", meta = (ExcludeFromHash))
+	float ContinuingInteractionCostBias = 0.f;
+
 #if WITH_EDITORONLY_DATA
 	// These settings allow users to trim the start and end of animations in the database to preserve start/end frames for blending, and prevent the system from selecting the very last frames before it blends out.
 	// valid animation frames will be AnimationAssetTimeStart + ExcludeFromDatabaseParameters.Min, AnimationAssetTimeEnd + ExcludeFromDatabaseParameters.Max
@@ -627,7 +635,7 @@ private:
 	typedef TArray<int32, TInlineAllocator<256, TMemStackAllocator<>>> FNonSelectableIdx;
 	void PopulateNonSelectableIdx(FNonSelectableIdx& NonSelectableIdx, UE::PoseSearch::FSearchContext& SearchContext
 #if UE_POSE_SEARCH_TRACE_ENABLED
-		, TConstArrayView<float> QueryValues, TConstArrayView<float> DynamicWeightsSqrt
+		, float ContinuingPoseCostAddend, float ContinuingInteractionCostAddend, TConstArrayView<float> QueryValues, TConstArrayView<float> DynamicWeightsSqrt
 #endif //UE_POSE_SEARCH_TRACE_ENABLED
 	) const;
 
diff --git a/Engine/Plugins/Animation/PoseSearch/Source/Runtime/Public/PoseSearch/PoseSearchIndex.h b/Engine/Plugins/Animation/PoseSearch/Source/Runtime/Public/PoseSearch/PoseSearchIndex.h
--- a/Engine/Plugins/Animation/PoseSearch/Source/Runtime/Public/PoseSearch/PoseSearchIndex.h
+++ b/Engine/Plugins/Animation/PoseSearch/Source/Runtime/Public/PoseSearch/PoseSearchIndex.h
@@ -603,14 +603,11 @@ struct FSearchIndex : public FSearchIndexBase
 		return PCAValues.Num() / DataCardinality;
 	}
 
-	UE_DEPRECATED(5.6, "Use other ComparePoses signatures instead")
+	UE_DEPRECATED(5.6, "Use the bAlignedAndPadded templated way of calculating FPoseSearchCost(CompareFeatureVectors<false>(PoseValues, QueryValues, WeightsSqrt), PoseMetadata[PoseIdx].GetCostAddend(), ContinuingPoseCostBias, 0.f) instead")
 	POSESEARCH_API FPoseSearchCost ComparePoses(int32 PoseIdx, float ContinuingPoseCostBias, TConstArrayView<float> PoseValues, TConstArrayView<float> QueryValues) const;
-	UE_DEPRECATED(5.6, "Use other CompareAlignedPoses signatures instead")
+	UE_DEPRECATED(5.6, "Use the bAlignedAndPadded templated way of calculating FPoseSearchCost(CompareFeatureVectors<true>(PoseValues, QueryValues, WeightsSqrt), PoseMetadata[PoseIdx].GetCostAddend(), ContinuingPoseCostBias, 0.f) instead")
 	POSESEARCH_API FPoseSearchCost CompareAlignedPoses(int32 PoseIdx, float ContinuingPoseCostBias, TConstArrayView<float> PoseValues, TConstArrayView<float> QueryValues) const;
 
-	POSESEARCH_API FPoseSearchCost ComparePoses(int32 PoseIdx, float ContinuingPoseCostBias, TConstArrayView<float> PoseValues, TConstArrayView<float> QueryValues, TConstArrayView<float> DynamicWeightsSqrt) const;
-	POSESEARCH_API FPoseSearchCost CompareAlignedPoses(int32 PoseIdx, float ContinuingPoseCostBias, TConstArrayView<float> PoseValues, TConstArrayView<float> QueryValues, TConstArrayView<float> DynamicWeightsSqrt) const;
-
 	void PruneDuplicatePCAValues(float SimilarityThreshold, int32 NumberOfPrincipalComponents);
 	void PrunePCAValuesFromBlockTransitionPoses(int32 NumberOfPrincipalComponents);
 
@@ -640,11 +637,7 @@ struct FVPTreeDataSource
         return SearchIndex.GetNumValuesVectors(DataCardinality);
     }
 
-    static float GetDistance(const TConstArrayView<float> A, const TConstArrayView<float> B)
-    {
-		// Estracting the Sqrt to satisfy the triangle inequality metric space requirements, since a <= b+c doesn't imply a^2 <= b^2 + c^2
-		return FMath::Sqrt(CompareFeatureVectors(A, B));
-    }
+    static float GetDistance(const TConstArrayView<float> A, const TConstArrayView<float> B);
 
 private:
 	const FSearchIndex& SearchIndex;
diff --git a/Engine/Plugins/Animation/PoseSearch/Source/Runtime/Public/PoseSearch/PoseSearchInteractionIsland.h b/Engine/Plugins/Animation/PoseSearch/Source/Runtime/Public/PoseSearch/PoseSearchInteractionIsland.h
--- a/Engine/Plugins/Animation/PoseSearch/Source/Runtime/Public/PoseSearch/PoseSearchInteractionIsland.h
+++ b/Engine/Plugins/Animation/PoseSearch/Source/Runtime/Public/PoseSearch/PoseSearchInteractionIsland.h
@@ -40,6 +40,16 @@ struct FInteractionSearchContextBase
 #endif // DO_CHECK
 };
 
+// Experimental, this feature might be removed without warning, not for production use
+typedef TArray<TPair<TWeakObjectPtr<AActor>, TWeakObjectPtr<AActor>>> FDisabledCollisions;
+
+// Experimental, this feature might be removed without warning, not for production use
+struct FValidInteractionSearch
+{
+	FInteractionSearchContextBase SearchContext;
+	FDisabledCollisions DisabledCollisions;
+};
+
 // @todo: make FPoseSearchContinuingProperties.PlayingAsset a TWeakObjectPtr and remove this struct!
 // Experimental, this feature might be removed without warning, not for production use
 struct FInteractionSearchContext : FInteractionSearchContextBase
@@ -50,9 +60,10 @@ struct FInteractionSearchContext : FInteractionSearchContextBase
 	bool bIsPlayingAssetMirrored = false;
 	FVector PlayingAssetBlendParameters = FVector::ZeroVector;
 	EPoseSearchInterruptMode InterruptMode = EPoseSearchInterruptMode::DoNotInterrupt;
+	bool bIsContinuingInteraction = false;
 	TArray<int32, TInlineAllocator<PreallocatedRolesNum>> TickPriorities;
 
-	FPoseSearchContinuingProperties GetContinuingProperties() const
+	FPoseSearchContinuingProperties GetContinuingProperties(const TConstArrayView<FValidInteractionSearch> ValidInteractionSearches) const
 	{
 		FPoseSearchContinuingProperties ContinuingProperties;
 		ContinuingProperties.PlayingAsset = PlayingAsset.Get();
@@ -60,6 +71,16 @@ struct FInteractionSearchContext : FInteractionSearchContextBase
 		ContinuingProperties.bIsPlayingAssetMirrored = bIsPlayingAssetMirrored;
 		ContinuingProperties.PlayingAssetBlendParameters = PlayingAssetBlendParameters;
 		ContinuingProperties.InterruptMode = InterruptMode;
+
+		for (const FValidInteractionSearch& ValidInteractionSearch : ValidInteractionSearches)
+		{
+			if (ValidInteractionSearch.SearchContext.IsEquivalent(*this))
+			{
+				ContinuingProperties.bIsContinuingInteraction = true;
+				break;
+			}
+		}
+
 		return ContinuingProperties;
 	}
 };
@@ -94,7 +115,7 @@ struct FInteractionIsland
 	FInteractionIsland(ULevel* Level, UPoseSearchInteractionSubsystem* Subsystem);
 	~FInteractionIsland();
 
-	bool DoSearch_AnyThread(const UObject* AnimContext, const FPoseSearchContinuingProperties& ContinuingProperties, FPoseSearchBlueprintResult& Result);
+	bool DoSearch_AnyThread(const UObject* AnimContext, const TConstArrayView<FValidInteractionSearch> ValidInteractionSearches, FPoseSearchBlueprintResult& Result);
 	bool GetResult_AnyThread(const UObject* AnimContext, FPoseSearchBlueprintResult& Result, bool bCompareOwningActors = false);
 
 	const TArray<TWeakObjectPtr<UActorComponent>>& GetTickActorComponents() const { return TickActorComponents; }
diff --git a/Engine/Plugins/Animation/PoseSearch/Source/Runtime/Public/PoseSearch/PoseSearchInteractionLibrary.h b/Engine/Plugins/Animation/PoseSearch/Source/Runtime/Public/PoseSearch/PoseSearchInteractionLibrary.h
--- a/Engine/Plugins/Animation/PoseSearch/Source/Runtime/Public/PoseSearch/PoseSearchInteractionLibrary.h
+++ b/Engine/Plugins/Animation/PoseSearch/Source/Runtime/Public/PoseSearch/PoseSearchInteractionLibrary.h
@@ -80,13 +80,13 @@ public:
 	// PoseHistoryName is the name of the pose history node used for the associated motion matching search
 	// if bValidateResultAgainstAvailabilities is true, the result will be invalidated if doesn't respect the new availabilities
 	UFUNCTION(Experimental, BlueprintPure, Category = "Animation|Pose Search", meta = (BlueprintThreadSafe))
-	static FPoseSearchBlueprintResult MotionMatchInteraction_Pure(TArray<FPoseSearchInteractionAvailability> Availabilities, const UObject* AnimContext, FPoseSearchContinuingProperties ContinuingProperties, FName PoseHistoryName, bool bValidateResultAgainstAvailabilities = true);
+	static FPoseSearchBlueprintResult MotionMatchInteraction_Pure(TArray<FPoseSearchInteractionAvailability> Availabilities, const UObject* AnimContext, FName PoseHistoryName, bool bValidateResultAgainstAvailabilities = true);
 
 	// BlueprintCallable version of MotionMatchInteraction_Pure
 	UFUNCTION(Experimental, BlueprintCallable, Category = "Animation|Pose Search", meta = (BlueprintThreadSafe))
-	static FPoseSearchBlueprintResult MotionMatchInteraction(TArray<FPoseSearchInteractionAvailability> Availabilities, const UObject* AnimContext, FPoseSearchContinuingProperties ContinuingProperties, FName PoseHistoryName, bool bValidateResultAgainstAvailabilities = true);
+	static FPoseSearchBlueprintResult MotionMatchInteraction(TArray<FPoseSearchInteractionAvailability> Availabilities, const UObject* AnimContext, FName PoseHistoryName, bool bValidateResultAgainstAvailabilities = true);
 
-	static void MotionMatchInteraction(FPoseSearchBlueprintResult& Result, const TArrayView<const FPoseSearchInteractionAvailability> Availabilities, const UObject* AnimContext, const FPoseSearchContinuingProperties& ContinuingProperties, FName PoseHistoryName, const UE::PoseSearch::IPoseHistory* PoseHistory, bool bValidateResultAgainstAvailabilities = true);
+	static void MotionMatchInteraction(FPoseSearchBlueprintResult& Result, const TArrayView<const FPoseSearchInteractionAvailability> Availabilities, const UObject* AnimContext, FName PoseHistoryName, const UE::PoseSearch::IPoseHistory* PoseHistory, bool bValidateResultAgainstAvailabilities = true);
 
 	static FPoseSearchBlueprintResult GetCachedInteraction(const UObject* AnimContext, bool bCompareOwningActors = false);
 	
diff --git a/Engine/Plugins/Animation/PoseSearch/Source/Runtime/Public/PoseSearch/PoseSearchInteractionSubsystem.h b/Engine/Plugins/Animation/PoseSearch/Source/Runtime/Public/PoseSearch/PoseSearchInteractionSubsystem.h
--- a/Engine/Plugins/Animation/PoseSearch/Source/Runtime/Public/PoseSearch/PoseSearchInteractionSubsystem.h
+++ b/Engine/Plugins/Animation/PoseSearch/Source/Runtime/Public/PoseSearch/PoseSearchInteractionSubsystem.h
@@ -39,16 +39,6 @@ typedef TArray<FInteractionAvailability> FInteractionAvailabilities;
 // Experimental, this feature might be removed without warning, not for production use
 typedef TMap<FObjectKey, FInteractionAvailabilities> FAvailabilitiesMap;
 
-// Experimental, this feature might be removed without warning, not for production use
-typedef TArray<TPair<TWeakObjectPtr<AActor>, TWeakObjectPtr<AActor>>> FDisabledCollisions;
-
-// Experimental, this feature might be removed without warning, not for production use
-struct FValidInteractionSearch
-{
-	FInteractionSearchContextBase SearchContext;
-	FDisabledCollisions DisabledCollisions;
-};
-
 struct FAnimContextInfos;
 struct FTagToDatabases;
 
@@ -118,8 +108,7 @@ public:
 	// if bValidateResultAgainstAvailabilities is true, the result will be validated against the currently published availabilities,
 	// that could differ from the previous frame one for this AnimContext, and discarded if not valid within the context of the new availabilities
 	void Query_AnyThread(const TArrayView<const FPoseSearchInteractionAvailability> Availabilities, const UObject* AnimContext, 
-		const FPoseSearchContinuingProperties& ContinuingProperties, FPoseSearchBlueprintResult& Result,
-		FName PoseHistoryName, const UE::PoseSearch::IPoseHistory* PoseHistory, bool bValidateResultAgainstAvailabilities);
+		FPoseSearchBlueprintResult& Result,	FName PoseHistoryName, const UE::PoseSearch::IPoseHistory* PoseHistory, bool bValidateResultAgainstAvailabilities);
 
 	void GetResult_AnyThread(const UObject* AnimContext, FPoseSearchBlueprintResult& Result, bool bCompareOwningActors = false);
 
diff --git a/Engine/Plugins/Animation/PoseSearch/Source/Runtime/Public/PoseSearch/PoseSearchLibrary.h b/Engine/Plugins/Animation/PoseSearch/Source/Runtime/Public/PoseSearch/PoseSearchLibrary.h
--- a/Engine/Plugins/Animation/PoseSearch/Source/Runtime/Public/PoseSearch/PoseSearchLibrary.h
+++ b/Engine/Plugins/Animation/PoseSearch/Source/Runtime/Public/PoseSearch/PoseSearchLibrary.h
@@ -128,11 +128,14 @@ struct POSESEARCH_API FPoseSearchContinuingProperties
 	bool bIsPlayingAssetMirrored = false;
 
 	// PlayingAsset associated BlendParameters (if PlayingAsset is a UBlendSpace)
-	UPROPERTY(Transient, VisibleAnywhere, BlueprintReadOnly, Category=State)
+	UPROPERTY(VisibleAnywhere, BlueprintReadOnly, Category=State)
 	FVector PlayingAssetBlendParameters = FVector::ZeroVector;
 
 	UPROPERTY(EditAnywhere, BlueprintReadWrite, Category=State)
 	EPoseSearchInterruptMode InterruptMode = EPoseSearchInterruptMode::DoNotInterrupt;
+
+	UPROPERTY(EditAnywhere, BlueprintReadWrite, Category=State)
+	bool bIsContinuingInteraction = false;
 };
 
 UCLASS()
diff --git a/Engine/Plugins/Experimental/AnimNextPoseSearch/Source/AnimNextPoseSearch/Private/MotionMatchingTrait.cpp b/Engine/Plugins/Experimental/AnimNextPoseSearch/Source/AnimNextPoseSearch/Private/MotionMatchingTrait.cpp
--- a/Engine/Plugins/Experimental/AnimNextPoseSearch/Source/AnimNextPoseSearch/Private/MotionMatchingTrait.cpp
+++ b/Engine/Plugins/Experimental/AnimNextPoseSearch/Source/AnimNextPoseSearch/Private/MotionMatchingTrait.cpp
@@ -181,18 +181,8 @@ namespace UE::AnimNext
 			// WIP motion matched interactions
 			CheckInteractionThreadSafety(AnimContext);
 
-			FPoseSearchContinuingProperties ContinuingProperties;
-			if (MotionMatchingState.SearchResult.bIsInteraction)
-			{
-				// @todo: retrieve the playing UMultiAnimAsset from the UAnimAsset blendspace is playing via GetAnimAsset() and blend parameters via GetBlendParameters() and mirror state
-				ContinuingProperties.PlayingAsset = MotionMatchingState.SearchResult.SelectedAnim;
-				ContinuingProperties.PlayingAssetAccumulatedTime = TimelineNormalizedTime;
-				ContinuingProperties.bIsPlayingAssetMirrored = MotionMatchingState.SearchResult.bIsMirrored;
-				ContinuingProperties.PlayingAssetBlendParameters = MotionMatchingState.SearchResult.BlendParameters;
-			}
-
 			FPoseSearchBlueprintResult SearchResult;
-			UPoseSearchInteractionLibrary::MotionMatchInteraction(SearchResult, Availabilities, AnimContext, ContinuingProperties, FName(), PoseHistory, SharedData->GetbValidateResultAgainstAvailabilities(Binding));
+			UPoseSearchInteractionLibrary::MotionMatchInteraction(SearchResult, Availabilities, AnimContext, FName(), PoseHistory, SharedData->GetbValidateResultAgainstAvailabilities(Binding));
 			check(SearchResult.ActorRootTransforms.Num() == SearchResult.ActorRootBoneTransforms.Num());
 
 			if (SearchResult.SelectedAnim)
@@ -223,17 +213,24 @@ namespace UE::AnimNext
 				//MotionMatchingState.SearchResult.BlendParameters = GetBlendParameters();
 
 				// checking if the kept alive interaction has reached the end of animation 
-				if (MotionMatchingState.SearchResult.SelectedDatabase &&
-					INDEX_NONE == MotionMatchingState.SearchResult.SelectedDatabase->GetPoseIndex(MotionMatchingState.SearchResult.SelectedAnim.Get(), MotionMatchingState.SearchResult.SelectedTime, MotionMatchingState.SearchResult.bIsMirrored, MotionMatchingState.SearchResult.BlendParameters))
+				if (!MotionMatchingState.SearchResult.SelectedDatabase || INDEX_NONE == MotionMatchingState.SearchResult.SelectedDatabase->GetPoseIndex(MotionMatchingState.SearchResult.SelectedAnim.Get(), MotionMatchingState.SearchResult.SelectedTime, MotionMatchingState.SearchResult.bIsMirrored, MotionMatchingState.SearchResult.BlendParameters))
 				{
 					MotionMatchingState.SearchResult = FPoseSearchBlueprintResult();
 				}
 				else
 				{
-					// we're keeping alive only the animation part of the search result
-					MotionMatchingState.SearchResult.ActorRootTransforms.Reset();
-					MotionMatchingState.SearchResult.ActorRootBoneTransforms.Reset();
-					MotionMatchingState.SearchResult.AnimContexts.Reset();
+					const UAnimationAsset* AnimationAssetForRole = MotionMatchingState.SearchResult.GetAnimationAssetForRole();
+					if (!AnimationAssetForRole || MotionMatchingState.SearchResult.SelectedTime > AnimationAssetForRole->GetPlayLength())
+					{
+						MotionMatchingState.SearchResult = FPoseSearchBlueprintResult();
+					}
+					else
+					{
+						// we're keeping alive only the animation part of the search result
+						MotionMatchingState.SearchResult.ActorRootTransforms.Reset();
+						MotionMatchingState.SearchResult.ActorRootBoneTransforms.Reset();
+						MotionMatchingState.SearchResult.AnimContexts.Reset();
+					}
 				}
 			}
 			else
