PoseSearch - adding support for FPoseSearchInteractionAvailability::BroadPhaseRadiusIncrementOnInteraction

diff --git a/Engine/Plugins/Animation/PoseSearch/Source/Runtime/Private/PoseSearchInteractionIsland.cpp b/Engine/Plugins/Animation/PoseSearch/Source/Runtime/Private/PoseSearchInteractionIsland.cpp
--- a/Engine/Plugins/Animation/PoseSearch/Source/Runtime/Private/PoseSearchInteractionIsland.cpp
+++ b/Engine/Plugins/Animation/PoseSearch/Source/Runtime/Private/PoseSearchInteractionIsland.cpp
@@ -1050,7 +1050,7 @@ bool FInteractionIsland::DoSearch_AnyThread(const UObject* AnimContext, const TC
 
 		// @todo: we could perform multiple UPoseSearchLibrary::MotionMatch in parallel!
 		PoseSearchResults[SearchIndex] = UPoseSearchLibrary::MotionMatch(AnimContexts, SearchContext.Roles, PoseHistories, 
-			AssetsToSearch, SearchContext.GetContinuingProperties(ValidInteractionSearches), PoseSearchFutureProperties, FPoseSearchEvent());
+			AssetsToSearch, SearchContext.GetContinuingProperties(), PoseSearchFutureProperties, FPoseSearchEvent());
 	}
 		
 	InitSearchResults(SearchResults, PoseSearchResults, SearchContexts);
diff --git a/Engine/Plugins/Animation/PoseSearch/Source/Runtime/Private/PoseSearchInteractionSubsystem.cpp b/Engine/Plugins/Animation/PoseSearch/Source/Runtime/Private/PoseSearchInteractionSubsystem.cpp
--- a/Engine/Plugins/Animation/PoseSearch/Source/Runtime/Private/PoseSearchInteractionSubsystem.cpp
+++ b/Engine/Plugins/Animation/PoseSearch/Source/Runtime/Private/PoseSearchInteractionSubsystem.cpp
@@ -4,18 +4,13 @@
 #include "PoseSearch/AnimNode_PoseSearchHistoryCollector.h"
 #include "PoseSearch/PoseSearchDatabase.h"
 #include "PoseSearch/PoseSearchInteractionIsland.h"
+#include "PoseSearch/PoseSearchInteractionUtils.h"
 #include "PoseSearch/PoseSearchRole.h"
 #include "PoseSearch/PoseSearchSchema.h"
 #include "VisualLogger/VisualLogger.h"
 
 namespace UE::PoseSearch
 {
-
-#if ENABLE_DRAW_DEBUG && ENABLE_ANIM_DEBUG
-static bool GVarPoseSearchInteractionShowIslands = false;
-static FAutoConsoleVariableRef CVarPoseSearchInteractionShowIslands(TEXT("a.PoseSearchInteraction.ShowIslands"), GVarPoseSearchInteractionShowIslands, TEXT("Show Pose Search Interaction Islands"));
-#endif // ENABLE_DRAW_DEBUG && ENABLE_ANIM_DEBUG
-
 #if !NO_CVARS
 static bool GVarPoseSearchInteractionEnabled = true;
 static FAutoConsoleVariableRef CVarPoseSearchInteractionEnabled(TEXT("a.PoseSearchInteraction.Enabled"), GVarPoseSearchInteractionEnabled, TEXT("Enable/Disable Pose Search Interaction"));
@@ -55,81 +50,53 @@ const IPoseHistory* FInteractionAvailability::GetPoseHistory(const UObject* Anim
 	return nullptr;
 }
 
-struct FAnimContextInfo;
-struct FNearbyAnimContextInfo
-{
-	FNearbyAnimContextInfo(const FAnimContextInfo& InAnimContextInfo, float InDistanceSquared)
-		: AnimContextInfo(&InAnimContextInfo)
-		, DistanceSquared(InDistanceSquared)
-	{
-	}
-
-	const FAnimContextInfo* AnimContextInfo = nullptr;
-	float DistanceSquared = 0.f;
-};
 struct FAnimContextInfo
 {
-	float GetCachedDistanceSquaredFrom(const FAnimContextInfo& OtherAnimContextInfo) const
+	// performs broad phase analysis checking if at least one of the Availabilities associated to AnimContext can interact with OtherAnimContextInfo.
+	// This is a more relaxed analysis than the one performed in FRoledAnimContextInfo::CanInteractWith
+	bool CanInteractWith(const FAnimContextInfo& OtherAnimContextInfo) const
 	{
 		check(this != &OtherAnimContextInfo);
+		
+		// @todo: enable this code if we ended up requiring preventing interactions between the same actor!
+		// const AActor* AnimContextActor = GetContextOwningActor(AnimContext);
+		// const AActor* OtherAnimContextActor = GetContextOwningActor(OtherAnimContextInfo.AnimContext);
+		// if (AnimContextActor == OtherAnimContextActor)
+		// {
+		// 	return false;
+		// }
 
-		// since there should be only few NearbyAnimContextInfos we can efficiently search linearly rather than calculating and
-		// return (Location - OtherAnimContextInfo.Location).SquaredLength()...
-		for (const FNearbyAnimContextInfo& NearbyAnimContextInfo : NearbyAnimContextInfos)
-		{
-			if (NearbyAnimContextInfo.AnimContextInfo == &OtherAnimContextInfo)
-			{
-				return NearbyAnimContextInfo.DistanceSquared;
-			}
-		}
-
-#if DO_CHECK
-		// sanity check: if we couldn't find OtherAnimContextInfo in the NearbyAnimContextInfos, it means it didn't pass the broad phase,
-		// that means its location is further away than MaxDistanceSquared
 		const FVector DeltaLocation = Location - OtherAnimContextInfo.Location;
 		const float DistanceSquared = DeltaLocation.SquaredLength();
-		check(DistanceSquared > MaxDistanceSquared);
-#endif // DO_CHECK
+		const float MaxDistance = FMath::Min(GetAvailabilitiesMaxBroadPhaseRadius(), OtherAnimContextInfo.GetAvailabilitiesMaxBroadPhaseRadius());
+		const float MaxDistanceSquared = MaxDistance * MaxDistance;
+		return DistanceSquared <= MaxDistanceSquared;
+	}
 
-		// we could pass any number higher than MaxDistanceSquared. I like UE_BIG_NUMBER more :)
-		return UE_BIG_NUMBER;
+	float GetAvailabilitiesMaxBroadPhaseRadius() const
+	{
+		float AvailabilitiesMaxBroadPhaseRadius = 0.f;
+		check(Availabilities);
+		for (const FInteractionAvailability& Availability : *Availabilities)
+		{
+			// @todo: optimize the AvailabilitiesMaxBroadPhaseRadius, since adding Availability.BroadPhaseRadiusIncrementOnInteraction is required ONLY if AnimContext is already part of an interaction
+			AvailabilitiesMaxBroadPhaseRadius = FMath::Max(AvailabilitiesMaxBroadPhaseRadius, Availability.BroadPhaseRadius + Availability.BroadPhaseRadiusIncrementOnInteraction);
+		}
+		return AvailabilitiesMaxBroadPhaseRadius;
 	}
 
 	const UObject* AnimContext = nullptr;
-	int32 TickPriority = DefaultTickPriority;
+	// array of all valid avalilabilities associated to AnimContext
 	const FInteractionAvailabilities* Availabilities = nullptr;
+	// cached AnimContext location
 	FVector Location = FVector::ZeroVector;
-	float MaxDistanceSquared = 0.f;
-	TArray<FNearbyAnimContextInfo, TInlineAllocator<8, TMemStackAllocator<>>> NearbyAnimContextInfos;
+	TArray<const FAnimContextInfo*, TInlineAllocator<8, TMemStackAllocator<>>> NearbyAnimContextInfos;
 };
 struct FAnimContextInfos : TArray<FAnimContextInfo, TMemStackAllocator<>> {};
 
 typedef TArray<const UPoseSearchDatabase*, TInlineAllocator<32, TMemStackAllocator<>>> FDatabasesPerTag;
 struct FTagToDatabases : TMap<FName, FDatabasesPerTag, TInlineSetAllocator<8, TMemStackSetAllocator<>>> {};
 
-template <typename EvaluateCombinationType>
-static void GeneratePermutationsRecursive(int32 DataCardinality, int32 DataIndex, TArrayView<int32> Combination, int32 CombinationIndex, EvaluateCombinationType EvaluateCombination)
-{
-	if (CombinationIndex == Combination.Num())
-	{ 
-		EvaluateCombination(Combination);
-	}
-	else if (DataIndex < DataCardinality)
-	{
-		Combination[CombinationIndex] = DataIndex;
-		GeneratePermutationsRecursive(DataCardinality, DataIndex + 1, Combination, CombinationIndex + 1, EvaluateCombination);
-		GeneratePermutationsRecursive(DataCardinality, DataIndex + 1, Combination, CombinationIndex, EvaluateCombination);
-	}
-} 
-
-template <typename EvaluateCombinationType>
-static void GeneratePermutations(int32 DataCardinality, int32 CombinationCardinality, EvaluateCombinationType EvaluateCombination)
-{
-	TArray<int32, TInlineAllocator<PreallocatedRolesNum>> Combination;
-	Combination.SetNum(CombinationCardinality);
-	GeneratePermutationsRecursive(DataCardinality, 0, Combination, 0, EvaluateCombination);
-}
-
 } // namespace UE::PoseSearch
 
 UE::PoseSearch::FInteractionIsland& UPoseSearchInteractionSubsystem::CreateIsland()
@@ -166,7 +133,7 @@ void UPoseSearchInteractionSubsystem::DestroyAllIslands()
 	}
 }
 
-void UPoseSearchInteractionSubsystem::RegenerateAllIslands(TConstArrayView<UE::PoseSearch::FInteractionSearchContext> SearchContexts)
+void UPoseSearchInteractionSubsystem::RegenerateAllIslands(float DeltaSeconds)
 {
 	QUICK_SCOPE_CYCLE_COUNTER(STAT_UPoseSearchInteractionSubsystem_RegenerateAllIslands);
 
@@ -174,6 +141,21 @@ void UPoseSearchInteractionSubsystem::RegenerateAllIslands(TConstArrayView<UE::P
 
 	check(IsInGameThread());
 
+	// FScopeLock Lock(&AvailabilitiesMapMutex); is not necessary since UPoseSearchInteractionSubsystem gets ticked outside the parallel animation jobs
+
+	// generating all the possible interaction tuples of AnimContext(s) with roles and pose histories (defined in FInteractionSearchContext)
+	FInteractionSearchContexts SearchContexts;
+	GenerateSearchContexts(DeltaSeconds, SearchContexts);
+
+#if ENABLE_DRAW_DEBUG && ENABLE_ANIM_DEBUG
+	// drawing the current frame islands to be consistent with the search, before regenerating the islands with the newly published availabilities
+	DebugDrawIslands();
+#endif // ENABLE_DRAW_DEBUG && ENABLE_ANIM_DEBUG
+	
+#if ENABLE_ANIM_DEBUG
+	DebugLogTickDependencies();
+#endif // ENABLE_ANIM_DEBUG
+
 #if !NO_CVARS
 	if (!GVarPoseSearchInteractionCacheIslands)
 	{
@@ -553,7 +535,7 @@ void UPoseSearchInteractionSubsystem::GenerateAnimContextInfosAndTagToDatabases(
 		}
 		else if (SubsystemWorld != AnimContext->GetWorld())
 		{
-			UE_LOG(LogPoseSearch, Error, TEXT("UPoseSearchInteractionSubsystem::GenerateAnimContextInfosAndTagToDatabaseMap AnimContext %s is from World %s, and supposed to be in World @s!"), *AnimContext->GetName(), *AnimContext->GetWorld()->GetName(), *SubsystemWorld->GetName());
+			UE_LOG(LogPoseSearch, Error, TEXT("UPoseSearchInteractionSubsystem::GenerateAnimContextInfosAndTagToDatabaseMap AnimContext %s is from World %s, and supposed to be in World %s!"), *AnimContext->GetName(), *AnimContext->GetWorld()->GetName(), *SubsystemWorld->GetName());
 		}
 		else
 		{
@@ -616,20 +598,10 @@ void UPoseSearchInteractionSubsystem::GenerateAnimContextInfosAndTagToDatabases(
 	for (int32 AnimContextInfoIndex = 0; AnimContextInfoIndex < NumAnimContextInfos; ++AnimContextInfoIndex)
 	{
 		FAnimContextInfo& AnimContextInfo = AnimContextInfos[AnimContextInfoIndex];
-
 		AnimContextInfo.AnimContext = AnimContextsWithAvailabilities[AnimContextInfoIndex].Key;
 		AnimContextInfo.Availabilities = AnimContextsWithAvailabilities[AnimContextInfoIndex].Value;
-		check(AnimContextInfo.AnimContext && AnimContextInfo.Availabilities && !AnimContextInfo.Availabilities->IsEmpty());
-
 		AnimContextInfo.Location = GetContextLocation(AnimContextInfo.AnimContext);
-
-		AnimContextInfo.TickPriority = -MAX_int32;
-		for (const FInteractionAvailability& Availability : *AnimContextInfo.Availabilities)
-		{
-			AnimContextInfo.MaxDistanceSquared = FMath::Max(AnimContextInfo.MaxDistanceSquared, Availability.BroadPhaseRadius);
-			AnimContextInfo.TickPriority = FMath::Max(AnimContextInfo.TickPriority, Availability.TickPriority);
-		}
-		AnimContextInfo.MaxDistanceSquared *= AnimContextInfo.MaxDistanceSquared;
+		check(AnimContextInfo.AnimContext && AnimContextInfo.Availabilities && !AnimContextInfo.Availabilities->IsEmpty());
 	}
 
 	// solving the broad phase using the AnimContextInfos
@@ -637,22 +609,19 @@ void UPoseSearchInteractionSubsystem::GenerateAnimContextInfosAndTagToDatabases(
 	{
 		for (int32 AnimContextInfoIndexB = AnimContextInfoIndexA + 1; AnimContextInfoIndexB < NumAnimContextInfos; ++AnimContextInfoIndexB)
 		{
-			const FVector DeltaLocation = AnimContextInfos[AnimContextInfoIndexA].Location - AnimContextInfos[AnimContextInfoIndexB].Location;
-			const float DistanceSquared = DeltaLocation.SquaredLength();
-			const float MaxDistanceSquared = FMath::Max(AnimContextInfos[AnimContextInfoIndexA].MaxDistanceSquared, AnimContextInfos[AnimContextInfoIndexB].MaxDistanceSquared);
-			if (DistanceSquared <= MaxDistanceSquared)
+			if (AnimContextInfos[AnimContextInfoIndexA].CanInteractWith(AnimContextInfos[AnimContextInfoIndexB]))
 			{
 				// the AnimContext of AnimContextInfos[AnimContextInfoIndexA] can potentially interact with the one from AnimContextInfos[AnimContextInfoIndexB]:
 				// linking AnimContextInfos[AnimContextInfoIndexA] to AnimContextInfos[AnimContextInfoIndexA] and vice versa to keep track of this when evaluating the broad phase.
 				// Since AnimContextInfos does't reallocate anymore, it's safe to store pointers to internal elements of the array!
-				AnimContextInfos[AnimContextInfoIndexA].NearbyAnimContextInfos.Emplace(AnimContextInfos[AnimContextInfoIndexB], DistanceSquared);
-				AnimContextInfos[AnimContextInfoIndexB].NearbyAnimContextInfos.Emplace(AnimContextInfos[AnimContextInfoIndexA], DistanceSquared);
+				AnimContextInfos[AnimContextInfoIndexA].NearbyAnimContextInfos.Emplace(&AnimContextInfos[AnimContextInfoIndexB]);
+				AnimContextInfos[AnimContextInfoIndexB].NearbyAnimContextInfos.Emplace(&AnimContextInfos[AnimContextInfoIndexA]);
 			}
 		}
 	}
 }
 
-void UPoseSearchInteractionSubsystem::GenerateSearchContexts(UE::PoseSearch::FInteractionSearchContexts& SearchContexts) const
+void UPoseSearchInteractionSubsystem::GenerateSearchContexts(float DeltaSeconds, UE::PoseSearch::FInteractionSearchContexts& SearchContexts) const
 {
 	QUICK_SCOPE_CYCLE_COUNTER(STAT_UPoseSearchInteractionSubsystem_GenerateSearchContexts);
 
@@ -662,47 +631,71 @@ void UPoseSearchInteractionSubsystem::GenerateSearchContexts(UE::PoseSearch::FIn
 
 	struct FRoledAnimContextInfo
 	{
-		FRoledAnimContextInfo(const FAnimContextInfo& InAnimContextInfo, const FRole InRole, const IPoseHistory* InPoseHistory, const UPoseSearchDatabase& InDatabase, float InMaxDistanceSquared, bool bInDisableCollisions)
-			: AnimContextInfo(&InAnimContextInfo)
+		FRoledAnimContextInfo(const FInteractionAvailability& InAvailability, const FAnimContextInfo& InAnimContextInfo, const FRole InRole, const IPoseHistory* InPoseHistory, const UPoseSearchDatabase& InDatabase)
+			: Availability(&InAvailability)
+			, AnimContextInfo(&InAnimContextInfo)
 			, Role(InRole)
 			, PoseHistory(InPoseHistory)
 			, Database(&InDatabase)
-			, MaxDistanceSquared(InMaxDistanceSquared)
-			, bDisableCollisions(bInDisableCollisions)
 		{
-			check(MaxDistanceSquared > UE_SMALL_NUMBER);
+		}
+
+		// perform narrow phase analysis checking if the AnimContextInfo with the specialized properties from Availability, Role, PoseHistory, Database,
+		// can interact with OtherRoledAnimContextInfo. This is a less relaxed analysis than the one performed in FAnimContextInfo::CanInteractWith
+		bool CanInteractWith(const FRoledAnimContextInfo& OtherRoledAnimContextInfo, bool bWasSearchContextInteracting) const
+		{
+			check(this != &OtherRoledAnimContextInfo);
+
+			const FVector DeltaLocation = AnimContextInfo->Location - OtherRoledAnimContextInfo.AnimContextInfo->Location;
+			const float DistanceSquared = DeltaLocation.SquaredLength();
+
+			float MaxDistance;
+			if (bWasSearchContextInteracting)
+			{
+				MaxDistance = FMath::Min(Availability->BroadPhaseRadius + Availability->BroadPhaseRadiusIncrementOnInteraction, OtherRoledAnimContextInfo.Availability->BroadPhaseRadius + OtherRoledAnimContextInfo.Availability->BroadPhaseRadiusIncrementOnInteraction);
+			}
+			else
+			{
+				MaxDistance = FMath::Min(Availability->BroadPhaseRadius, OtherRoledAnimContextInfo.Availability->BroadPhaseRadius);
+			}
+			
+			const float MaxDistanceSquared = MaxDistance * MaxDistance;
+			return DistanceSquared <= MaxDistanceSquared;
 		}
 
 		bool operator==(const FRoledAnimContextInfo& Other) const = default;
 
+		// Availability that spawned this FRoledAnimContextInfo
+		const FInteractionAvailability* Availability = nullptr;
+		// AnimContextInfo containing all the information regarding the AnimContext that spawned this FRoledAnimContextInfo,
+		// including all the availabilities associated to the AnimContext as well as all the other AnimContext(s) it can potentially interact with
 		const FAnimContextInfo* AnimContextInfo = nullptr;
-		FRole Role = DefaultRole;
+		const FRole Role = DefaultRole;
 		const IPoseHistory* PoseHistory = nullptr;
 		const UPoseSearchDatabase* Database = nullptr;
-		float MaxDistanceSquared = 0.f;
-		bool bDisableCollisions = false;
 	};
+	
 	struct FRoledAnimContextInfos : TArray<FRoledAnimContextInfo, TInlineAllocator<PreallocatedRolesNum, TMemStackAllocator<>>>
 	{
-		void AddRoledAnimContextInfos(const FAnimContextInfo& AnimContextInfo, const TArray<FName>& RolesFilter, const IPoseHistory* PoseHistory, const UPoseSearchDatabase& Database, float MaxDistanceSquared, bool bDisableCollisions)
+		void AddRoledAnimContextInfos(const FInteractionAvailability& Availability, const FAnimContextInfo& AnimContextInfo, const IPoseHistory* PoseHistory, const UPoseSearchDatabase& Database)
 		{
 			const UPoseSearchSchema* Schema = Database.Schema;
 			check(Schema);
-			if (RolesFilter.IsEmpty())
+			if (Availability.RolesFilter.IsEmpty())
 			{
 				// adding ALL the possible roles from the database:
 				for (const FPoseSearchRoledSkeleton& RoledSkeleton : Schema->GetRoledSkeletons())
 				{
-					AddUnique(FRoledAnimContextInfo(AnimContextInfo, RoledSkeleton.Role, PoseHistory, Database, MaxDistanceSquared, bDisableCollisions));
+					AddUnique(FRoledAnimContextInfo(Availability, AnimContextInfo, RoledSkeleton.Role, PoseHistory, Database));
 				}
 			}
 			else
 			{
-				for (const FRole& Role : RolesFilter)
+				for (const FRole& Role : Availability.RolesFilter)
 				{
 					if (Schema->GetRoledSkeleton(Role))
 					{
-						AddUnique(FRoledAnimContextInfo(AnimContextInfo, Role, PoseHistory, Database, MaxDistanceSquared, bDisableCollisions));
+						AddUnique(FRoledAnimContextInfo(Availability, AnimContextInfo, Role, PoseHistory, Database));
 					}
 					else
 					{
@@ -756,10 +749,10 @@ void UPoseSearchInteractionSubsystem::GenerateSearchContexts(UE::PoseSearch::FIn
 			{
 				OnNewAnimAnimContextInfoFound(AnimContextInfoToVisit);
 
-				for (const FNearbyAnimContextInfo& NearbyAnimContextInfo : AnimContextInfoToVisit.NearbyAnimContextInfos)
+				for (const FAnimContextInfo* NearbyAnimContextInfo : AnimContextInfoToVisit.NearbyAnimContextInfos)
 				{
-					check(NearbyAnimContextInfo.AnimContextInfo);
-					VisitRecursively(*NearbyAnimContextInfo.AnimContextInfo, OnNewAnimAnimContextInfoFound);
+					check(NearbyAnimContextInfo);
+					VisitRecursively(*NearbyAnimContextInfo, OnNewAnimAnimContextInfoFound);
 				}
 			}
 		}
@@ -774,6 +767,8 @@ void UPoseSearchInteractionSubsystem::GenerateSearchContexts(UE::PoseSearch::FIn
 	FTagToDatabases TagToDatabases;
 	GenerateAnimContextInfosAndTagToDatabases(AnimContextInfos, TagToDatabases);
 
+	const TConstArrayView<FValidInteractionSearch> PreviousValidInteractionSearches = ValidInteractionSearches;
+
 	TMap<const UPoseSearchDatabase*, FRoledAnimContextInfos, TInlineSetAllocator<PreallocatedSearchesNum, TMemStackSetAllocator<>>> DatabaseToRoledAnimContextInfos;
 	// visiting ALL the AnimContexts in AnimContextInfos and relying on the cached information to refine potential interactions
 	FAnimContextInfoVisitor AnimContextInfoVisitor(AnimContextInfos,
@@ -785,18 +780,13 @@ void UPoseSearchInteractionSubsystem::GenerateSearchContexts(UE::PoseSearch::FIn
 			check(AnimContextInfo.Availabilities);
 			for (const FInteractionAvailability& Availability : *AnimContextInfo.Availabilities)
 			{
-				if (Availability.BroadPhaseRadius <= UE_SMALL_NUMBER)
-				{
-					UE_LOG(LogPoseSearch, Log, TEXT("UPoseSearchInteractionSubsystem::GenerateSearchContexts too small Availability.BroadPhaseRadius. Character will not interact with anything within such small radius!"));
-				}
-				else if (const IPoseHistory* PoseHistory = Availability.GetPoseHistory(AnimContextInfo.AnimContext))
+				if (const IPoseHistory* PoseHistory = Availability.GetPoseHistory(AnimContextInfo.AnimContext))
 				{
 					if (const UPoseSearchDatabase* Database = Availability.Database.Get())
 					{
 						check(Database->Schema);
 						FRoledAnimContextInfos& RoledAnimContextInfos = DatabaseToRoledAnimContextInfos.FindOrAdd(Database);
-						const float MaxDistanceSquared = Availability.BroadPhaseRadius * Availability.BroadPhaseRadius;
-						RoledAnimContextInfos.AddRoledAnimContextInfos(AnimContextInfo, Availability.RolesFilter, PoseHistory, *Database, MaxDistanceSquared, Availability.bDisableCollisions);
+						RoledAnimContextInfos.AddRoledAnimContextInfos(Availability, AnimContextInfo, PoseHistory, *Database);
 					}
 					else if (Availability.IsTagValid())
 					{
@@ -804,12 +794,11 @@ void UPoseSearchInteractionSubsystem::GenerateSearchContexts(UE::PoseSearch::FIn
 						if (const FDatabasesPerTag* DatabasesPerTag = TagToDatabases.Find(Availability.Tag))
 						{
 							check(!DatabasesPerTag->IsEmpty());
-							const float MaxDistanceSquared = Availability.BroadPhaseRadius * Availability.BroadPhaseRadius;
 							for (const UPoseSearchDatabase* DatabaseFromTag : *DatabasesPerTag)
 							{
 								check(DatabaseFromTag && DatabaseFromTag->Schema);
 								FRoledAnimContextInfos& RoledAnimContextInfos = DatabaseToRoledAnimContextInfos.FindOrAdd(DatabaseFromTag);
-								RoledAnimContextInfos.AddRoledAnimContextInfos(AnimContextInfo, Availability.RolesFilter, PoseHistory, *DatabaseFromTag, MaxDistanceSquared, Availability.bDisableCollisions);
+								RoledAnimContextInfos.AddRoledAnimContextInfos(Availability, AnimContextInfo, PoseHistory, *DatabaseFromTag);
 							}
 						}
 						else
@@ -832,14 +821,14 @@ void UPoseSearchInteractionSubsystem::GenerateSearchContexts(UE::PoseSearch::FIn
 			}
 		},
 		// OnDoneGroupingAnimContexts: called when the FAnimContextInfoVisitor reaches the end of a group of possibly interacting AnimContext(s)
-		[&DatabaseToRoledAnimContextInfos, &SearchContexts]()
+		[&DatabaseToRoledAnimContextInfos, &SearchContexts, &PreviousValidInteractionSearches]()
 		{
-			// for each database now we try to create all the possible permutations of the roled anim instances
+			// for each database now we try to create all the possible combinations of the roled anim instances
 			// for example, given a database set up with assets for 2 characters interactions with roles RoleA and RoleB
 			// and 2 anim instances, all of them willing to partecipate in the 2 characters interaction with both roles RoleA and RoleB:
 			// CharA could be taking RoleA and RoleB,
 			// CharB could be taking RoleA and RoleB,
-			// we generate all the permutations from the array of options:
+			// we generate all the combinations from the array of options:
 			// CharA/RoleA, CharA/RoleB, CharB/RoleA, CharB/RoleB
 			//
 			// and we prune the invalid tuples:
@@ -861,9 +850,9 @@ void UPoseSearchInteractionSubsystem::GenerateSearchContexts(UE::PoseSearch::FIn
 				const int32 CombinationCardinality = RoledSkeletons.Num();
 				const FRoledAnimContextInfos& RoledAnimContextInfos = DatabaseToRoledAnimContextInfosPair.Value;
 
-				GeneratePermutations(RoledAnimContextInfos.Num(), CombinationCardinality,
+				GenerateCombinations(RoledAnimContextInfos.Num(), CombinationCardinality,
 					// Combination is an array of indexes in RoledAnimContextInfos: 0 <= Combination[i] < RoledAnimContextInfos.Num()
-					[Database, &RoledSkeletons, &RoledAnimContextInfos, &SearchContexts](const TConstArrayView<int32> Combination)
+					[Database, &RoledSkeletons, &RoledAnimContextInfos, &SearchContexts, &PreviousValidInteractionSearches](const TConstArrayView<int32> Combination)
 					{
 						// CombinationCardinality represents the number of roles as well as the number interacting AnimContext(s) (ultimately number of Characters involved in the interaction)
 						const int32 CombinationCardinality = Combination.Num();
@@ -878,43 +867,31 @@ void UPoseSearchInteractionSubsystem::GenerateSearchContexts(UE::PoseSearch::FIn
 							UniqueAnimContexts.Add(RoledAnimContextInfo.AnimContextInfo->AnimContext, &bIsAlreadyInSet);
 							if (bIsAlreadyInSet)
 							{
-								// we have a duplicate AnimContext. this permutation is NOT valid
+								// we have a duplicate AnimContext. this combination is NOT valid
 								return false;
 							}
 						}
 
 						FInteractionSearchContext SearchContext;
 						SearchContext.Database = Database;
-						SearchContext.AnimContexts.Reserve(CombinationCardinality);
-						SearchContext.PoseHistories.Reserve(CombinationCardinality);
-						SearchContext.Roles.Reserve(CombinationCardinality);
 
-						// checking if this combination is valid for the Database:
+						// setting up a FRoledAnimContextInfo in RoledAnimContextInfos describing 
+						// this potential interaction properties about how to perform the search
 						for (int32 CombinationIndex = 0; CombinationIndex < CombinationCardinality; ++CombinationIndex)
 						{
 							const int32 RoledAnimContextIndex = Combination[CombinationIndex];
 							const FRoledAnimContextInfo& RoledAnimContextInfo = RoledAnimContextInfos[RoledAnimContextIndex];
 
-							// checking the narrow phase!
-							for (int32 OtherCombinationIndex = CombinationIndex + 1; OtherCombinationIndex < CombinationCardinality; ++OtherCombinationIndex)
-							{
-								const int32 OtherRoledAnimContextIndex = Combination[OtherCombinationIndex];
-								const FRoledAnimContextInfo& OtherRoledAnimContextInfo = RoledAnimContextInfos[OtherRoledAnimContextIndex];
-
-								const float DistanceSquared = RoledAnimContextInfo.AnimContextInfo->GetCachedDistanceSquaredFrom(*OtherRoledAnimContextInfo.AnimContextInfo);
-								const float MinDistanceSquared = FMath::Min(RoledAnimContextInfo.MaxDistanceSquared, OtherRoledAnimContextInfo.MaxDistanceSquared);
-								if (DistanceSquared > MinDistanceSquared)
-								{
-									return false;
-								}
-							}
-
-							check(RoledAnimContextInfo.PoseHistory);
+							check(RoledAnimContextInfo.PoseHistory && RoledAnimContextInfo.Availability);
 							SearchContext.AnimContexts.Add(RoledAnimContextInfo.AnimContextInfo->AnimContext);
 							SearchContext.Roles.Add(RoledAnimContextInfo.Role);
 							SearchContext.PoseHistories.Add(RoledAnimContextInfo.PoseHistory->AsWeak());
-							SearchContext.bDisableCollisions |= RoledAnimContextInfo.bDisableCollisions;
-							SearchContext.TickPriorities.Add(RoledAnimContextInfo.AnimContextInfo->TickPriority);
+							SearchContext.bDisableCollisions |= RoledAnimContextInfo.Availability->bDisableCollisions;
+							SearchContext.TickPriorities.Add(RoledAnimContextInfo.Availability->TickPriority);
+							
+							#if ENABLE_ANIM_DEBUG
+							SearchContext.DebugAvailabilities.Add(*RoledAnimContextInfo.Availability);
+							#endif // ENABLE_ANIM_DEBUG
 						}
 
 						// does SearchContext cover all the roles required by this interaction?
@@ -926,6 +903,36 @@ void UPoseSearchInteractionSubsystem::GenerateSearchContexts(UE::PoseSearch::FIn
 								return false;
 							}
 						}
+
+						// looking for a preexisting valid interaction resmbling SearchContext
+						for (const FValidInteractionSearch& PreviousValidInteractionSearch : PreviousValidInteractionSearches)
+						{
+							if (PreviousValidInteractionSearch.SearchContext.IsEquivalent(SearchContext))
+							{
+								SearchContext.bIsContinuingInteraction = true;
+								break;
+							}
+						}
+
+						// checking if this combination is valid for the Database:
+						for (int32 CombinationIndex = 0; CombinationIndex < CombinationCardinality; ++CombinationIndex)
+						{
+							const int32 RoledAnimContextIndex = Combination[CombinationIndex];
+							const FRoledAnimContextInfo& RoledAnimContextInfo = RoledAnimContextInfos[RoledAnimContextIndex];
+
+							// checking the narrow phase!
+							for (int32 OtherCombinationIndex = CombinationIndex + 1; OtherCombinationIndex < CombinationCardinality; ++OtherCombinationIndex)
+							{
+								const int32 OtherRoledAnimContextIndex = Combination[OtherCombinationIndex];
+								const FRoledAnimContextInfo& OtherRoledAnimContextInfo = RoledAnimContextInfos[OtherRoledAnimContextIndex];
+
+								// if any of the RoledAnimContextInfo cannot interact with any OtherRoledAnimContextInfo the inteaction cannot happen!
+								if (!RoledAnimContextInfo.CanInteractWith(OtherRoledAnimContextInfo, SearchContext.bIsContinuingInteraction))
+								{
+									return false;
+								}
+							}
+						}
 #if DO_CHECK
 						for (const FInteractionSearchContext& ContainedSearchContext : SearchContexts)
 						{
@@ -940,6 +947,9 @@ void UPoseSearchInteractionSubsystem::GenerateSearchContexts(UE::PoseSearch::FIn
 			// done using DatabaseToRoledAnimContextInfos. clearing up for the next group of AnimContext(s)
 			DatabaseToRoledAnimContextInfos.Reset();
 		});
+
+	// populating the continuing pose properties for the SearchContexts from the current Islands
+	PopulateContinuingProperties(DeltaSeconds, SearchContexts);
 }
 
 void UPoseSearchInteractionSubsystem::OnInteractionStart(UE::PoseSearch::FValidInteractionSearch& ValidInteractionSearch)
@@ -1081,32 +1091,6 @@ void UPoseSearchInteractionSubsystem::UpdateValidInteractionSearches()
 	ValidInteractionSearches = NewValidInteractionSearches;
 }
 
-// @todo: add a debug critical section to avoid overlapping of GetInteractingAnimContextsWith with UpdateValidInteractionSearches to avoid / detecting threading issues!
-TArray<const UObject*> UPoseSearchInteractionSubsystem::GetInteractingAnimContextsWith(const UObject* AnimContext) const
-{
-	using namespace UE::PoseSearch;
-
-	TArray<const UObject*> InteractingAnimContexts;
-	for (const FValidInteractionSearch& ValidInteractionSearch : ValidInteractionSearches)
-	{
-		const int32 AnimContextIndex = ValidInteractionSearch.SearchContext.AnimContexts.IndexOfByKey(AnimContext);
-		if (AnimContextIndex != INDEX_NONE)
-		{
-			for (int32 Index = 0; Index < ValidInteractionSearch.SearchContext.AnimContexts.Num(); ++Index)
-			{
-				if (Index != AnimContextIndex)
-				{
-					if (const UObject* InteractingAnimContext = ValidInteractionSearch.SearchContext.AnimContexts[Index].Get())
-					{
-						InteractingAnimContexts.Add(InteractingAnimContext);
-					}
-				}
-			}
-		}
-	}
-	return InteractingAnimContexts;
-}
-
 void UPoseSearchInteractionSubsystem::Tick(float DeltaSeconds)
 {
 	QUICK_SCOPE_CYCLE_COUNTER(STAT_UPoseSearchInteractionSubsystem_Tick);
@@ -1138,25 +1122,7 @@ void UPoseSearchInteractionSubsystem::Tick(float DeltaSeconds)
 
 	check(IsInGameThread());
 
-	// FScopeLock Lock(&AvailabilitiesMapMutex); is not necessary since UPoseSearchInteractionSubsystem gets ticked outside the parallel animation jobs
-
-	// generating all the possible interaction tuples of AnimContext(s) with roles and pose histories (defined in FInteractionSearchContext)
-	FInteractionSearchContexts SearchContexts;
-	GenerateSearchContexts(SearchContexts);
-
-	// populating the continuing pose properties for the SearchContexts from the current Islands
-	PopulateContinuingProperties(DeltaSeconds, SearchContexts);
-
-#if ENABLE_DRAW_DEBUG && ENABLE_ANIM_DEBUG
-	// drawing the current frame islands to be consistent with the search, before regenerating the islands with the newly published availabilities
-	DebugDrawIslands();
-#endif // ENABLE_DRAW_DEBUG && ENABLE_ANIM_DEBUG
-	
-#if ENABLE_ANIM_DEBUG
-	DebugLogTickDependencies();
-#endif // ENABLE_ANIM_DEBUG
-
-	RegenerateAllIslands(SearchContexts);
+	RegenerateAllIslands(DeltaSeconds);
 
 	// clearing up AvailabilitiesMap for the next frame
 	AvailabilitiesMap.Reset();
@@ -1231,14 +1197,16 @@ void UPoseSearchInteractionSubsystem::GetResult_AnyThread(const UObject* AnimCon
 	}
 }
 
-#if ENABLE_DRAW_DEBUG && ENABLE_ANIM_DEBUG
+#if ENABLE_ANIM_DEBUG
 void UPoseSearchInteractionSubsystem::DebugDrawIslands() const
 {
+#if ENABLE_VISUAL_LOG
+
 	using namespace UE::PoseSearch;
 
 	check(IsInGameThread());
 
-	if (!GVarPoseSearchInteractionShowIslands)
+	if (!FVisualLogger::IsRecording())
 	{
 		return;
 	}
@@ -1262,6 +1230,18 @@ void UPoseSearchInteractionSubsystem::DebugDrawIslands() const
 	static const int32 NumColors = sizeof(Colors) / sizeof(Colors[0]);
 	int32 CurrentColorIndex = 0;
 
+	TArray<const UObject*, TInlineAllocator<256>> AllAnimContexts;
+	for (const FInteractionIsland* Island : Islands)
+	{
+		for (const TWeakObjectPtr<const UObject>& IslandAnimContextPtr : Island->GetIslandAnimContexts())
+		{
+			if (const UObject* IslandAnimContext = IslandAnimContextPtr.Get())
+			{
+				AllAnimContexts.Add(IslandAnimContext);
+			}
+		}
+	}
+
 	for (const FInteractionIsland* Island : Islands)
 	{
 		if (Island->IsInitialized())
@@ -1274,30 +1254,39 @@ void UPoseSearchInteractionSubsystem::DebugDrawIslands() const
 				{
 					if (const UObject* AnimContext = SearchContext.AnimContexts[Index].Get())
 					{
-						if (const FInteractionAvailabilities* Availabilities = AvailabilitiesMap.Find(AnimContext))
+						const FPoseSearchInteractionAvailability& DebugAvailability = SearchContext.DebugAvailabilities[Index];
+						float MaxBroadPhaseRadius;
+						if (SearchContext.bIsContinuingInteraction)
+						{
+							MaxBroadPhaseRadius = DebugAvailability.BroadPhaseRadius + DebugAvailability.BroadPhaseRadiusIncrementOnInteraction;
+						}
+						else
+						{
+							MaxBroadPhaseRadius = DebugAvailability.BroadPhaseRadius;
+						}
+
+						if (MaxBroadPhaseRadius > UE_SMALL_NUMBER)
 						{
-							float MaxBroadPhaseRadius = 0.f;
-							for (const FInteractionAvailability& Availability : *Availabilities)
+							const FTransform& Transform = GetContextTransform(AnimContext);
+							static const TCHAR* LogName = TEXT("PoseSearchInteraction");
+								
+							for (const UObject* IslandAnimContext : AllAnimContexts)
 							{
-								MaxBroadPhaseRadius = FMath::Max(MaxBroadPhaseRadius, Availability.BroadPhaseRadius);
+								UE_VLOG_CIRCLE(IslandAnimContext, LogName, Display, Transform.GetLocation(), FVector::UpVector, MaxBroadPhaseRadius, Color, TEXT(""));
 							}
 
-							if (MaxBroadPhaseRadius > UE_SMALL_NUMBER)
+							if (!Island->HasTickDependencies())
 							{
-								const UWorld* World = GetWorld();
-								const FTransform& Transform = GetContextTransform(AnimContext);
-								DrawDebugCircle(World, Transform.GetLocation(), MaxBroadPhaseRadius, 40, Color, false, 0.f, SDPG_Foreground, 0.f, FVector::XAxisVector, FVector::YAxisVector, false);
+								const FVector ForwardAxisStart = Transform.TransformPosition(FVector::ForwardVector * MaxBroadPhaseRadius);
+								const FVector ForwardAxisEnd = Transform.TransformPosition(FVector::ForwardVector * -MaxBroadPhaseRadius);
 
-								if (!Island->HasTickDependencies())
-								{
-									const FVector ForwardAxisStart = Transform.TransformPosition(FVector::ForwardVector * MaxBroadPhaseRadius);
-									const FVector ForwardAxisEnd = Transform.TransformPosition(FVector::ForwardVector * -MaxBroadPhaseRadius);
-
-									const FVector LeftAxisStart = Transform.TransformPosition(FVector::LeftVector * MaxBroadPhaseRadius);
-									const FVector LeftAxisEnd = Transform.TransformPosition(FVector::LeftVector * -MaxBroadPhaseRadius);
+								const FVector LeftAxisStart = Transform.TransformPosition(FVector::LeftVector * MaxBroadPhaseRadius);
+								const FVector LeftAxisEnd = Transform.TransformPosition(FVector::LeftVector * -MaxBroadPhaseRadius);
 
-									DrawDebugLine(World, ForwardAxisStart, ForwardAxisEnd, Color, false, 0.f, SDPG_Foreground, 0.f);
-									DrawDebugLine(World, LeftAxisStart, LeftAxisEnd, Color, false, 0.f, SDPG_Foreground, 0.f);
+								for (const UObject* IslandAnimContext : AllAnimContexts)
+								{
+									UE_VLOG_SEGMENT(IslandAnimContext, LogName, Display, ForwardAxisStart, ForwardAxisEnd, Color, TEXT(""));
+									UE_VLOG_SEGMENT(IslandAnimContext, LogName, Display, LeftAxisStart, LeftAxisEnd, Color, TEXT(""));
 								}
 							}
 						}
@@ -1308,10 +1297,9 @@ void UPoseSearchInteractionSubsystem::DebugDrawIslands() const
 			CurrentColorIndex = (CurrentColorIndex + 1) % NumColors;
 		}
 	}
+#endif // ENABLE_VISUAL_LOG
 }
-#endif // ENABLE_DRAW_DEBUG && ENABLE_ANIM_DEBUG
 
-#if ENABLE_ANIM_DEBUG
 void UPoseSearchInteractionSubsystem::DebugLogTickDependencies() const
 {
 #if !NO_CVARS
diff --git a/Engine/Plugins/Animation/PoseSearch/Source/Runtime/Public/PoseSearch/PoseSearchInteractionAvailability.h b/Engine/Plugins/Animation/PoseSearch/Source/Runtime/Public/PoseSearch/PoseSearchInteractionAvailability.h
new file mode 100644
--- /dev/null
+++ b/Engine/Plugins/Animation/PoseSearch/Source/Runtime/Public/PoseSearch/PoseSearchInteractionAvailability.h
@@ -0,0 +1,89 @@
+// Copyright Epic Games, Inc. All Rights Reserved.
+
+#pragma once
+
+#include "PoseSearch/PoseSearchHistory.h"
+#include "PoseSearchInteractionAvailability.generated.h"
+
+class UPoseSearchDatabase;
+
+// input for MotionMatchInteraction_Pure: it declares that the associated character ("AnimContext" that could be an AnimInstance or an AnimNextCharacterComponent)
+// is willing to partecipate in an interaction described by a UMultiAnimAsset (derived by UPoseSearchInteractionAsset) contained in the UPoseSearchDatabase Database
+// with one of the roles in RolesFilter (if empty ANY of the Database roles can be taken) the MotionMatchInteraction_Pure will ultimately setup a motion matching query
+// using looking for the pose history "PoseHistoryName" to gather bone and trajectory positions for this character for an interaction to be valid,
+// the query needs to find all the other interacting characters within BroadPhaseRadius, and reach a maximum cost of MaxCost
+// Experimental, this feature might be removed without warning, not for production use
+USTRUCT(Experimental, BlueprintType, Category = "Animation|Pose Search")
+struct POSESEARCH_API FPoseSearchInteractionAvailability
+{
+	GENERATED_BODY()
+
+	// Database describing the interaction. It'll contains multi character UMultiAnimAsset and a schema with multiple skeletons with associated roles
+	UPROPERTY(EditAnywhere, BlueprintReadWrite, Category=Settings)
+	TObjectPtr<UPoseSearchDatabase> Database;
+
+	// in case this availability Database is valid (not null), Tag (if IsTagValid()) is used to flag the Database with a specific name. Different availabilities can share the same Tag.
+	// in case this availability Database is NOT valid, we use the valid Tag to figure out all the possible databases that can be assigned to this availability from all the published availabilities.
+	// The reason behind Tag is, for example, to be able to have NPCs been able to interact with a main character (MC), without the MC having a direct dependency to the database used 
+	// for the interaction allowing those NPCs to be contextually loaded/unloaded, streamed in/out, with the obvious advantages for the the memory managment of the "payload" database.
+	// Another reason for Tag, is to facilitate the setup of interactions, where the MC have to publish only one availability with its own assigned Role (in RolesFilter) 
+	// automatically contextually resolved in multiple different types of possible databases: it could be MC-NPC, MC-Vehicle, MC-Whatever
+	UPROPERTY(EditAnywhere, BlueprintReadWrite, Category=Settings)
+	FName Tag;
+
+	// roles the character is willing to take to partecipate in this interaction. If empty ANY of the Database roles can be taken
+	UPROPERTY(EditAnywhere, BlueprintReadWrite, Category=Settings)
+	TArray<FName> RolesFilter;
+
+	// the associated character to this FPoseSearchInteractionAvailability will partecipate in an interaction only if all the necessary roles gest assigned to character within BroadPhaseRadius centimeters
+	UPROPERTY(EditAnywhere, BlueprintReadWrite, Category=Settings)
+	float BroadPhaseRadius = 500.f;
+
+	// during interaction the BroadPhaseRadius will be incremented by BroadPhaseRadiusIncrementOnInteraction to create 
+	// geometrical histeresys, where it's harder for actors to get into interaction rather than staying in interaction
+	UPROPERTY(EditAnywhere, BlueprintReadWrite, Category=Settings)
+	float BroadPhaseRadiusIncrementOnInteraction = 10.f;
+	
+	// if true, the system will disable collsions between interacting characters
+	UPROPERTY(EditAnywhere, BlueprintReadWrite, Category=Settings)
+	bool bDisableCollisions = false;
+
+	// the Actor with the higher TickPriority of any Availability request will be elected as the MainActor of the interaction island (containing all the actors that could interact with each other)
+	// the main Actor will tick first and all the other interacting actors will tick after in a concurrently from each other. TickPriority is useful if your setup is already enforcing tick dependency 
+	// between actors, and the motion matching interaction system needs to play nicely with them.
+	UPROPERTY(EditAnywhere, BlueprintReadWrite, Category=Settings)
+	int32 TickPriority = UE::PoseSearch::DefaultTickPriority;
+
+	bool IsTagValid() const { return !Tag.IsNone(); }
+
+	bool operator==(const FPoseSearchInteractionAvailability& Other) const = default;
+};
+
+namespace UE::PoseSearch
+{
+
+// Experimental, this feature might be removed without warning, not for production use
+struct FInteractionAvailability : public FPoseSearchInteractionAvailability
+{
+	FInteractionAvailability(const FPoseSearchInteractionAvailability& InAvailability, FName InPoseHistoryName, const IPoseHistory* InPoseHistory)
+		: FPoseSearchInteractionAvailability(InAvailability)
+		, PoseHistoryName(InPoseHistoryName)
+		, PoseHistory(InPoseHistory)
+	{
+		check(!PoseHistory || PoseHistory->AsWeak().IsValid());
+	}
+
+	FString GetPoseHistoryName() const;
+	const IPoseHistory* GetPoseHistory(const UObject* AnimContext) const;
+
+private:
+	FName PoseHistoryName;
+	const IPoseHistory* PoseHistory = nullptr;
+};
+
+// Experimental, this feature might be removed without warning, not for production use
+typedef TArray<FInteractionAvailability> FInteractionAvailabilities;
+// Experimental, this feature might be removed without warning, not for production use
+typedef TMap<FObjectKey, FInteractionAvailabilities> FAvailabilitiesMap;
+
+} // namespace UE::PoseSearch
\ No newline at end of file
diff --git a/Engine/Plugins/Animation/PoseSearch/Source/Runtime/Public/PoseSearch/PoseSearchInteractionIsland.h b/Engine/Plugins/Animation/PoseSearch/Source/Runtime/Public/PoseSearch/PoseSearchInteractionIsland.h
--- a/Engine/Plugins/Animation/PoseSearch/Source/Runtime/Public/PoseSearch/PoseSearchInteractionIsland.h
+++ b/Engine/Plugins/Animation/PoseSearch/Source/Runtime/Public/PoseSearch/PoseSearchInteractionIsland.h
@@ -2,6 +2,7 @@
 
 #pragma once
 
+#include "PoseSearch/PoseSearchInteractionAvailability.h"
 #include "PoseSearch/PoseSearchInteractionValidator.h"
 #include "PoseSearch/PoseSearchLibrary.h"
 
@@ -63,7 +64,11 @@ struct FInteractionSearchContext : FInteractionSearchContextBase
 	bool bIsContinuingInteraction = false;
 	TArray<int32, TInlineAllocator<PreallocatedRolesNum>> TickPriorities;
 
-	FPoseSearchContinuingProperties GetContinuingProperties(const TConstArrayView<FValidInteractionSearch> ValidInteractionSearches) const
+#if ENABLE_ANIM_DEBUG
+	TArray<FPoseSearchInteractionAvailability, TInlineAllocator<PreallocatedRolesNum>> DebugAvailabilities;
+#endif // ENABLE_ANIM_DEBUG
+
+	FPoseSearchContinuingProperties GetContinuingProperties() const
 	{
 		FPoseSearchContinuingProperties ContinuingProperties;
 		ContinuingProperties.PlayingAsset = PlayingAsset.Get();
@@ -71,16 +76,7 @@ struct FInteractionSearchContext : FInteractionSearchContextBase
 		ContinuingProperties.bIsPlayingAssetMirrored = bIsPlayingAssetMirrored;
 		ContinuingProperties.PlayingAssetBlendParameters = PlayingAssetBlendParameters;
 		ContinuingProperties.InterruptMode = InterruptMode;
-
-		for (const FValidInteractionSearch& ValidInteractionSearch : ValidInteractionSearches)
-		{
-			if (ValidInteractionSearch.SearchContext.IsEquivalent(*this))
-			{
-				ContinuingProperties.bIsContinuingInteraction = true;
-				break;
-			}
-		}
-
+		ContinuingProperties.bIsContinuingInteraction = bIsContinuingInteraction;
 		return ContinuingProperties;
 	}
 };
@@ -165,7 +161,7 @@ private:
 	// all the AnimContext in this island. Each SearchContexts will contain a subset of IslandAnimContexts 
 	TArray<TWeakObjectPtr<const UObject>> IslandAnimContexts;
 
-	// there's one FSearchContext for each search we need to perform (including all the possible roles permutations). Added by UPoseSearchInteractionSubsystem::Tick
+	// there's one FSearchContext for each search we need to perform (including all the possible roles combinations). Added by UPoseSearchInteractionSubsystem::Tick
 	// islands don't get deallocated, so this array once warmed up will not hit the allocator and waste cycles
 	TArray<FInteractionSearchContext> SearchContexts;
 
diff --git a/Engine/Plugins/Animation/PoseSearch/Source/Runtime/Public/PoseSearch/PoseSearchInteractionLibrary.h b/Engine/Plugins/Animation/PoseSearch/Source/Runtime/Public/PoseSearch/PoseSearchInteractionLibrary.h
--- a/Engine/Plugins/Animation/PoseSearch/Source/Runtime/Public/PoseSearch/PoseSearchInteractionLibrary.h
+++ b/Engine/Plugins/Animation/PoseSearch/Source/Runtime/Public/PoseSearch/PoseSearchInteractionLibrary.h
@@ -6,61 +6,6 @@
 #include "PoseSearch/PoseSearchInteractionIsland.h"
 #include "PoseSearchInteractionLibrary.generated.h"
 
-// input for MotionMatchInteraction_Pure: it declares that the associated character ("AnimContext" that could be an AnimInstance or an AnimNextCharacterComponent)
-// is willing to partecipate in an interaction described by a UMultiAnimAsset (derived by UPoseSearchInteractionAsset) contained in the UPoseSearchDatabase Database
-// with one of the roles in RolesFilter (if empty ANY of the Database roles can be taken) the MotionMatchInteraction_Pure will ultimately setup a motion matching query
-// using looking for the pose history "PoseHistoryName" to gather bone and trajectory positions for this character for an interaction to be valid,
-// the query needs to find all the other interacting characters within BroadPhaseRadius, and reach a maximum cost of MaxCost
-// Experimental, this feature might be removed without warning, not for production use
-USTRUCT(Experimental, BlueprintType, Category = "Animation|Pose Search")
-struct POSESEARCH_API FPoseSearchInteractionAvailability
-{
-	GENERATED_BODY()
-
-	// Database describing the interaction. It'll contains multi character UMultiAnimAsset and a schema with multiple skeletons with associated roles
-	UPROPERTY(EditAnywhere, BlueprintReadWrite, Category=Settings)
-	TObjectPtr<UPoseSearchDatabase> Database;
-
-	// in case this availability Database is valid (not null), Tag (if IsTagValid()) is used to flag the Database with a specific name. Different availabilities can share the same Tag.
-	// in case this availability Database is NOT valid, we use the valid Tag to figure out all the possible databases that can be assigned to this availability from all the published availabilities.
-	// The reason behind Tag is, for example, to be able to have NPCs been able to interact with a main character (MC), without the MC having a direct dependency to the database used 
-	// for the interaction allowing those NPCs to be contextually loaded/unloaded, streamed in/out, with the obvious advantages for the the memory managment of the "payload" database.
-	// Another reason for Tag, is to facilitate the setup of interactions, where the MC have to publish only one availability with its own assigned Role (in RolesFilter) 
-	// automatically contextually resolved in multiple different types of possible databases: it could be MC-NPC, MC-Vehicle, MC-Whatever
-	UPROPERTY(EditAnywhere, BlueprintReadWrite, Category=Settings)
-	FName Tag;
-
-	// roles the character is willing to take to partecipate in this interaction. If empty ANY of the Database roles can be taken
-	UPROPERTY(EditAnywhere, BlueprintReadWrite, Category=Settings)
-	TArray<FName> RolesFilter;
-
-	// the associated character to this FPoseSearchInteractionAvailability will partecipate in an interaction only if all the necessary roles gest assigned to character within BroadPhaseRadius centimeters
-	UPROPERTY(EditAnywhere, BlueprintReadWrite, Category=Settings)
-	float BroadPhaseRadius = 500.f;
-
-	// if true, the system will disable collsions between interacting characters
-	UPROPERTY(EditAnywhere, BlueprintReadWrite, Category=Settings)
-	bool bDisableCollisions = false;
-
-	// the Actor with the higher TickPriority of any Availability request will be elected as the MainActor of the interaction island (containing all the actors that could interact with each other)
-	// the main Actor will tick first and all the other interacting actors will tick after in a concurrently from each other. TickPriority is useful if your setup is already enforcing tick dependency 
-	// between actors, and the motion matching interaction system needs to play nicely with them.
-	UPROPERTY(EditAnywhere, BlueprintReadWrite, Category=Settings)
-	int32 TickPriority = UE::PoseSearch::DefaultTickPriority;
-
-	bool IsTagValid() const { return !Tag.IsNone(); }
-
-	bool operator==(const FPoseSearchInteractionAvailability& Other) const
-	{
-		return Database == Other.Database
-			&& Tag == Other.Tag
-			&& RolesFilter == Other.RolesFilter
-			&& BroadPhaseRadius == Other.BroadPhaseRadius
-			&& bDisableCollisions == Other.bDisableCollisions
-			&& TickPriority == Other.TickPriority;
-	}
-};
-
 // Experimental, this feature might be removed without warning, not for production use
 UCLASS(Experimental)
 class POSESEARCH_API UPoseSearchInteractionLibrary : public UBlueprintFunctionLibrary
diff --git a/Engine/Plugins/Animation/PoseSearch/Source/Runtime/Public/PoseSearch/PoseSearchInteractionSubsystem.h b/Engine/Plugins/Animation/PoseSearch/Source/Runtime/Public/PoseSearch/PoseSearchInteractionSubsystem.h
--- a/Engine/Plugins/Animation/PoseSearch/Source/Runtime/Public/PoseSearch/PoseSearchInteractionSubsystem.h
+++ b/Engine/Plugins/Animation/PoseSearch/Source/Runtime/Public/PoseSearch/PoseSearchInteractionSubsystem.h
@@ -2,10 +2,10 @@
 
 #pragma once
 
-#include "Subsystems/WorldSubsystem.h"
 #include "PoseSearch/PoseSearchContext.h"
 #include "PoseSearch/PoseSearchInteractionLibrary.h"
 #include "PoseSearch/PoseSearchInteractionValidator.h"
+#include "Subsystems/WorldSubsystem.h"
 #include "PoseSearchInteractionSubsystem.generated.h"
 
 struct FPoseSearchContinuingProperties;
@@ -14,34 +14,8 @@ class UPoseSearchDatabase;
 
 namespace UE::PoseSearch
 {
-
-// Experimental, this feature might be removed without warning, not for production use
-struct FInteractionAvailability : public FPoseSearchInteractionAvailability
-{
-	FInteractionAvailability(const FPoseSearchInteractionAvailability& InAvailability, FName InPoseHistoryName, const IPoseHistory* InPoseHistory)
-		: FPoseSearchInteractionAvailability(InAvailability)
-		, PoseHistoryName(InPoseHistoryName)
-		, PoseHistory(InPoseHistory)
-	{
-		check(!PoseHistory || PoseHistory->AsWeak().IsValid());
-	}
-
-	FString GetPoseHistoryName() const;
-	const IPoseHistory* GetPoseHistory(const UObject* AnimContext) const;
-
-private:
-	FName PoseHistoryName;
-	const IPoseHistory* PoseHistory = nullptr;
-};
-
-// Experimental, this feature might be removed without warning, not for production use
-typedef TArray<FInteractionAvailability> FInteractionAvailabilities;
-// Experimental, this feature might be removed without warning, not for production use
-typedef TMap<FObjectKey, FInteractionAvailabilities> FAvailabilitiesMap;
-
 struct FAnimContextInfos;
 struct FTagToDatabases;
-
 } // namespace UE::PoseSearch
 
 // World subsystem accepting the publication of characters (via their AnimInstance(s)) FPoseSearchInteractionAvailability, representing the characters willingness to partecipate in an 
@@ -81,7 +55,7 @@ struct FTagToDatabases;
 // - next frame UPoseSearchInteractionLibrary::MotionMatchInteraction_Pure(TArray<FPoseSearchInteractionAvailability> Availabilities, UObject* AnimInstance), with the context of all the published 
 //   availabilities and created islands, will find the associated FInteractionIsland to the AnimInstance and call FInteractionIsland::DoSearch_AnyThread 
 //   (via UPoseSearchInteractionSubsystem::Query_AnyThread) that will perform ALL (YES, ALL, so the bigger the island the slower the execution) the motion matching searches
-//   for all the possible Actors / databases / Roles permutations, and populate FInteractionIsland::SearchResults with ALL the results for the island.
+//   for all the possible Actors / databases / Roles combinations, and populate FInteractionIsland::SearchResults with ALL the results for the island.
 //   Ultimately the MotionMatchInteraction_Pure will return the SearchResults associated to the requesting AnimInstance with information about what animation to play
 //   at what time with wich Role.
 
@@ -112,9 +86,6 @@ public:
 
 	void GetResult_AnyThread(const UObject* AnimContext, FPoseSearchBlueprintResult& Result, bool bCompareOwningActors = false);
 
-	// returns and array containing all the other AnimContext(s) that the input AnimContext is currently interacting with
-	TArray<const UObject*> GetInteractingAnimContextsWith(const UObject* AnimContext) const;
-
 	TConstArrayView<UE::PoseSearch::FInteractionIsland*> GetInteractionIslands() const { return Islands; }
 
 private:
@@ -127,7 +98,7 @@ private:
 
 	void DestroyIsland(int32 Index);
 	void DestroyAllIslands();
-	void RegenerateAllIslands(TConstArrayView<UE::PoseSearch::FInteractionSearchContext> SearchContexts);
+	void RegenerateAllIslands(float DeltaSeconds);
 
 #if DO_CHECK
 	bool ValidateAllIslands() const;
@@ -136,17 +107,14 @@ private:
 	void PopulateContinuingProperties(float DeltaSeconds, TArrayView<UE::PoseSearch::FInteractionSearchContext> SearchContexts) const;
 	UE::PoseSearch::FInteractionIsland* FindIsland(const UObject* AnimContext, bool bCompareOwningActors = false);
 	
-#if ENABLE_DRAW_DEBUG && ENABLE_ANIM_DEBUG
-	void DebugDrawIslands() const;
-#endif // ENABLE_DRAW_DEBUG && ENABLE_ANIM_DEBUG
-	
 #if ENABLE_ANIM_DEBUG
+	void DebugDrawIslands() const;
 	void DebugLogTickDependencies() const;
 #endif // ENABLE_ANIM_DEBUG
 
 	void AddAvailabilities(const TArrayView<const FPoseSearchInteractionAvailability> Availabilities, const UObject* AnimContext, FName PoseHistoryName, const UE::PoseSearch::IPoseHistory* PoseHistory);
 	void GenerateAnimContextInfosAndTagToDatabases(UE::PoseSearch::FAnimContextInfos& AnimContextInfos, UE::PoseSearch::FTagToDatabases& TagToDatabases) const;
-	void GenerateSearchContexts(UE::PoseSearch::FInteractionSearchContexts& SearchContexts) const;
+	void GenerateSearchContexts(float DeltaSeconds, UE::PoseSearch::FInteractionSearchContexts& SearchContexts) const;
 
 	void UpdateValidInteractionSearches();
 	static void OnInteractionStart(UE::PoseSearch::FValidInteractionSearch& ValidInteractionSearch);
diff --git a/Engine/Plugins/Animation/PoseSearch/Source/Runtime/Public/PoseSearch/PoseSearchInteractionUtils.h b/Engine/Plugins/Animation/PoseSearch/Source/Runtime/Public/PoseSearch/PoseSearchInteractionUtils.h
--- a/Engine/Plugins/Animation/PoseSearch/Source/Runtime/Public/PoseSearch/PoseSearchInteractionUtils.h
+++ b/Engine/Plugins/Animation/PoseSearch/Source/Runtime/Public/PoseSearch/PoseSearchInteractionUtils.h
@@ -22,4 +22,38 @@ POSESEARCH_API void CalculateFullAlignedTransforms(const FPoseSearchBlueprintRes
 // Experimental, this feature might be removed without warning, not for production use
 POSESEARCH_API FTransform CalculateDeltaAlignment(const FTransform& MeshWithoutOffset, const FTransform& MeshWithOffset, const FTransform& FullAlignedTransform, float WarpingRotationRatio, float WarpingTranslationRatio);
 
+
+template <typename EvaluateCombinationType>
+static void GenerateCombinationsRecursive(int32 DataCardinality, int32 DataIndex, TArrayView<int32> Combination, int32 CombinationIndex, EvaluateCombinationType EvaluateCombination)
+{
+	if (CombinationIndex == Combination.Num())
+	{ 
+		EvaluateCombination(Combination);
+	}
+	else if (DataIndex < DataCardinality)
+	{
+		Combination[CombinationIndex] = DataIndex;
+		GenerateCombinationsRecursive(DataCardinality, DataIndex + 1, Combination, CombinationIndex + 1, EvaluateCombination);
+		GenerateCombinationsRecursive(DataCardinality, DataIndex + 1, Combination, CombinationIndex, EvaluateCombination);
+	}
+} 
+
+// generates all the possible (unique) combinations for the indexes of a set of DataCardinality number of elements grouped in tuples of CombinationCardinality size:
+// the number of generated combinations P(DataCardinality, CombinationCardinality) = (DataCardinality! / (DataCardinality - CombinationCardinality)!) / CombinationCardinality!
+// see: https://en.wikipedia.org/wiki/Combination
+// 
+// for example with a set with 3 (DataCardinality) elements, that we want to combine in pairs (2 CombinationCardinality),
+// EvaluateCombination will be called (3! / (3-2)!) / 2! = 3 times with input parameters:
+// EvaluateCombination([0, 1])
+// EvaluateCombination([0, 2])
+// EvaluateCombination([1, 2])
+
+template <typename EvaluateCombinationType>
+static void GenerateCombinations(int32 DataCardinality, int32 CombinationCardinality, EvaluateCombinationType EvaluateCombination)
+{
+	TArray<int32, TInlineAllocator<PreallocatedRolesNum>> Combination;
+	Combination.SetNum(CombinationCardinality);
+	GenerateCombinationsRecursive(DataCardinality, 0, Combination, 0, EvaluateCombination);
+}
+
 } // namespace UE::PoseSearch
\ No newline at end of file
