PoseSearch
- MM interaction fix for FInteractionSearchContextBase being generated with non sorted AnimContexts leading to failing detecting continuing interactions, unwanted reselections and non deterministic behaviors across multiple frames
- added WarpingRotationCurveName/WarpingTranslationCurveName to MM trait to control warping via curves

diff --git a/Engine/Plugins/Animation/PoseSearch/Source/Runtime/Private/PoseSearchInteractionIsland.cpp b/Engine/Plugins/Animation/PoseSearch/Source/Runtime/Private/PoseSearchInteractionIsland.cpp
--- a/Engine/Plugins/Animation/PoseSearch/Source/Runtime/Private/PoseSearchInteractionIsland.cpp
+++ b/Engine/Plugins/Animation/PoseSearch/Source/Runtime/Private/PoseSearchInteractionIsland.cpp
@@ -124,13 +124,13 @@ static FInteractionSearchResult InitSearchResult(const FSearchResult& SearchResu
 	static_cast<FSearchResult&>(InteractionSearchResult) = SearchResult;
 	InteractionSearchResult.SearchIndex = SearchIndex;
 
-	const int32 AnimContextsNum = SearchContext.AnimContexts.Num();
+	const int32 AnimContextsNum = SearchContext.Num();
 	InteractionSearchResult.ActorRootTransforms.SetNum(AnimContextsNum);
 	InteractionSearchResult.ActorRootBoneTransforms.SetNum(AnimContextsNum);
 
 	for (int32 AnimContextIndex = 0; AnimContextIndex < AnimContextsNum; ++AnimContextIndex)
 	{
-		const UObject* AnimContext = SearchContext.AnimContexts[AnimContextIndex].Get();
+		const UObject* AnimContext = SearchContext.GetAnimContext(AnimContextIndex);
 		const USkeleton* Skeleton = AnimContext ? GetContextSkeleton(AnimContext) : nullptr;
 		const IPoseHistory* PoseHistory = SearchContext.GetPoseHistory(AnimContextIndex);
 		if (Skeleton && PoseHistory)
@@ -157,9 +157,9 @@ static bool IsPoseSearchResultUsable(int32 SearchIndex, TConstArrayView<FSearchR
 		return false;
 	}
 
-	for (const TWeakObjectPtr<const UObject>& SearchAnimContext : SearchContexts[SearchIndex].AnimContexts)
+	for (int32 AnimContextIndex = 0; AnimContextIndex < SearchContexts[SearchIndex].Num(); ++AnimContextIndex)
 	{
-		if (const UObject* ValidSearchAnimContext = SearchAnimContext.Get())
+		if (const UObject* ValidSearchAnimContext = SearchContexts[SearchIndex].GetAnimContext(AnimContextIndex))
 		{
 			if (VisitedAnimContexts.Find(ValidSearchAnimContext))
 			{
@@ -208,8 +208,8 @@ static void InitSearchResults(TArray<FInteractionSearchResult>& SearchResults, T
 					return true;
 				}
 
-				const int32 NumRolesA = SearchContexts[IndexA].Roles.Num();
-				const int32 NumRolesB = SearchContexts[IndexB].Roles.Num();
+				const int32 NumRolesA = SearchContexts[IndexA].Num();
+				const int32 NumRolesB = SearchContexts[IndexB].Num();
 
 				if (NumRolesA > NumRolesB)
 				{
@@ -232,10 +232,11 @@ static void InitSearchResults(TArray<FInteractionSearchResult>& SearchResults, T
 			{
 				SearchResults.Add(InitSearchResult(PoseSearchResults[SearchIndex], SearchContexts[SearchIndex], SearchIndex));
 
-				for (const TWeakObjectPtr<const UObject>& SearchAnimContext : SearchContexts[SearchIndex].AnimContexts)
+				for (int32 AnimContextIndex = 0; AnimContextIndex < SearchContexts[SearchIndex].Num(); ++AnimContextIndex)
 				{
-					check(SearchAnimContext.Get());
-					VisitedAnimContexts.Add(SearchAnimContext.Get());
+					const UObject* SearchAnimContext = SearchContexts[SearchIndex].GetAnimContext(AnimContextIndex);
+					check(SearchAnimContext);
+					VisitedAnimContexts.Add(SearchAnimContext);
 				}
 			}
 		}
@@ -313,38 +314,11 @@ const UE::PoseSearch::IPoseHistory* FInteractionSearchContextBase::GetPoseHistor
 
 bool FInteractionSearchContextBase::IsEquivalent(const FInteractionSearchContextBase& Other) const
 {
-	if (Database != Other.Database)
-	{
-		return false;
-	}
-
-	const int32 Num = AnimContexts.Num();
-	if (Num != Other.AnimContexts.Num())
-	{
-		return false;
-	}
-
-	int32 CommonRoledAnimContexts = 0;
-	for (int32 IndexThis = 0; IndexThis < Num; ++IndexThis)
-	{
-		for (int32 IndexOther = 0; IndexOther < Num; ++IndexOther)
-		{
-			if (AnimContexts[IndexThis] == Other.AnimContexts[IndexOther] &&
-				PoseHistories[IndexThis] == Other.PoseHistories[IndexOther] &&
-				Roles[IndexThis] == Other.Roles[IndexOther])
-			{
-				++CommonRoledAnimContexts;
-			}
-		}
-	}
-
-	// using >= in case there are duplicated AnimContexts in this or Other (this->IsValid() should be false!)
-	if (CommonRoledAnimContexts >= Num)
-	{
-		return true;
-	}
-
-	return false;
+	// skipping bDisableCollisions for equality
+	return Database == Other.Database &&
+		AnimContexts == Other.AnimContexts &&
+		PoseHistories == Other.PoseHistories &&
+		Roles == Other.Roles;
 }
 
 #if ENABLE_VISUAL_LOG
@@ -424,6 +398,15 @@ bool FInteractionSearchContextBase::CheckForConsistency() const
 		}
 	}
 
+	for (int32 IndexA = 1; IndexA < Num; ++IndexA)
+	{
+		// AnimContexts must be sorted to have deterministic searches across multiple frames
+		if (AnimContexts[IndexA - 1].Get() >= AnimContexts[IndexA].Get())
+		{
+			return false;
+		}
+	}
+
 	for (int32 IndexA = 0; IndexA < Num; ++IndexA)
 	{
 		for (int32 IndexB = IndexA + 1; IndexB < Num; ++IndexB)
@@ -495,9 +478,9 @@ void FInteractionIsland::FPreTickFunction::ExecuteTick(float DeltaTime, enum ELe
 
 		for (const FInteractionSearchContext& SearchContext : Island->SearchContexts)
 		{
-			for (int32 Index = 0; Index < SearchContext.AnimContexts.Num(); ++Index)
+			for (int32 Index = 0; Index < SearchContext.Num(); ++Index)
 			{
-				if (const UObject* AnimContext = SearchContext.AnimContexts[Index].Get())
+				if (const UObject* AnimContext = SearchContext.GetAnimContext(Index))
 				{
 					if (const IPoseHistory* PoseHistory = SearchContext.GetPoseHistory(Index))
 					{
@@ -1018,15 +1001,14 @@ bool FInteractionIsland::DoSearch_AnyThread(const UObject* AnimContext, const TC
 			return false;
 		}
 
-		const int32 NumRoles = SearchContext.AnimContexts.Num();
-		check(NumRoles == SearchContext.PoseHistories.Num());
-
+		const int32 NumRoles = SearchContext.Num();
+		
 		AnimContexts.Reset();
 		PoseHistories.Reset();
 
 		for (int32 RoleIndex = 0; RoleIndex < NumRoles; ++RoleIndex)
 		{
-			const UObject* SearchContextAnimContext = SearchContext.AnimContexts[RoleIndex].Get();
+			const UObject* SearchContextAnimContext = SearchContext.GetAnimContext(RoleIndex);
 			if (!SearchContextAnimContext)
 			{
 				UE_LOG(LogPoseSearch, Error, TEXT("FInteractionIsland::DoSearch_AnyThread null anim context"));
@@ -1049,7 +1031,7 @@ bool FInteractionIsland::DoSearch_AnyThread(const UObject* AnimContext, const TC
 		FPoseSearchFutureProperties PoseSearchFutureProperties;
 
 		// @todo: we could perform multiple UPoseSearchLibrary::MotionMatch in parallel!
-		PoseSearchResults[SearchIndex] = UPoseSearchLibrary::MotionMatch(AnimContexts, SearchContext.Roles, PoseHistories, 
+		PoseSearchResults[SearchIndex] = UPoseSearchLibrary::MotionMatch(AnimContexts, SearchContext.GetRoles(), PoseHistories,
 			AssetsToSearch, SearchContext.GetContinuingProperties(), PoseSearchFutureProperties, FPoseSearchEvent());
 	}
 		
@@ -1069,16 +1051,16 @@ bool FInteractionIsland::GetResult_AnyThread(const UObject* AnimContext, FPoseSe
 	for (const FInteractionSearchResult& SearchResult : SearchResults)
 	{
 		const FInteractionSearchContext& SearchContext = SearchContexts[SearchResult.SearchIndex];
-		for (int32 AnimContextIndex = 0; AnimContextIndex < SearchContext.AnimContexts.Num(); ++AnimContextIndex)
+		for (int32 AnimContextIndex = 0; AnimContextIndex < SearchContext.Num(); ++AnimContextIndex)
 		{
 			bool bAnimContextFound = false;
 			if (bCompareOwningActors)
 			{
-				bAnimContextFound = GetContextOwningActor(SearchContext.AnimContexts[AnimContextIndex].Get()) == Actor;
+				bAnimContextFound = GetContextOwningActor(SearchContext.GetAnimContext(AnimContextIndex)) == Actor;
 			}
 			else
 			{
-				bAnimContextFound = SearchContext.AnimContexts[AnimContextIndex].Get() == AnimContext;
+				bAnimContextFound = SearchContext.GetAnimContext(AnimContextIndex) == AnimContext;
 			}		
 
 			if (bAnimContextFound)
@@ -1102,7 +1084,7 @@ bool FInteractionIsland::GetResult_AnyThread(const UObject* AnimContext, FPoseSe
 				Result.SelectedDatabase = Database;
 				Result.SearchCost = SearchResult.PoseCost;
 				Result.bIsInteraction = true;
-				Result.Role = SearchContext.Roles[AnimContextIndex];
+				Result.Role = SearchContext.GetRole(AnimContextIndex);
 
 				// figuring out the WantedPlayRate
 				Result.WantedPlayRate = 1.f;
@@ -1127,7 +1109,7 @@ bool FInteractionIsland::GetResult_AnyThread(const UObject* AnimContext, FPoseSe
 					Result.ActorRootBoneTransforms.SetNum(NumRoles);
 					Result.AnimContexts.SetNum(NumRoles);
 
-					const FRoleToIndex InteractionSearchContextRoleToIndex = MakeRoleToIndex(SearchContext.Roles);
+					const FRoleToIndex InteractionSearchContextRoleToIndex = MakeRoleToIndex(SearchContext.GetRoles());
 
 					for (int32 MultiAnimAssetRoleIndex = 0; MultiAnimAssetRoleIndex < NumRoles; ++MultiAnimAssetRoleIndex)
 					{
@@ -1135,7 +1117,7 @@ bool FInteractionIsland::GetResult_AnyThread(const UObject* AnimContext, FPoseSe
 						{
 							Result.ActorRootTransforms[MultiAnimAssetRoleIndex] = SearchResult.ActorRootTransforms[*InteractionSearchContextRoleIndex];
 							Result.ActorRootBoneTransforms[MultiAnimAssetRoleIndex] = SearchResult.ActorRootBoneTransforms[*InteractionSearchContextRoleIndex];
-							Result.AnimContexts[MultiAnimAssetRoleIndex] = SearchContext.AnimContexts[*InteractionSearchContextRoleIndex].Get();
+							Result.AnimContexts[MultiAnimAssetRoleIndex] = SearchContext.GetAnimContext(*InteractionSearchContextRoleIndex);
 						}
 						else
 						{
@@ -1154,7 +1136,7 @@ bool FInteractionIsland::GetResult_AnyThread(const UObject* AnimContext, FPoseSe
 					Result.ActorRootBoneTransforms = SearchResult.ActorRootBoneTransforms;
 					
 					Result.AnimContexts.SetNum(1);
-					Result.AnimContexts[0] = SearchContext.AnimContexts[0].Get();
+					Result.AnimContexts[0] = SearchContext.GetAnimContext(0);
 				}
 
 				// we found our AnimContext: we can stop searching
@@ -1176,9 +1158,7 @@ const FInteractionSearchResult* FInteractionIsland::FindSearchResult(const FInte
 	for (const FInteractionSearchResult& SearchResult : SearchResults)
 	{
 		const FInteractionSearchContext& LocalSearchContext = SearchContexts[SearchResult.SearchIndex];
-		if (LocalSearchContext.Database == SearchContext.Database &&
-			LocalSearchContext.AnimContexts == SearchContext.AnimContexts &&
-			LocalSearchContext.Roles == SearchContext.Roles)
+		if (LocalSearchContext.IsEquivalent(SearchContext))
 		{
 			return &SearchResult;
 		}
diff --git a/Engine/Plugins/Animation/PoseSearch/Source/Runtime/Private/PoseSearchInteractionSubsystem.cpp b/Engine/Plugins/Animation/PoseSearch/Source/Runtime/Private/PoseSearchInteractionSubsystem.cpp
--- a/Engine/Plugins/Animation/PoseSearch/Source/Runtime/Private/PoseSearchInteractionSubsystem.cpp
+++ b/Engine/Plugins/Animation/PoseSearch/Source/Runtime/Private/PoseSearchInteractionSubsystem.cpp
@@ -176,9 +176,9 @@ void UPoseSearchInteractionSubsystem::RegenerateAllIslands(float DeltaSeconds)
 	{
 		bool Contains(const FInteractionSearchContext& SearchContext) const
 		{
-			for (const TWeakObjectPtr<const UObject>& AnimContext : SearchContext.AnimContexts)
+			for (int32 AnimContextIndex = 0; AnimContextIndex < SearchContext.Num(); ++AnimContextIndex)
 			{
-				if (AnimContextToTickPriority.Find(AnimContext.Get()))
+				if (AnimContextToTickPriority.Find(SearchContext.GetAnimContext(AnimContextIndex)))
 				{
 					return true;
 				}
@@ -188,9 +188,9 @@ void UPoseSearchInteractionSubsystem::RegenerateAllIslands(float DeltaSeconds)
 
 		void Add(const FInteractionSearchContext& SearchContext, int32 SearchContextIndex)
 		{
-			for (int32 AnimContextIndex = 0; AnimContextIndex < SearchContext.AnimContexts.Num(); ++AnimContextIndex)
+			for (int32 AnimContextIndex = 0; AnimContextIndex < SearchContext.Num(); ++AnimContextIndex)
 			{
-				if (const UObject* AnimContext = SearchContext.AnimContexts[AnimContextIndex].Get())
+				if (const UObject* AnimContext = SearchContext.GetAnimContext(AnimContextIndex))
 				{
 					if (int32* TickPriority = AnimContextToTickPriority.Find(AnimContext))
 					{
@@ -239,21 +239,28 @@ void UPoseSearchInteractionSubsystem::RegenerateAllIslands(float DeltaSeconds)
 	TArray<FInteractionSearchContextGroup, TInlineAllocator<PreallocatedSearchesNum, TMemStackAllocator<>>> SearchContextGroups;
 	for (int32 SearchContextIndex = 0; SearchContextIndex < SearchContexts.Num(); ++SearchContextIndex)
 	{
+		// evaluating where to place SearchContext..
 		const FInteractionSearchContext& SearchContext = SearchContexts[SearchContextIndex];
 
 		int32 MainSearchContextGroupIndex = INDEX_NONE;
 		for (int32 SearchContextGroupIndex = 0; SearchContextGroupIndex < SearchContextGroups.Num();)
 		{
+			// ..if SearchContextGroups[SearchContextGroupIndex] contains ANY of the AnimContexts from SearchContext..
 			if (SearchContextGroups[SearchContextGroupIndex].Contains(SearchContext))
 			{
 				if (MainSearchContextGroupIndex == INDEX_NONE)
 				{
+					// ..we add SearchContext to SearchContextGroups[SearchContextGroupIndex] 
+					// and set MainSearchContextGroupIndex to SearchContextGroupIndex to know what is the group containing SearchContext, so..
 					MainSearchContextGroupIndex = SearchContextGroupIndex;
 					SearchContextGroups[MainSearchContextGroupIndex].Add(SearchContext, SearchContextIndex);
 					++SearchContextGroupIndex;
 				}
 				else
 				{
+					// ..in case SearchContext has already being inserted in MainSearchContextGroupIndex group 
+					// we merge the newly found SearchContextGroups[SearchContextGroupIndex] to SearchContextGroups[MainSearchContextGroupIndex]
+					// (containing another of the the AnimContexts)
 					SearchContextGroups[MainSearchContextGroupIndex].Merge(SearchContextGroups[SearchContextGroupIndex]);
 					SearchContextGroups.RemoveAt(SearchContextGroupIndex);
 				}
@@ -284,11 +291,12 @@ void UPoseSearchInteractionSubsystem::RegenerateAllIslands(float DeltaSeconds)
 			const FInteractionSearchContext& SearchContext = SearchContexts[SearchContextsIndex];
 			// if there're at least two AnimContext(s) potentially interacting with each other 
 			// (where the search involves 2+ characters) tick dependencies are required to be thread safe
-			bAreTickDependenciesRequired |= SearchContext.AnimContexts.Num() > 1;
+			bAreTickDependenciesRequired |= SearchContext.Num() > 1;
 			Island.AddSearchContext(SearchContext);
 		}
 
 		// sorting SearchContextGroup.AnimContextToTickPriority by TickPriority
+		// (using SortedByTickPriorityAnimContexts since SearchContextGroup.AnimContextToTickPriority it's a TMap)
 		SortedByTickPriorityAnimContexts.Reset();
 		SortedByTickPriorityAnimContexts.Reserve(SearchContextGroup.AnimContextToTickPriority.Num());
 		for (const FInteractionSearchContextGroup::FAnimContextToTickPriorityPair& AnimContextToTickPriorityPair : SearchContextGroup.AnimContextToTickPriority)
@@ -301,7 +309,10 @@ void UPoseSearchInteractionSubsystem::RegenerateAllIslands(float DeltaSeconds)
 				return B.Value < A.Value;
 			});
 
-		// injecting tick dependencies between island AnimContext
+		// injecting tick dependencies between island AnimContext following their TickPriorities, 
+		// so the AnimContext with the highest TickPriority will be elected as "Main Actor", being evaluated, 
+		// and performing all the island searches, before any other Actor in the same island\
+		// (that'll end up using the cached search results in a multithread manner)
 		for (const FInteractionSearchContextGroup::FAnimContextToTickPriorityPair& SortedByTickPriorityAnimContext : SortedByTickPriorityAnimContexts)
 		{
 			Island.InjectToActor(SortedByTickPriorityAnimContext.Key, bAreTickDependenciesRequired);
@@ -337,9 +348,9 @@ bool UPoseSearchInteractionSubsystem::ValidateAllIslands() const
 		FIslandAnimContexts& IslandAnimContexts = IslandsAnimContexts.AddDefaulted_GetRef();
 		for (const FInteractionSearchContext& SearchContext : Island->GetSearchContexts())
 		{
-			for (const TWeakObjectPtr<const UObject>& AnimContextPtr : SearchContext.AnimContexts)
+			for (int32 AnimContextIndex = 0; AnimContextIndex < SearchContext.Num(); ++AnimContextIndex)
 			{
-				if (const UObject* AnimContext = AnimContextPtr.Get())
+				if (const UObject* AnimContext = SearchContext.GetAnimContext(AnimContextIndex))
 				{
 					IslandAnimContexts.Add(AnimContext);
 				}
@@ -842,13 +853,19 @@ void UPoseSearchInteractionSubsystem::GenerateSearchContexts(float DeltaSeconds,
 			//
 			// CharB/RoleA - CharB/RoleB -> invalid because of same CharB
 
-			for (const TPair<const UPoseSearchDatabase*, FRoledAnimContextInfos>& DatabaseToRoledAnimContextInfosPair : DatabaseToRoledAnimContextInfos)
+			for (TPair<const UPoseSearchDatabase*, FRoledAnimContextInfos>& DatabaseToRoledAnimContextInfosPair : DatabaseToRoledAnimContextInfos)
 			{
 				const UPoseSearchDatabase* Database = DatabaseToRoledAnimContextInfosPair.Key;
 				check(Database->Schema);
 				const TArray<FPoseSearchRoledSkeleton>& RoledSkeletons = Database->Schema->GetRoledSkeletons();
 				const int32 CombinationCardinality = RoledSkeletons.Num();
-				const FRoledAnimContextInfos& RoledAnimContextInfos = DatabaseToRoledAnimContextInfosPair.Value;
+				FRoledAnimContextInfos& RoledAnimContextInfos = DatabaseToRoledAnimContextInfosPair.Value;
+
+				// sort RoledAnimContextInfos to generate deterministic SearchContext across multiple frames!
+				RoledAnimContextInfos.Sort([](const FRoledAnimContextInfo& RoledAnimContextInfoA, const FRoledAnimContextInfo& RoledAnimContextInfoB)
+					{
+						return RoledAnimContextInfoA.AnimContextInfo->AnimContext < RoledAnimContextInfoB.AnimContextInfo->AnimContext;
+					});
 
 				GenerateCombinations(RoledAnimContextInfos.Num(), CombinationCardinality,
 					// Combination is an array of indexes in RoledAnimContextInfos: 0 <= Combination[i] < RoledAnimContextInfos.Num()
@@ -883,9 +900,7 @@ void UPoseSearchInteractionSubsystem::GenerateSearchContexts(float DeltaSeconds,
 							const FRoledAnimContextInfo& RoledAnimContextInfo = RoledAnimContextInfos[RoledAnimContextIndex];
 
 							check(RoledAnimContextInfo.PoseHistory && RoledAnimContextInfo.Availability);
-							SearchContext.AnimContexts.Add(RoledAnimContextInfo.AnimContextInfo->AnimContext);
-							SearchContext.Roles.Add(RoledAnimContextInfo.Role);
-							SearchContext.PoseHistories.Add(RoledAnimContextInfo.PoseHistory->AsWeak());
+							SearchContext.Add(RoledAnimContextInfo.AnimContextInfo->AnimContext, RoledAnimContextInfo.PoseHistory, RoledAnimContextInfo.Role);
 							SearchContext.bDisableCollisions |= RoledAnimContextInfo.Availability->bDisableCollisions;
 							SearchContext.TickPriorities.Add(RoledAnimContextInfo.Availability->TickPriority);
 							
@@ -898,7 +913,7 @@ void UPoseSearchInteractionSubsystem::GenerateSearchContexts(float DeltaSeconds,
 						for (const FPoseSearchRoledSkeleton& RoledSkeleton : RoledSkeletons)
 						{
 							// CombinationCardinality is usually 2-3, so we can search the SearchContext.Roles array for duplicates without requiring a faster container like TSet
-							if (!SearchContext.Roles.Contains(RoledSkeleton.Role))
+							if (!SearchContext.GetRoles().Contains(RoledSkeleton.Role))
 							{
 								return false;
 							}
@@ -938,6 +953,8 @@ void UPoseSearchInteractionSubsystem::GenerateSearchContexts(float DeltaSeconds,
 						{
 							check(!ContainedSearchContext.IsEquivalent(SearchContext));
 						}
+
+						check(SearchContext.CheckForConsistency());
 #endif // DO_CHECK
 						SearchContexts.Emplace(SearchContext);
 						return true;
@@ -966,9 +983,9 @@ void UPoseSearchInteractionSubsystem::OnInteractionStart(UE::PoseSearch::FValidI
 		TArray<AActor*, TInlineAllocator<PreallocatedRolesNum, TMemStackAllocator<>>> Actors;
 		TArray<UPrimitiveComponent*, TInlineAllocator<PreallocatedRolesNum, TMemStackAllocator<>>> PrimitiveComponents;
 
-		for (const TWeakObjectPtr<const UObject>& AnimContextPtr : ValidInteractionSearch.SearchContext.AnimContexts)
+		for (int32 AnimContextIndex = 0; AnimContextIndex < ValidInteractionSearch.SearchContext.Num(); ++AnimContextIndex)
 		{
-			if (const UObject* AnimContext = AnimContextPtr.Get())
+			if (const UObject* AnimContext = ValidInteractionSearch.SearchContext.GetAnimContext(AnimContextIndex))
 			{
 				AActor* Actor = const_cast<AActor*>(GetContextOwningActor(AnimContext));
 				check(Actor);
@@ -1250,9 +1267,9 @@ void UPoseSearchInteractionSubsystem::DebugDrawIslands() const
 
 			for (const FInteractionSearchContext& SearchContext : Island->GetSearchContexts())
 			{
-				for (int32 Index = 0; Index < SearchContext.AnimContexts.Num(); ++Index)
+				for (int32 Index = 0; Index < SearchContext.Num(); ++Index)
 				{
-					if (const UObject* AnimContext = SearchContext.AnimContexts[Index].Get())
+					if (const UObject* AnimContext = SearchContext.GetAnimContext(Index))
 					{
 						const FPoseSearchInteractionAvailability& DebugAvailability = SearchContext.DebugAvailabilities[Index];
 						float MaxBroadPhaseRadius;
diff --git a/Engine/Plugins/Animation/PoseSearch/Source/Runtime/Public/PoseSearch/PoseSearchInteractionIsland.h b/Engine/Plugins/Animation/PoseSearch/Source/Runtime/Public/PoseSearch/PoseSearchInteractionIsland.h
--- a/Engine/Plugins/Animation/PoseSearch/Source/Runtime/Public/PoseSearch/PoseSearchInteractionIsland.h
+++ b/Engine/Plugins/Animation/PoseSearch/Source/Runtime/Public/PoseSearch/PoseSearchInteractionIsland.h
@@ -19,16 +19,36 @@ struct IInteractionIslandDependency;
 // Experimental, this feature might be removed without warning, not for production use
 struct FInteractionSearchContextBase
 {
+private:
 	TArray<TWeakObjectPtr<const UObject>, TInlineAllocator<PreallocatedRolesNum>> AnimContexts;
 	TArray<TWeakPtr<const UE::PoseSearch::IPoseHistory, ESPMode::ThreadSafe>, TInlineAllocator<PreallocatedRolesNum>> PoseHistories;
 	TArray<UE::PoseSearch::FRole, TInlineAllocator<PreallocatedRolesNum>> Roles;
 
+public:
 	TWeakObjectPtr<const UPoseSearchDatabase> Database;
 	bool bDisableCollisions = false;
 
+	void Add(const UObject* AnimContext, const UE::PoseSearch::IPoseHistory* PoseHistory, const UE::PoseSearch::FRole Role)
+	{
+		check(AnimContext && PoseHistory);
+		
+		// AnimContexts must be sorted to have deterministic searches across multiple frames
+		// as well as for faster IsEquivalent implementation
+		check(AnimContexts.IsEmpty() || GetAnimContext(AnimContexts.Num() - 1) < AnimContext);
+
+		AnimContexts.Add(AnimContext);
+		PoseHistories.Add(PoseHistory->AsWeak());
+		Roles.Add(Role);
+	}
+
+	int32 Num() const { return AnimContexts.Num(); }
+	const UObject* GetAnimContext(int32 Index) const { return AnimContexts[Index].Get(); }
+	const UE::PoseSearch::IPoseHistory* GetPoseHistory(int32 Index) const;
+	const UE::PoseSearch::FRole GetRole(int32 Index) const { return Roles[Index]; }
+	const TConstArrayView<UE::PoseSearch::FRole> GetRoles() const { return Roles; }
+
 	bool IsValid() const { return !AnimContexts.IsEmpty(); }
 	bool IsEquivalent(const FInteractionSearchContextBase& Other) const;
-	const UE::PoseSearch::IPoseHistory* GetPoseHistory(int32 Index) const;
 
 #if ENABLE_VISUAL_LOG
 	void VLogContext(const FColor& Color) const;
diff --git a/Engine/Plugins/Experimental/AnimNextPoseSearch/Source/AnimNextPoseSearch/Private/MotionMatchingTrait.cpp b/Engine/Plugins/Experimental/AnimNextPoseSearch/Source/AnimNextPoseSearch/Private/MotionMatchingTrait.cpp
--- a/Engine/Plugins/Experimental/AnimNextPoseSearch/Source/AnimNextPoseSearch/Private/MotionMatchingTrait.cpp
+++ b/Engine/Plugins/Experimental/AnimNextPoseSearch/Source/AnimNextPoseSearch/Private/MotionMatchingTrait.cpp
@@ -507,11 +507,12 @@ namespace UE::AnimNext
 								Task.bWarpUsingRootBone = SharedData->GetbWarpUsingRootBone(Binding);
 								Task.WarpingRotationRatio = SharedData->GetWarpingRotationRatio(Binding);
 								Task.WarpingTranslationRatio = SharedData->GetWarpingTranslationRatio(Binding);
+								Task.WarpingRotationCurveName = SharedData->GetWarpingRotationCurveName(Binding);
+								Task.WarpingTranslationCurveName = SharedData->GetWarpingTranslationCurveName(Binding);
 #if ENABLE_ANIM_DEBUG
 								// Debug Object for VisualLogger
 								Task.HostObject = Context.GetHostObject();
 #endif // ENABLE_ANIM_DEBUG 
-
 								Context.AppendTask(Task);
 							}
 						}
@@ -581,7 +582,11 @@ void FAnimNextMotionMatchingTask::Execute(UE::AnimNext::FEvaluationVM& VM) const
 	{
 		return;
 	}
-			
+
+	bool OutHasElement = false;
+	const float FinalWarpingRotationRatio = FMath::Clamp(Keyframe->Get()->Curves.Get(WarpingRotationCurveName, OutHasElement, WarpingRotationRatio), 0.f, 1.f);
+	const float FinalWarpingTranslationRatio = FMath::Clamp(Keyframe->Get()->Curves.Get(WarpingTranslationCurveName, OutHasElement, WarpingTranslationRatio), 0.f, 1.f);
+
 	TArray<FTransform, TInlineAllocator<PreallocatedRolesNum>> FullAlignedTransforms;
 	FullAlignedTransforms.SetNum(NumRoles);
 
@@ -592,7 +597,7 @@ void FAnimNextMotionMatchingTask::Execute(UE::AnimNext::FEvaluationVM& VM) const
 	const FTransform& MeshWithOffset = MeshWithoutOffset;
 
 	// NoTe: keep in mind DeltaAlignment is relative to the previous execution frame so we still need to extract and and apply the current animation root motion transform to get to the current frame full aligned transform.
-	const FTransform DeltaAlignment = CalculateDeltaAlignment(MeshWithoutOffset, MeshWithOffset, FullAlignedTransforms[CurrentResultRoleIndex], WarpingRotationRatio, WarpingTranslationRatio);
+	const FTransform DeltaAlignment = CalculateDeltaAlignment(MeshWithoutOffset, MeshWithOffset, FullAlignedTransforms[CurrentResultRoleIndex], FinalWarpingRotationRatio, FinalWarpingTranslationRatio);
 	
 	const FTransform DeltaAlignmentWithRootMotion = DeltaAlignment * RootMotionDelta;
 
diff --git a/Engine/Plugins/Experimental/AnimNextPoseSearch/Source/AnimNextPoseSearch/Private/MotionMatchingTrait.h b/Engine/Plugins/Experimental/AnimNextPoseSearch/Source/AnimNextPoseSearch/Private/MotionMatchingTrait.h
--- a/Engine/Plugins/Experimental/AnimNextPoseSearch/Source/AnimNextPoseSearch/Private/MotionMatchingTrait.h
+++ b/Engine/Plugins/Experimental/AnimNextPoseSearch/Source/AnimNextPoseSearch/Private/MotionMatchingTrait.h
@@ -105,6 +105,12 @@ struct FMotionMatchingTraitSharedData : public FAnimNextTraitSharedData
 	UPROPERTY(EditAnywhere, Category = Interaction, meta = (PinHiddenByDefault))
 	float WarpingTranslationRatio = 1.f;
 
+	UPROPERTY(EditAnywhere, Category = Interaction, meta = (PinHiddenByDefault))
+	FName WarpingRotationCurveName;
+
+	UPROPERTY(EditAnywhere, Category = Interaction, meta = (PinHiddenByDefault))
+	FName WarpingTranslationCurveName;
+
 	// Latent pin support boilerplate
 	#define TRAIT_LATENT_PROPERTIES_ENUMERATOR(GeneratorMacro) \
 		GeneratorMacro(Databases) \
@@ -130,6 +136,8 @@ struct FMotionMatchingTraitSharedData : public FAnimNextTraitSharedData
 		GeneratorMacro(bWarpUsingRootBone) \
 		GeneratorMacro(WarpingRotationRatio) \
 		GeneratorMacro(WarpingTranslationRatio) \
+		GeneratorMacro(WarpingRotationCurveName) \
+		GeneratorMacro(WarpingTranslationCurveName) \
 
 	GENERATE_TRAIT_LATENT_PROPERTIES(FMotionMatchingTraitSharedData, TRAIT_LATENT_PROPERTIES_ENUMERATOR)
 	#undef TRAIT_LATENT_PROPERTIES_ENUMERATOR
@@ -184,6 +192,8 @@ struct FAnimNextMotionMatchingTask : public FAnimNextEvaluationTask
 	bool bWarpUsingRootBone = true;
 	float WarpingRotationRatio = 1.f;
 	float WarpingTranslationRatio = 1.f;
+	FName WarpingRotationCurveName;
+	FName WarpingTranslationCurveName;
 
 #if ENABLE_ANIM_DEBUG
 	// Debug Object for VisualLogger
