Mark PoseSearch asset editor animation sampler as not requiring compressed data

diff --git a/Engine/Plugins/Animation/PoseSearch/Source/Editor/Private/PoseSearchInteractionAssetEditor.cpp b/Engine/Plugins/Animation/PoseSearch/Source/Editor/Private/PoseSearchInteractionAssetEditor.cpp
--- a/Engine/Plugins/Animation/PoseSearch/Source/Editor/Private/PoseSearchInteractionAssetEditor.cpp
+++ b/Engine/Plugins/Animation/PoseSearch/Source/Editor/Private/PoseSearchInteractionAssetEditor.cpp
@@ -46,7 +46,7 @@ bool FInteractionAssetPreviewActor::SpawnPreviewActor(UWorld* World, const UPose
 		FAnimationAssetSampler& Sampler = Samplers[ActorIndex];
 
 		const FTransform Origin = InteractionAsset->GetDebugWarpOrigin(GetRole(), ActorIndex == DebugActor);
-		Sampler.Init(PreviewAsset, Origin, BlendParameters);
+		Sampler.Init(PreviewAsset, Origin, BlendParameters, FAnimationAssetSampler::DefaultRootTransformSamplingRate, true, false);
 
 		const FTransform ActorTransform = Sampler.ExtractRootTransform(CurrentTime);
 
