[ObjectTrace] removing some header dependencies

diff --git a/Engine/Plugins/Animation/AnimationWarping/Source/Runtime/Private/BoneControllers/AnimNode_FootPlacement.cpp b/Engine/Plugins/Animation/AnimationWarping/Source/Runtime/Private/BoneControllers/AnimNode_FootPlacement.cpp
--- a/Engine/Plugins/Animation/AnimationWarping/Source/Runtime/Private/BoneControllers/AnimNode_FootPlacement.cpp
+++ b/Engine/Plugins/Animation/AnimationWarping/Source/Runtime/Private/BoneControllers/AnimNode_FootPlacement.cpp
@@ -6,6 +6,7 @@
 #include "AnimationRuntime.h"
 #include "Animation/AnimInstanceProxy.h"
 #include "Components/SkeletalMeshComponent.h"
+#include "Engine/World.h"
 #include "GameFramework/Character.h"
 #include "GameFramework/CharacterMovementComponent.h"
 #include "Animation/AnimRootMotionProvider.h"
diff --git a/Engine/Plugins/Animation/BlendStack/Source/Runtime/Private/AnimNode_BlendStack.cpp b/Engine/Plugins/Animation/BlendStack/Source/Runtime/Private/AnimNode_BlendStack.cpp
--- a/Engine/Plugins/Animation/BlendStack/Source/Runtime/Private/AnimNode_BlendStack.cpp
+++ b/Engine/Plugins/Animation/BlendStack/Source/Runtime/Private/AnimNode_BlendStack.cpp
@@ -14,9 +14,9 @@
 #include "Animation/BlendSpace.h"
 #include "Animation/NamedValueArray.h"
 #include "BlendedCurveUtils.h"
-#include "BlendStack/AnimNode_BlendStackInput.h"
 #include "BlendStack/BlendStackDefines.h"
 #include "BlendStackAnimEventsFilterScope.h"
+#include "ObjectTrace.h"
 
 #include UE_INLINE_GENERATED_CPP_BY_NAME(AnimNode_BlendStack)
 
diff --git a/Engine/Plugins/Animation/GameplayInsights/Source/GameplayInsights/Private/ObjectPropertiesTrack.cpp b/Engine/Plugins/Animation/GameplayInsights/Source/GameplayInsights/Private/ObjectPropertiesTrack.cpp
--- a/Engine/Plugins/Animation/GameplayInsights/Source/GameplayInsights/Private/ObjectPropertiesTrack.cpp
+++ b/Engine/Plugins/Animation/GameplayInsights/Source/GameplayInsights/Private/ObjectPropertiesTrack.cpp
@@ -61,7 +61,7 @@ void FObjectPropertiesTrack::InitTooltip(FTooltipDrawState& Tooltip, const ITimi
 
 		FTimingEventSearchParameters SearchParameters(HoveredTimingEvent.GetStartTime(), HoveredTimingEvent.GetEndTime(), ETimingEventSearchFlags::StopAtFirstMatch);
 
-		FindObjectEvent(SearchParameters, [this, &Tooltip, &HoveredTimingEvent, &GameplayProvider](double InFoundStartTime, double InFoundEndTime, uint32 InFoundDepth, const FObjectPropertiesMessage& InMessage)
+		FindObjectEvent(SearchParameters, [this, &Tooltip, &HoveredTimingEvent](double InFoundStartTime, double InFoundEndTime, uint32 InFoundDepth, const FObjectPropertiesMessage& InMessage)
 		{
 			Tooltip.ResetContent();
 
diff --git a/Engine/Plugins/Animation/GameplayInsights/Source/GameplayInsights/Private/PropertyTrack.cpp b/Engine/Plugins/Animation/GameplayInsights/Source/GameplayInsights/Private/PropertyTrack.cpp
--- a/Engine/Plugins/Animation/GameplayInsights/Source/GameplayInsights/Private/PropertyTrack.cpp
+++ b/Engine/Plugins/Animation/GameplayInsights/Source/GameplayInsights/Private/PropertyTrack.cpp
@@ -6,7 +6,6 @@
 #include "AnimationProvider.h"
 #include "GameplayProvider.h"
 #include "IRewindDebugger.h"
-// #include "ObjectTrace.h"
 #include "TraceServices/Model/AnalysisSession.h"
 #include "TraceServices/Model/Frames.h"
 #include "SPropertyTrackView.h"
diff --git a/Engine/Plugins/Animation/GameplayInsights/Source/GameplayInsights/Private/SPropertiesDebugViewBase.h b/Engine/Plugins/Animation/GameplayInsights/Source/GameplayInsights/Private/SPropertiesDebugViewBase.h
--- a/Engine/Plugins/Animation/GameplayInsights/Source/GameplayInsights/Private/SPropertiesDebugViewBase.h
+++ b/Engine/Plugins/Animation/GameplayInsights/Source/GameplayInsights/Private/SPropertiesDebugViewBase.h
@@ -2,7 +2,6 @@
 
 #pragma once
 
-#include "CoreMinimal.h"
 #include "IRewindDebuggerView.h"
 #include "IRewindDebuggerViewCreator.h"
 #include "SVariantValueView.h"
diff --git a/Engine/Plugins/Animation/GameplayInsights/Source/RewindDebuggerVLog/Private/VisualLoggerProvider.h b/Engine/Plugins/Animation/GameplayInsights/Source/RewindDebuggerVLog/Private/VisualLoggerProvider.h
--- a/Engine/Plugins/Animation/GameplayInsights/Source/RewindDebuggerVLog/Private/VisualLoggerProvider.h
+++ b/Engine/Plugins/Animation/GameplayInsights/Source/RewindDebuggerVLog/Private/VisualLoggerProvider.h
@@ -6,7 +6,6 @@
 #include "Model/PointTimeline.h"
 #include "Containers/Map.h"
 #include "Templates/SharedPointer.h"
-#include "Model/IntervalTimeline.h"
 #include "VisualLogger/VisualLoggerTypes.h"
 
 namespace TraceServices { class IAnalysisSession; }
@@ -16,12 +15,12 @@ class FVisualLoggerProvider : public IVisualLoggerProvider
 public:
 	static FName ProviderName;
 
-	FVisualLoggerProvider(TraceServices::IAnalysisSession& InSession);
+	explicit FVisualLoggerProvider(TraceServices::IAnalysisSession& InSession);
 
 	/** IVisualLoggerProvider interface */
 	virtual bool ReadVisualLogEntryTimeline(uint64 InObjectId, TFunctionRef<void(const VisualLogEntryTimeline&)> Callback) const;
 
-	virtual void EnumerateCategories(TFunctionRef<void(const FName&)> Callback) const;
+	virtual void EnumerateCategories(TFunctionRef<void(const FName&)> Callback) const override;
 
 	/** Add an object event message */
 	void AppendVisualLogEntry(uint64 InObjectId, double InTime, const FVisualLogEntry& Entry);
diff --git a/Engine/Plugins/Animation/PoseSearch/Source/Editor/Private/Chooser/PoseSearchColumnEditor.cpp b/Engine/Plugins/Animation/PoseSearch/Source/Editor/Private/Chooser/PoseSearchColumnEditor.cpp
--- a/Engine/Plugins/Animation/PoseSearch/Source/Editor/Private/Chooser/PoseSearchColumnEditor.cpp
+++ b/Engine/Plugins/Animation/PoseSearch/Source/Editor/Private/Chooser/PoseSearchColumnEditor.cpp
@@ -2,16 +2,13 @@
 
 #include "PoseSearchColumnEditor.h"
 
-#include <ChooserColumnHeader.h>
-
+#include "Animation/AnimationAsset.h"
+#include "ObjectChooserWidgetFactories.h"
 #include "PoseSearch/Chooser/PoseSearchChooserColumn.h"
 #include "SPropertyAccessChainWidget.h"
-#include "ObjectChooserWidgetFactories.h"
-#include "Widgets/SBoxPanel.h"
-#include "Widgets/Layout/SBox.h"
-#include "Widgets/Layout/SSpacer.h"
 #include "Widgets/Input/SNumericEntryBox.h"
-#include "ScopedTransaction.h"
+#include "Widgets/SBoxPanel.h"
+#include <ChooserColumnHeader.h>
 
 #define LOCTEXT_NAMESPACE "PoseSearchColumnEditor"
 
diff --git a/Engine/Plugins/Animation/PoseSearch/Source/Editor/Private/PoseSearchDatabaseDataDetails.cpp b/Engine/Plugins/Animation/PoseSearch/Source/Editor/Private/PoseSearchDatabaseDataDetails.cpp
--- a/Engine/Plugins/Animation/PoseSearch/Source/Editor/Private/PoseSearchDatabaseDataDetails.cpp
+++ b/Engine/Plugins/Animation/PoseSearch/Source/Editor/Private/PoseSearchDatabaseDataDetails.cpp
@@ -1,6 +1,7 @@
 // Copyright Epic Games, Inc. All Rights Reserved.
 
 #include "PoseSearchDatabaseDataDetails.h"
+#include "Animation/AnimationAsset.h"
 #include "PoseSearch/PoseSearchDatabase.h"
 #include "PoseSearch/PoseSearchDefines.h"
 #include "PoseSearch/PoseSearchDerivedData.h"
diff --git a/Engine/Plugins/Animation/PoseSearch/Source/Editor/Private/PoseSearchDatabaseEditorUtils.cpp b/Engine/Plugins/Animation/PoseSearch/Source/Editor/Private/PoseSearchDatabaseEditorUtils.cpp
--- a/Engine/Plugins/Animation/PoseSearch/Source/Editor/Private/PoseSearchDatabaseEditorUtils.cpp
+++ b/Engine/Plugins/Animation/PoseSearch/Source/Editor/Private/PoseSearchDatabaseEditorUtils.cpp
@@ -2,6 +2,7 @@
 
 #include "PoseSearchDatabaseEditorUtils.h"
 
+#include "Animation/AnimationAsset.h"
 #include "PoseSearch/MultiAnimAsset.h"
 #include "PoseSearch/PoseSearchDatabase.h"
 #include "PoseSearch/PoseSearchSchema.h"
diff --git a/Engine/Plugins/Animation/PoseSearch/Source/Runtime/Private/AnimNode_PoseSearchHistoryCollector.cpp b/Engine/Plugins/Animation/PoseSearch/Source/Runtime/Private/AnimNode_PoseSearchHistoryCollector.cpp
--- a/Engine/Plugins/Animation/PoseSearch/Source/Runtime/Private/AnimNode_PoseSearchHistoryCollector.cpp
+++ b/Engine/Plugins/Animation/PoseSearch/Source/Runtime/Private/AnimNode_PoseSearchHistoryCollector.cpp
@@ -1,12 +1,14 @@
 // Copyright Epic Games, Inc. All Rights Reserved.
 
 #include "PoseSearch/AnimNode_PoseSearchHistoryCollector.h"
+#include "Animation/AnimInstance.h"
 #include "Animation/AnimInstanceProxy.h"
 #include "Animation/AnimNodeMessages.h"
 #include "Animation/AnimStats.h"
+#include "Components/SkeletalMeshComponent.h"
 #include "Engine/SkinnedAsset.h"
-#include "PoseSearch/PoseSearchInteractionValidator.h"
 #include "PoseSearch/PoseSearchHistoryAttribute.h"
+#include "PoseSearch/PoseSearchInteractionValidator.h"
 
 #include UE_INLINE_GENERATED_CPP_BY_NAME(AnimNode_PoseSearchHistoryCollector)
 
diff --git a/Engine/Plugins/Animation/PoseSearch/Source/Runtime/Private/PoseSearchAssetSamplerLibrary.cpp b/Engine/Plugins/Animation/PoseSearch/Source/Runtime/Private/PoseSearchAssetSamplerLibrary.cpp
--- a/Engine/Plugins/Animation/PoseSearch/Source/Runtime/Private/PoseSearchAssetSamplerLibrary.cpp
+++ b/Engine/Plugins/Animation/PoseSearch/Source/Runtime/Private/PoseSearchAssetSamplerLibrary.cpp
@@ -2,6 +2,7 @@
 
 #include "PoseSearch/PoseSearchAssetSamplerLibrary.h"
 #include "Animation/AnimInstance.h"
+#include "Components/SkeletalMeshComponent.h"
 #include "DrawDebugHelpers.h"
 #include "Engine/SkinnedAsset.h"
 #include "PoseSearch/PoseSearchDefines.h"
diff --git a/Engine/Plugins/Animation/PoseSearch/Source/Runtime/Private/PoseSearchContext.cpp b/Engine/Plugins/Animation/PoseSearch/Source/Runtime/Private/PoseSearchContext.cpp
--- a/Engine/Plugins/Animation/PoseSearch/Source/Runtime/Private/PoseSearchContext.cpp
+++ b/Engine/Plugins/Animation/PoseSearch/Source/Runtime/Private/PoseSearchContext.cpp
@@ -3,15 +3,15 @@
 #include "PoseSearch/PoseSearchContext.h"
 #include "Animation/AnimInstanceProxy.h"
 #include "AnimationRuntime.h"
+#include "BoneContainer.h"
+#include "Components/SkeletalMeshComponent.h"
 #include "Engine/SkeletalMesh.h"
 #include "PoseSearch/PoseSearchDatabase.h"
-#include "PoseSearch/PoseSearchHistory.h"
-#include "PoseSearch/PoseSearchSchema.h"
-#include "Animation/TrajectoryTypes.h"
 #include "PoseSearch/PoseSearchFeatureChannel_Heading.h"
 #include "PoseSearch/PoseSearchFeatureChannel_PermutationTime.h"
 #include "PoseSearch/PoseSearchFeatureChannel_Position.h"
-#include "BoneContainer.h"
+#include "PoseSearch/PoseSearchHistory.h"
+#include "PoseSearch/PoseSearchSchema.h"
 
 namespace UE::PoseSearch
 {
diff --git a/Engine/Plugins/Animation/PoseSearch/Source/Runtime/Private/PoseSearchInteractionSubsystem.cpp b/Engine/Plugins/Animation/PoseSearch/Source/Runtime/Private/PoseSearchInteractionSubsystem.cpp
--- a/Engine/Plugins/Animation/PoseSearch/Source/Runtime/Private/PoseSearchInteractionSubsystem.cpp
+++ b/Engine/Plugins/Animation/PoseSearch/Source/Runtime/Private/PoseSearchInteractionSubsystem.cpp
@@ -1,6 +1,8 @@
 // Copyright Epic Games, Inc. All Rights Reserved.
 
 #include "PoseSearch/PoseSearchInteractionSubsystem.h"
+#include "Components/PrimitiveComponent.h"
+#include "Engine/World.h"
 #include "PoseSearch/AnimNode_PoseSearchHistoryCollector.h"
 #include "PoseSearch/PoseSearchDatabase.h"
 #include "PoseSearch/PoseSearchInteractionIsland.h"
diff --git a/Engine/Plugins/Animation/PoseSearch/Source/Runtime/Private/PoseSearchInteractionValidator.cpp b/Engine/Plugins/Animation/PoseSearch/Source/Runtime/Private/PoseSearchInteractionValidator.cpp
--- a/Engine/Plugins/Animation/PoseSearch/Source/Runtime/Private/PoseSearchInteractionValidator.cpp
+++ b/Engine/Plugins/Animation/PoseSearch/Source/Runtime/Private/PoseSearchInteractionValidator.cpp
@@ -3,6 +3,7 @@
 #include "PoseSearch/PoseSearchInteractionValidator.h"
 #include "PoseSearch/PoseSearchContext.h"
 #include "PoseSearch/PoseSearchInteractionSubsystem.h"
+#include "GameFramework/Actor.h"
 
 namespace UE::PoseSearch
 {
diff --git a/Engine/Plugins/Animation/PoseSearch/Source/Runtime/Private/PoseSearchLibrary.cpp b/Engine/Plugins/Animation/PoseSearch/Source/Runtime/Private/PoseSearchLibrary.cpp
--- a/Engine/Plugins/Animation/PoseSearch/Source/Runtime/Private/PoseSearchLibrary.cpp
+++ b/Engine/Plugins/Animation/PoseSearch/Source/Runtime/Private/PoseSearchLibrary.cpp
@@ -1,19 +1,20 @@
 // Copyright Epic Games, Inc. All Rights Reserved.
 
 #include "PoseSearch/PoseSearchLibrary.h"
+
+#if UE_POSE_SEARCH_TRACE_ENABLED
+#include "ObjectTrace.h"
+#endif
+#include "Animation/AnimationAsset.h"
 #include "Animation/AnimInstanceProxy.h"
 #include "Animation/AnimMontage.h"
-#include "Animation/AnimNode_Inertialization.h"
 #include "Animation/AnimNode_SequencePlayer.h"
-#include "Animation/AnimRootMotionProvider.h"
 #include "Animation/AnimComposite.h"
 #include "Animation/AnimSequence.h"
 #include "Animation/AnimSubsystem_Tag.h"
 #include "Animation/BlendSpace.h"
-#include "Animation/BuiltInAttributeTypes.h"
 #include "Animation/AnimTrace.h"
 #include "GameFramework/Character.h"
-#include "StructUtils/InstancedStruct.h"
 #include "PoseSearch/AnimNode_MotionMatching.h"
 #include "PoseSearch/AnimNode_PoseSearchHistoryCollector.h"
 #include "PoseSearch/MultiAnimAsset.h"
@@ -60,7 +61,7 @@ namespace UE::PoseSearch
 		// no need to check if Database is indexing if found into AssetsToSearchPerDatabaseMap, since it already passed RequestAsyncBuildIndex successfully in a previous AddToSearchForDatabase call
 		if (!AssetsToSearch && (EAsyncBuildIndexResult::Success != FAsyncPoseSearchDatabasesManagement::RequestAsyncBuildIndex(Database, ERequestAsyncBuildFlag::ContinueRequest)))
 		{
-			// database is still indexing.. moving on
+			// database is still indexing... moving on
 			return true;
 		}
 #endif // WITH_EDITOR
@@ -96,7 +97,7 @@ namespace UE::PoseSearch
 
 	// this function is looking for UPoseSearchDatabase(s) to search for the input AssetToSearch:
 	// if AssetToSearch is a database search it ALL,
-	// if it's a sequence containing UAnimNotifyState_PoseSearchBranchIn, we add to the search of the dabase UAnimNotifyState_PoseSearchBranchIn::Database the asset AssetToSearch
+	// if it's a sequence containing UAnimNotifyState_PoseSearchBranchIn, we add to the search of the database UAnimNotifyState_PoseSearchBranchIn::Database the asset AssetToSearch
 	// returns bAsyncBuildIndexInProgress
 	static bool AddToSearch(FAssetsToSearchPerDatabaseMap& AssetsToSearchPerDatabaseMap, const UObject* AssetToSearch)
 	{
@@ -156,7 +157,7 @@ namespace UE::PoseSearch
 #endif // WITH_EDITOR
 		}
 		
-		// checking if any of the AssetsToSearch (databse) contains PlayingAnimationAsset
+		// checking if any of the AssetsToSearch (database) contains PlayingAnimationAsset
 		for (const UObject* AssetToSearch : AssetsToSearch)
 		{
 			if (const UPoseSearchDatabase* Database = Cast<UPoseSearchDatabase>(AssetToSearch))
diff --git a/Engine/Plugins/Animation/PoseSearch/Source/Runtime/Private/PoseSearchResult.cpp b/Engine/Plugins/Animation/PoseSearch/Source/Runtime/Private/PoseSearchResult.cpp
--- a/Engine/Plugins/Animation/PoseSearch/Source/Runtime/Private/PoseSearchResult.cpp
+++ b/Engine/Plugins/Animation/PoseSearch/Source/Runtime/Private/PoseSearchResult.cpp
@@ -2,10 +2,10 @@
 
 #include "PoseSearch/PoseSearchResult.h"
 #include "Animation/BlendSpace.h"
+#include "Animation/AnimationAsset.h"
 #include "PoseSearch/MultiAnimAsset.h"
 #include "PoseSearch/PoseSearchDatabase.h"
 #include "PoseSearch/PoseSearchSchema.h"
-#include "StructUtils/InstancedStruct.h"
 
 namespace UE::PoseSearch
 {
diff --git a/Engine/Plugins/Animation/PoseSearch/Source/Runtime/Private/PoseSearchTrajectoryLibrary.cpp b/Engine/Plugins/Animation/PoseSearch/Source/Runtime/Private/PoseSearchTrajectoryLibrary.cpp
--- a/Engine/Plugins/Animation/PoseSearch/Source/Runtime/Private/PoseSearchTrajectoryLibrary.cpp
+++ b/Engine/Plugins/Animation/PoseSearch/Source/Runtime/Private/PoseSearchTrajectoryLibrary.cpp
@@ -3,6 +3,7 @@
 #include "PoseSearch/PoseSearchTrajectoryLibrary.h"
 #include "PoseSearch/PoseSearchTrajectoryPredictor.h"
 #include "Animation/AnimInstanceProxy.h"
+#include "Animation/AnimInstance.h"
 #include "Components/SkeletalMeshComponent.h"
 #include "Engine/Engine.h"
 #include "GameFramework/Character.h"
diff --git a/Engine/Plugins/Animation/PoseSearch/Source/Runtime/Public/PoseSearch/Chooser/PoseSearchChooserColumn.h b/Engine/Plugins/Animation/PoseSearch/Source/Runtime/Public/PoseSearch/Chooser/PoseSearchChooserColumn.h
--- a/Engine/Plugins/Animation/PoseSearch/Source/Runtime/Public/PoseSearch/Chooser/PoseSearchChooserColumn.h
+++ b/Engine/Plugins/Animation/PoseSearch/Source/Runtime/Public/PoseSearch/Chooser/PoseSearchChooserColumn.h
@@ -15,6 +15,8 @@
 
 #define UE_API POSESEARCH_API
 
+class UAnimationAsset;
+
 USTRUCT(DisplayName = "Pose History Property Binding")
 struct FPoseHistoryContextProperty :  public FChooserParameterPoseHistoryBase
 {
diff --git a/Engine/Plugins/Animation/PoseSearch/Source/Runtime/Public/PoseSearch/PoseSearchAssetIndexer.h b/Engine/Plugins/Animation/PoseSearch/Source/Runtime/Public/PoseSearch/PoseSearchAssetIndexer.h
--- a/Engine/Plugins/Animation/PoseSearch/Source/Runtime/Public/PoseSearch/PoseSearchAssetIndexer.h
+++ b/Engine/Plugins/Animation/PoseSearch/Source/Runtime/Public/PoseSearch/PoseSearchAssetIndexer.h
@@ -4,11 +4,13 @@
 
 #if WITH_EDITOR
 
+#include "Animation/AnimCurveTypes.h"
 #include "BonePose.h"
 #include "PoseSearch/PoseSearchIndex.h"
 #include "PoseSearch/PoseSearchMirrorDataCache.h"
 #include "PoseSearch/PoseSearchSchema.h"
 
+struct FAnimNotifyContext;
 class UAnimNotifyState_PoseSearchBase;
 class UPoseSearchDatabase;
 
diff --git a/Engine/Plugins/Animation/PoseSearch/Source/Runtime/Public/PoseSearch/PoseSearchDefines.h b/Engine/Plugins/Animation/PoseSearch/Source/Runtime/Public/PoseSearch/PoseSearchDefines.h
--- a/Engine/Plugins/Animation/PoseSearch/Source/Runtime/Public/PoseSearch/PoseSearchDefines.h
+++ b/Engine/Plugins/Animation/PoseSearch/Source/Runtime/Public/PoseSearch/PoseSearchDefines.h
@@ -3,6 +3,7 @@
 #pragma once
 
 #include "Async/ParallelFor.h"
+#include "BoneIndices.h"
 #include "Logging/LogMacros.h"
 #include "Animation/AnimTrace.h"
 
diff --git a/Engine/Plugins/Animation/PoseSearch/Source/Runtime/Public/PoseSearch/PoseSearchHistory.h b/Engine/Plugins/Animation/PoseSearch/Source/Runtime/Public/PoseSearch/PoseSearchHistory.h
--- a/Engine/Plugins/Animation/PoseSearch/Source/Runtime/Public/PoseSearch/PoseSearchHistory.h
+++ b/Engine/Plugins/Animation/PoseSearch/Source/Runtime/Public/PoseSearch/PoseSearchHistory.h
@@ -2,9 +2,11 @@
 
 #pragma once
 
+#include "BoneIndices.h"
 #include "BonePose.h"
 #include "Containers/RingBuffer.h"
 #include "DrawDebugHelpers.h"
+#include "Animation/AnimCurveTypes.h"
 #include "Misc/MTTransactionallySafeAccessDetector.h"
 #include "PoseSearch/PoseHistoryProvider.h"
 #include "PoseSearch/PoseSearchDefines.h"
@@ -16,6 +18,7 @@
 
 #define UE_API POSESEARCH_API
 
+class UAnimationAsset;
 struct FAnimInstanceProxy;
 class USkeleton;
 class UWorld;
diff --git a/Engine/Plugins/Animation/PoseSearch/Source/Runtime/Public/PoseSearch/PoseSearchInteractionIsland.h b/Engine/Plugins/Animation/PoseSearch/Source/Runtime/Public/PoseSearch/PoseSearchInteractionIsland.h
--- a/Engine/Plugins/Animation/PoseSearch/Source/Runtime/Public/PoseSearch/PoseSearchInteractionIsland.h
+++ b/Engine/Plugins/Animation/PoseSearch/Source/Runtime/Public/PoseSearch/PoseSearchInteractionIsland.h
@@ -2,6 +2,7 @@
 
 #pragma once
 
+#include "Engine/EngineBaseTypes.h"
 #include "PoseSearch/PoseSearchInteractionAvailability.h"
 #include "PoseSearch/PoseSearchInteractionValidator.h"
 #include "PoseSearch/PoseSearchLibrary.h"
diff --git a/Engine/Plugins/Animation/PoseSearch/Source/Runtime/Public/PoseSearch/PoseSearchResult.h b/Engine/Plugins/Animation/PoseSearch/Source/Runtime/Public/PoseSearch/PoseSearchResult.h
--- a/Engine/Plugins/Animation/PoseSearch/Source/Runtime/Public/PoseSearch/PoseSearchResult.h
+++ b/Engine/Plugins/Animation/PoseSearch/Source/Runtime/Public/PoseSearch/PoseSearchResult.h
@@ -2,12 +2,14 @@
 
 #pragma once
 
+#include "Animation/AnimTypes.h"
 #include "PoseSearch/PoseSearchCost.h"
 #include "PoseSearch/PoseSearchRole.h"
 #include "PoseSearchResult.generated.h"
 
 #define UE_API POSESEARCH_API
 
+class UAnimationAsset;
 class UPoseSearchDatabase;
 class UPoseSearchSchema;
 struct FGameplayTag;
diff --git a/Engine/Plugins/Animation/PoseSearch/Source/Runtime/Public/PoseSearch/Trace/PoseSearchTraceLogger.h b/Engine/Plugins/Animation/PoseSearch/Source/Runtime/Public/PoseSearch/Trace/PoseSearchTraceLogger.h
--- a/Engine/Plugins/Animation/PoseSearch/Source/Runtime/Public/PoseSearch/Trace/PoseSearchTraceLogger.h
+++ b/Engine/Plugins/Animation/PoseSearch/Source/Runtime/Public/PoseSearch/Trace/PoseSearchTraceLogger.h
@@ -2,6 +2,10 @@
 
 #pragma once
 
+#include "ObjectTraceDefines.h"
+#if OBJECT_TRACE_ENABLED
+#include "ObjectTrace.h"
+#endif
 #include "PoseSearch/PoseSearchContext.h"
 #include "PoseSearch/PoseSearchLibrary.h"
 
diff --git a/Engine/Plugins/Experimental/AnimNext/Source/AnimNext/Private/Graph/RigUnit_IsRecentlyRendered.cpp b/Engine/Plugins/Experimental/AnimNext/Source/AnimNext/Private/Graph/RigUnit_IsRecentlyRendered.cpp
--- a/Engine/Plugins/Experimental/AnimNext/Source/AnimNext/Private/Graph/RigUnit_IsRecentlyRendered.cpp
+++ b/Engine/Plugins/Experimental/AnimNext/Source/AnimNext/Private/Graph/RigUnit_IsRecentlyRendered.cpp
@@ -1,6 +1,7 @@
 // Copyright Epic Games, Inc. All Rights Reserved.
 
 #include "RigUnit_IsRecentlyRendered.h"
+#include "Components/SkinnedMeshComponent.h"
 
 FRigUnit_IsRecentlyRendered_Execute()
 {
diff --git a/Engine/Plugins/Experimental/AnimNext/Source/AnimNext/Public/RewindDebugger/AnimNextTrace.h b/Engine/Plugins/Experimental/AnimNext/Source/AnimNext/Public/RewindDebugger/AnimNextTrace.h
--- a/Engine/Plugins/Experimental/AnimNext/Source/AnimNext/Public/RewindDebugger/AnimNextTrace.h
+++ b/Engine/Plugins/Experimental/AnimNext/Source/AnimNext/Public/RewindDebugger/AnimNextTrace.h
@@ -2,10 +2,15 @@
 
 #pragma once
 
-#include "CoreTypes.h"
-#include "ObjectTrace.h"
+#if UE_ENABLE_INCLUDE_ORDER_DEPRECATED_IN_5_7
 #include "CoreMinimal.h"
+#include "ObjectTrace.h"
 #include "StructUtils/PropertyBag.h"
+#endif // UE_ENABLE_INCLUDE_ORDER_DEPRECATED_IN_5_7
+
+#include "CoreTypes.h"
+#include "Misc/Guid.h"
+#include "ObjectTraceDefines.h"
 #include "Trace/Trace.h"
 
 #define ANIMNEXT_TRACE_ENABLED (OBJECT_TRACE_ENABLED && !(UE_BUILD_SHIPPING || UE_BUILD_TEST))
@@ -14,6 +19,7 @@
 
 UE_TRACE_CHANNEL_EXTERN(AnimNextChannel, ANIMNEXT_API);
 
+class UObject;
 struct FAnimNextModuleInstance;
 struct FAnimNextDataInterfaceInstance;
 struct FAnimNextGraphInstance;
diff --git a/Engine/Plugins/Experimental/AnimNextAnimGraph/Source/AnimNextAnimGraph/Private/Graph/TraceAnimNextGraphInstances.cpp b/Engine/Plugins/Experimental/AnimNextAnimGraph/Source/AnimNextAnimGraph/Private/Graph/TraceAnimNextGraphInstances.cpp
--- a/Engine/Plugins/Experimental/AnimNextAnimGraph/Source/AnimNextAnimGraph/Private/Graph/TraceAnimNextGraphInstances.cpp
+++ b/Engine/Plugins/Experimental/AnimNextAnimGraph/Source/AnimNextAnimGraph/Private/Graph/TraceAnimNextGraphInstances.cpp
@@ -1,15 +1,17 @@
 // Copyright Epic Games, Inc. All Rights Reserved.
 
 #include "TraceAnimNextGraphInstances.h"
+
+#include "EvaluationVM/EvaluationProgram.h"
+#include "EvaluationVM/SerializableEvaluationProgram.h"
 #include "Graph/AnimNextGraphInstance.h"
 #include "Module/AnimNextModuleInstance.h"
+#include "ObjectTrace.h"
+#include "Serialization/MemoryWriter.h"
+#include "Serialization/ObjectAndNameAsStringProxyArchive.h"
 #include "TraitCore/ExecutionContext.h"
 #include "TraitCore/NodeInstance.h"
 #include "TraitInterfaces/IHierarchy.h"
-#include "EvaluationVM/EvaluationProgram.h"
-#include "EvaluationVM/SerializableEvaluationProgram.h"
-#include "Serialization/MemoryWriter.h"
-#include "Serialization/ObjectAndNameAsStringProxyArchive.h"
 
 #if ANIMNEXT_TRACE_ENABLED
 
diff --git a/Engine/Plugins/Experimental/AnimNextAnimGraph/Source/AnimNextAnimGraph/Private/Traits/NotifyDispatcher.cpp b/Engine/Plugins/Experimental/AnimNextAnimGraph/Source/AnimNextAnimGraph/Private/Traits/NotifyDispatcher.cpp
--- a/Engine/Plugins/Experimental/AnimNextAnimGraph/Source/AnimNextAnimGraph/Private/Traits/NotifyDispatcher.cpp
+++ b/Engine/Plugins/Experimental/AnimNextAnimGraph/Source/AnimNextAnimGraph/Private/Traits/NotifyDispatcher.cpp
@@ -5,6 +5,7 @@
 #include "TraitInterfaces/ITimeline.h"
 #include "AnimationRuntime.h"
 #include "Animation/AnimTrace.h"
+#include "Animation/AnimNotifies/AnimNotify.h"
 #include "Animation/AnimNotifies/AnimNotifyState.h"
 #include "Components/SkeletalMeshComponent.h"
 #include "Graph/AnimNextGraphInstance.h"
diff --git a/Engine/Plugins/Experimental/AnimNextPoseSearch/Source/AnimNextPoseSearch/Private/MotionMatchingTrait.h b/Engine/Plugins/Experimental/AnimNextPoseSearch/Source/AnimNextPoseSearch/Private/MotionMatchingTrait.h
--- a/Engine/Plugins/Experimental/AnimNextPoseSearch/Source/AnimNextPoseSearch/Private/MotionMatchingTrait.h
+++ b/Engine/Plugins/Experimental/AnimNextPoseSearch/Source/AnimNextPoseSearch/Private/MotionMatchingTrait.h
@@ -3,6 +3,7 @@
 #pragma once
 
 #include "CoreMinimal.h"
+#include "AlphaBlend.h"
 #include "Module/ModuleHandle.h"
 #include "PoseSearch/PoseSearchInteractionLibrary.h"
 #include "PoseSearch/PoseSearchLibrary.h"
diff --git a/Engine/Plugins/Experimental/Animation/HierarchyTableAnimation/Source/Runtime/Private/AnimNode_BlendProfileLayeredBlend.cpp b/Engine/Plugins/Experimental/Animation/HierarchyTableAnimation/Source/Runtime/Private/AnimNode_BlendProfileLayeredBlend.cpp
--- a/Engine/Plugins/Experimental/Animation/HierarchyTableAnimation/Source/Runtime/Private/AnimNode_BlendProfileLayeredBlend.cpp
+++ b/Engine/Plugins/Experimental/Animation/HierarchyTableAnimation/Source/Runtime/Private/AnimNode_BlendProfileLayeredBlend.cpp
@@ -6,13 +6,11 @@
 #include "Animation/AnimTrace.h"
 #include "Animation/AnimCurveTypes.h"
 #include "Animation/AnimCurveUtils.h"
-#include "HierarchyTable.h"
-#include "MaskProfile/HierarchyTableTypeMask.h"
 #include "Animation/AttributeTypes.h"
 #include "Animation/IAttributeBlendOperator.h"
-#include "SkeletonHierarchyTableType.h"
 #include "HierarchyTableBlendProfile.h"
 #include "BlendProfileStandalone.h"
+#include "Animation/AnimInstance.h"
 
 #include UE_INLINE_GENERATED_CPP_BY_NAME(AnimNode_BlendProfileLayeredBlend)
 
diff --git a/Engine/Plugins/Experimental/EvaluationNotifies/Source/Runtime/Private/AnimNext/AnimNextAlignment.cpp b/Engine/Plugins/Experimental/EvaluationNotifies/Source/Runtime/Private/AnimNext/AnimNextAlignment.cpp
--- a/Engine/Plugins/Experimental/EvaluationNotifies/Source/Runtime/Private/AnimNext/AnimNextAlignment.cpp
+++ b/Engine/Plugins/Experimental/EvaluationNotifies/Source/Runtime/Private/AnimNext/AnimNextAlignment.cpp
@@ -7,6 +7,7 @@
 
 #include "Animation/AnimRootMotionProvider.h"
 #include "Animation/AnimSequence.h"
+#include "Engine/World.h"
 #include "EvaluationNotifies/AnimNotifyState_Alignment.h"
 #include "EvaluationVM/EvaluationVM.h"
 #include "EvaluationVM/KeyframeState.h"
diff --git a/Engine/Source/Runtime/AnimGraphRuntime/Private/AnimGraphRuntimeTrace.cpp b/Engine/Source/Runtime/AnimGraphRuntime/Private/AnimGraphRuntimeTrace.cpp
--- a/Engine/Source/Runtime/AnimGraphRuntime/Private/AnimGraphRuntimeTrace.cpp
+++ b/Engine/Source/Runtime/AnimGraphRuntime/Private/AnimGraphRuntimeTrace.cpp
@@ -4,6 +4,7 @@
 
 #if ANIM_TRACE_ENABLED
 
+#include "ObjectTrace.h"
 #include "Trace/Trace.inl"
 #include "Animation/BlendSpace.h"
 #include "AnimNodes/AnimNode_BlendSpacePlayer.h"
diff --git a/Engine/Source/Runtime/AnimGraphRuntime/Private/AnimNodes/AnimNode_AimOffsetLookAt.cpp b/Engine/Source/Runtime/AnimGraphRuntime/Private/AnimNodes/AnimNode_AimOffsetLookAt.cpp
--- a/Engine/Source/Runtime/AnimGraphRuntime/Private/AnimNodes/AnimNode_AimOffsetLookAt.cpp
+++ b/Engine/Source/Runtime/AnimGraphRuntime/Private/AnimNodes/AnimNode_AimOffsetLookAt.cpp
@@ -2,16 +2,14 @@
 
 #include "AnimNodes/AnimNode_AimOffsetLookAt.h"
 #include "Animation/AnimInstanceProxy.h"
+#include "Animation/AnimInstance.h"
 #include "Animation/AnimStats.h"
 #include "AnimationRuntime.h"
 #include "Animation/BlendSpace.h"
-#include "Animation/BlendSpace1D.h"
 #include "Components/SkeletalMeshComponent.h"
 #include "Engine/SkeletalMesh.h"
 #include "Engine/SkeletalMeshSocket.h"
 #include "DrawDebugHelpers.h"
-#include "EngineGlobals.h"
-#include "Engine/Engine.h"
 #include "Animation/AnimTrace.h"
 
 #include UE_INLINE_GENERATED_CPP_BY_NAME(AnimNode_AimOffsetLookAt)
diff --git a/Engine/Source/Runtime/AnimGraphRuntime/Private/AnimNodes/AnimNode_ApplyAdditive.cpp b/Engine/Source/Runtime/AnimGraphRuntime/Private/AnimNodes/AnimNode_ApplyAdditive.cpp
--- a/Engine/Source/Runtime/AnimGraphRuntime/Private/AnimNodes/AnimNode_ApplyAdditive.cpp
+++ b/Engine/Source/Runtime/AnimGraphRuntime/Private/AnimNodes/AnimNode_ApplyAdditive.cpp
@@ -4,6 +4,7 @@
 #include "AnimationRuntime.h"
 #include "Animation/AnimStats.h"
 #include "Animation/AnimInstanceProxy.h"
+#include "Animation/AnimInstance.h"
 #include "Animation/AnimTrace.h"
 
 #include UE_INLINE_GENERATED_CPP_BY_NAME(AnimNode_ApplyAdditive)
diff --git a/Engine/Source/Runtime/AnimGraphRuntime/Private/AnimNodes/AnimNode_CopyPoseFromMesh.cpp b/Engine/Source/Runtime/AnimGraphRuntime/Private/AnimNodes/AnimNode_CopyPoseFromMesh.cpp
--- a/Engine/Source/Runtime/AnimGraphRuntime/Private/AnimNodes/AnimNode_CopyPoseFromMesh.cpp
+++ b/Engine/Source/Runtime/AnimGraphRuntime/Private/AnimNodes/AnimNode_CopyPoseFromMesh.cpp
@@ -4,6 +4,7 @@
 
 #include "Animation/AnimCurveUtils.h"
 #include "Animation/AnimInstanceProxy.h"
+#include "Animation/AnimInstance.h"
 #include "Animation/AnimStats.h"
 #include "Animation/AnimTrace.h"
 #include "Components/SkeletalMeshComponent.h"
diff --git a/Engine/Source/Runtime/AnimGraphRuntime/Private/AnimNodes/AnimNode_CurveSource.cpp b/Engine/Source/Runtime/AnimGraphRuntime/Private/AnimNodes/AnimNode_CurveSource.cpp
--- a/Engine/Source/Runtime/AnimGraphRuntime/Private/AnimNodes/AnimNode_CurveSource.cpp
+++ b/Engine/Source/Runtime/AnimGraphRuntime/Private/AnimNodes/AnimNode_CurveSource.cpp
@@ -3,9 +3,11 @@
 #include "AnimNodes/AnimNode_CurveSource.h"
 #include "AnimationRuntime.h"
 #include "Animation/AnimCurveUtils.h"
-#include "Animation/AnimInstanceProxy.h"
+#include "Animation/AnimInstance.h"
 #include "Animation/AnimStats.h"
 #include "Animation/AnimTrace.h"
+#include "Components/ActorComponent.h"
+#include "GameFramework/Actor.h"
 
 #include UE_INLINE_GENERATED_CPP_BY_NAME(AnimNode_CurveSource)
 
diff --git a/Engine/Source/Runtime/AnimGraphRuntime/Private/AnimNodes/AnimNode_PoseSnapshot.cpp b/Engine/Source/Runtime/AnimGraphRuntime/Private/AnimNodes/AnimNode_PoseSnapshot.cpp
--- a/Engine/Source/Runtime/AnimGraphRuntime/Private/AnimNodes/AnimNode_PoseSnapshot.cpp
+++ b/Engine/Source/Runtime/AnimGraphRuntime/Private/AnimNodes/AnimNode_PoseSnapshot.cpp
@@ -2,6 +2,7 @@
 
 #include "AnimNodes/AnimNode_PoseSnapshot.h"
 #include "Animation/AnimInstanceProxy.h"
+#include "Animation/AnimInstance.h"
 #include "Animation/AnimTrace.h"
 #include "Components/SkeletalMeshComponent.h"
 #include "Engine/SkeletalMesh.h"
diff --git a/Engine/Source/Runtime/AnimGraphRuntime/Private/AnimNodes/AnimNode_RotationOffsetBlendSpace.cpp b/Engine/Source/Runtime/AnimGraphRuntime/Private/AnimNodes/AnimNode_RotationOffsetBlendSpace.cpp
--- a/Engine/Source/Runtime/AnimGraphRuntime/Private/AnimNodes/AnimNode_RotationOffsetBlendSpace.cpp
+++ b/Engine/Source/Runtime/AnimGraphRuntime/Private/AnimNodes/AnimNode_RotationOffsetBlendSpace.cpp
@@ -2,6 +2,7 @@
 
 #include "AnimNodes/AnimNode_RotationOffsetBlendSpace.h"
 #include "Animation/AnimInstanceProxy.h"
+#include "Animation/AnimInstance.h"
 #include "AnimationRuntime.h"
 #include "Animation/AnimTrace.h"
 
diff --git a/Engine/Source/Runtime/AnimGraphRuntime/Private/AnimNodes/AnimNode_Slot.cpp b/Engine/Source/Runtime/AnimGraphRuntime/Private/AnimNodes/AnimNode_Slot.cpp
--- a/Engine/Source/Runtime/AnimGraphRuntime/Private/AnimNodes/AnimNode_Slot.cpp
+++ b/Engine/Source/Runtime/AnimGraphRuntime/Private/AnimNodes/AnimNode_Slot.cpp
@@ -2,6 +2,7 @@
 
 #include "AnimNodes/AnimNode_Slot.h"
 #include "Animation/AnimInstanceProxy.h"
+#include "Animation/AnimInstance.h"
 #include "Animation/AnimMontage.h"
 #include "Animation/AnimStats.h"
 #include "Animation/AnimTrace.h"
diff --git a/Engine/Source/Runtime/AnimGraphRuntime/Private/AnimNodes/AnimNode_TwoWayBlend.cpp b/Engine/Source/Runtime/AnimGraphRuntime/Private/AnimNodes/AnimNode_TwoWayBlend.cpp
--- a/Engine/Source/Runtime/AnimGraphRuntime/Private/AnimNodes/AnimNode_TwoWayBlend.cpp
+++ b/Engine/Source/Runtime/AnimGraphRuntime/Private/AnimNodes/AnimNode_TwoWayBlend.cpp
@@ -1,6 +1,7 @@
 // Copyright Epic Games, Inc. All Rights Reserved.
 
 #include "AnimNodes/AnimNode_TwoWayBlend.h"
+#include "Animation/AnimInstance.h"
 #include "Animation/AnimInstanceProxy.h"
 #include "Animation/AnimStats.h"
 #include "Animation/AnimTrace.h"
diff --git a/Engine/Source/Runtime/AnimGraphRuntime/Private/BoneControllers/AnimNode_HandIKRetargeting.cpp b/Engine/Source/Runtime/AnimGraphRuntime/Private/BoneControllers/AnimNode_HandIKRetargeting.cpp
--- a/Engine/Source/Runtime/AnimGraphRuntime/Private/BoneControllers/AnimNode_HandIKRetargeting.cpp
+++ b/Engine/Source/Runtime/AnimGraphRuntime/Private/BoneControllers/AnimNode_HandIKRetargeting.cpp
@@ -3,6 +3,9 @@
 #include "BoneControllers/AnimNode_HandIKRetargeting.h"
 #include "Animation/AnimStats.h"
 #include "Animation/AnimTrace.h"
+#if ENABLE_ANIM_DEBUG
+#include "HAL/IConsoleManager.h"
+#endif
 
 #include UE_INLINE_GENERATED_CPP_BY_NAME(AnimNode_HandIKRetargeting)
 
diff --git a/Engine/Source/Runtime/AnimGraphRuntime/Private/BoneControllers/AnimNode_SplineIK.cpp b/Engine/Source/Runtime/AnimGraphRuntime/Private/BoneControllers/AnimNode_SplineIK.cpp
--- a/Engine/Source/Runtime/AnimGraphRuntime/Private/BoneControllers/AnimNode_SplineIK.cpp
+++ b/Engine/Source/Runtime/AnimGraphRuntime/Private/BoneControllers/AnimNode_SplineIK.cpp
@@ -1,9 +1,9 @@
 // Copyright Epic Games, Inc. All Rights Reserved.
 
 #include "BoneControllers/AnimNode_SplineIK.h"
-#include "Animation/AnimTypes.h"
 #include "AnimationRuntime.h"
 #include "Animation/AnimInstanceProxy.h"
+#include "Animation/AnimInstance.h"
 #include "Animation/AnimStats.h"
 #include "Components/SkeletalMeshComponent.h"
 #include "Engine/SkeletalMesh.h"
@@ -297,8 +297,6 @@ void FAnimNode_SplineIK::BuildBoneSpline(const FReferenceSkeleton& RefSkeleton)
 {
 	if (CachedBoneReferences.Num() > 0)
 	{
-		const TArray<FTransform>& RefBonePose = RefSkeleton.GetRefBonePose();
-
 		TArray<FTransform> ComponentSpaceTransforms;
 		FAnimationRuntime::FillUpComponentSpaceTransforms(RefSkeleton, RefSkeleton.GetRefBonePose(), ComponentSpaceTransforms);
 
diff --git a/Engine/Source/Runtime/AnimGraphRuntime/Private/BoneControllers/AnimNode_SpringBone.cpp b/Engine/Source/Runtime/AnimGraphRuntime/Private/BoneControllers/AnimNode_SpringBone.cpp
--- a/Engine/Source/Runtime/AnimGraphRuntime/Private/BoneControllers/AnimNode_SpringBone.cpp
+++ b/Engine/Source/Runtime/AnimGraphRuntime/Private/BoneControllers/AnimNode_SpringBone.cpp
@@ -4,6 +4,7 @@
 #include "Components/SkeletalMeshComponent.h"
 #include "GameFramework/WorldSettings.h"
 #include "Animation/AnimInstanceProxy.h"
+#include "Animation/AnimInstance.h"
 #include "Animation/AnimStats.h"
 #include "Animation/AnimTrace.h"
 
diff --git a/Engine/Source/Runtime/Engine/Private/Animation/AnimInstance.cpp b/Engine/Source/Runtime/Engine/Private/Animation/AnimInstance.cpp
--- a/Engine/Source/Runtime/Engine/Private/Animation/AnimInstance.cpp
+++ b/Engine/Source/Runtime/Engine/Private/Animation/AnimInstance.cpp
@@ -28,6 +28,7 @@
 #include "Animation/ActiveMontageInstanceScope.h"
 #include "Animation/AnimNode_SaveCachedPose.h"
 #include "Animation/AnimSubsystem_SharedLinkedAnimLayers.h"
+#include "ObjectTrace.h"
 #if WITH_EDITOR
 #include "Engine/Blueprint.h"
 #include "BlueprintEditorSettings.h"
diff --git a/Engine/Source/Runtime/Engine/Private/Animation/AnimMontage.cpp b/Engine/Source/Runtime/Engine/Private/Animation/AnimMontage.cpp
--- a/Engine/Source/Runtime/Engine/Private/Animation/AnimMontage.cpp
+++ b/Engine/Source/Runtime/Engine/Private/Animation/AnimMontage.cpp
@@ -12,6 +12,7 @@
 #include "Animation/AnimData/IAnimationDataController.h"
 #include "UObject/UObjectThreadContext.h"
 #include "Animation/AssetMappingTable.h"
+#include "Animation/AnimInstance.h"
 #include "Animation/AnimSequence.h"
 #include "Animation/BlendProfile.h"
 #include "AnimationUtils.h"
diff --git a/Engine/Source/Runtime/Engine/Private/Animation/AnimSync.cpp b/Engine/Source/Runtime/Engine/Private/Animation/AnimSync.cpp
--- a/Engine/Source/Runtime/Engine/Private/Animation/AnimSync.cpp
+++ b/Engine/Source/Runtime/Engine/Private/Animation/AnimSync.cpp
@@ -2,6 +2,7 @@
 
 #include "Animation/AnimSync.h"
 #include "Animation/AnimInstanceProxy.h"
+#include "Animation/AnimInstance.h"
 #include "Animation/AnimStats.h"
 #include "Animation/AnimTrace.h"
 #include "AnimationUtils.h"
diff --git a/Engine/Source/Runtime/Engine/Private/Animation/AnimTrace.cpp b/Engine/Source/Runtime/Engine/Private/Animation/AnimTrace.cpp
--- a/Engine/Source/Runtime/Engine/Private/Animation/AnimTrace.cpp
+++ b/Engine/Source/Runtime/Engine/Private/Animation/AnimTrace.cpp
@@ -4,6 +4,8 @@
 
 #if ANIM_TRACE_ENABLED
 
+#include "ObjectTrace.h"
+#include "TraceFilter.h"
 #include "Animation/AnimInstanceProxy.h"
 #include "Components/SkeletalMeshComponent.h"
 #include "Components/ExternalMorphSet.h"
diff --git a/Engine/Source/Runtime/Engine/Private/ObjectPropertyTrace.cpp b/Engine/Source/Runtime/Engine/Private/ObjectPropertyTrace.cpp
--- a/Engine/Source/Runtime/Engine/Private/ObjectPropertyTrace.cpp
+++ b/Engine/Source/Runtime/Engine/Private/ObjectPropertyTrace.cpp
@@ -1,9 +1,9 @@
 // Copyright Epic Games, Inc. All Rights Reserved.
 
 #include "ObjectPropertyTrace.h"
-#include "Trace/Trace.h"
-#include "Trace/Trace.inl"
+#include "ObjectTrace.h"
 #include "Containers/Ticker.h"
+#include "UObject/UnrealType.h"
 
 #if WITH_ENGINE
 
@@ -256,9 +256,7 @@ namespace ObjectPropertyTrace
 					<< PropertiesStart2.ObjectId(ObjectId)
 					<< PropertiesStart2.RecordingTime(FObjectTrace::GetWorldElapsedTime(TracedObject->GetWorld()));
 
-				uint64 ClassId = FObjectTrace::GetObjectId(TracedObject->GetClass());
-
-				IterateProperties(TracedObject->GetClass(), TracedObject, [StartCycle, ClassId, ObjectId](const FString& InTypeString, const FString& InValue, int32 InId, int32 InParentId, FStringBuilderBase & InPropertyNameBuilder)
+				IterateProperties(TracedObject->GetClass(), TracedObject, [StartCycle, ObjectId](const FString& InTypeString, const FString& InValue, int32 InId, int32 InParentId, FStringBuilderBase & InPropertyNameBuilder)
 				{
 					const FString NameString = InPropertyNameBuilder.ToString();
 
diff --git a/Engine/Source/Runtime/Engine/Public/Animation/AnimAttributes.h b/Engine/Source/Runtime/Engine/Public/Animation/AnimAttributes.h
--- a/Engine/Source/Runtime/Engine/Public/Animation/AnimAttributes.h
+++ b/Engine/Source/Runtime/Engine/Public/Animation/AnimAttributes.h
@@ -2,9 +2,13 @@
 
 #pragma once
 
+#if UE_ENABLE_INCLUDE_ORDER_DEPRECATED_IN_5_7
 #include "CoreMinimal.h"
+#endif // UE_ENABLE_INCLUDE_ORDER_DEPRECATED_IN_5_7
 
-namespace UE { namespace Anim 
+#include "UObject/NameTypes.h"
+
+namespace UE::Anim
 {
 
 // Built-in attributes that most nodes will share
@@ -15,4 +19,4 @@ struct FAttributes
 	static ENGINE_API const FName Attributes;
 };
 
-}}	// namespace UE::Anim
+}	// namespace UE::Anim
diff --git a/Engine/Source/Runtime/Engine/Public/Animation/AnimTrace.h b/Engine/Source/Runtime/Engine/Public/Animation/AnimTrace.h
--- a/Engine/Source/Runtime/Engine/Public/Animation/AnimTrace.h
+++ b/Engine/Source/Runtime/Engine/Public/Animation/AnimTrace.h
@@ -2,39 +2,44 @@
 
 #pragma once
 
+#if UE_ENABLE_INCLUDE_ORDER_DEPRECATED_IN_5_7
+#include "CoreMinimal.h"
 #include "CoreTypes.h"
 #include "ObjectTrace.h"
-#include "CoreMinimal.h"
-#include "Trace/Trace.h"
+#endif // UE_ENABLE_INCLUDE_ORDER_DEPRECATED_IN_5_7
+
+#include "ObjectTraceDefines.h"
 
 #define ANIM_TRACE_ENABLED (OBJECT_TRACE_ENABLED && !(UE_BUILD_SHIPPING || UE_BUILD_TEST))
 
 #if ANIM_TRACE_ENABLED
 
-#include "AnimAttributes.h"
+#include "Animation/AnimAttributes.h"
+#include "BoneIndices.h"
+#include "Math/Color.h"
+#include "Math/Vector.h"
 
 UE_TRACE_CHANNEL_EXTERN(AnimationChannel, ENGINE_API);
 
-struct FAnimInstanceProxy;
-struct FAnimTickRecord;
-struct FAnimationBaseContext;
-struct FAnimationUpdateContext;
+class FAutoConsoleVariable;
 class UAnimInstance;
+class UPoseWatchPoseElement;
 class USkeletalMesh;
 class USkeletalMeshComponent;
-struct FAnimationInitializeContext;
-struct FAnimationUpdateContext;
 struct FAnimationBaseContext;
 struct FAnimationCacheBonesContext;
-struct FPoseContext;
-struct FComponentSpacePoseContext;
-class FName;
+struct FAnimationInitializeContext;
+struct FAnimationUpdateContext;
+struct FAnimInstanceProxy;
+struct FAnimMontageInstance;
 struct FAnimNode_SequencePlayerBase;
 struct FAnimNotifyEvent;
-struct FPassedMarker;
 struct FAnimSyncMarker;
-struct FAnimMontageInstance;
-class UPoseWatchPoseElement;
+struct FAnimTickRecord;
+struct FBlendedHeapCurve;
+struct FComponentSpacePoseContext;
+struct FPassedMarker;
+struct FPoseContext;
 
 extern ENGINE_API FAutoConsoleVariable CVarRecordExternalMorphTargets;
 
diff --git a/Engine/Source/Runtime/Engine/Public/ObjectPropertyTrace.h b/Engine/Source/Runtime/Engine/Public/ObjectPropertyTrace.h
--- a/Engine/Source/Runtime/Engine/Public/ObjectPropertyTrace.h
+++ b/Engine/Source/Runtime/Engine/Public/ObjectPropertyTrace.h
@@ -2,12 +2,12 @@
 
 #pragma once
 
+#if UE_ENABLE_INCLUDE_ORDER_DEPRECATED_IN_5_7
 #include "CoreTypes.h"
 #include "CoreMinimal.h"
+#endif // UE_ENABLE_INCLUDE_ORDER_DEPRECATED_IN_5_7
 
-#if WITH_ENGINE
-#include "ObjectTrace.h"
-#endif
+#include "ObjectTraceDefines.h"
 
 #if WITH_ENGINE
 #define OBJECT_PROPERTY_TRACE_ENABLED OBJECT_TRACE_ENABLED
@@ -17,6 +17,8 @@
 
 #if OBJECT_PROPERTY_TRACE_ENABLED
 
+class UObject;
+
 struct FObjectPropertyTrace
 {
 	/** Initialize object property tracing */
diff --git a/Engine/Source/Runtime/Engine/Public/ObjectTrace.h b/Engine/Source/Runtime/Engine/Public/ObjectTrace.h
--- a/Engine/Source/Runtime/Engine/Public/ObjectTrace.h
+++ b/Engine/Source/Runtime/Engine/Public/ObjectTrace.h
@@ -3,18 +3,13 @@
 #pragma once
 
 #include "CoreTypes.h"
-#include "Trace/Config.h"
 #include "Subsystems/WorldSubsystem.h"
+#include "ObjectTraceDefines.h"
+#include "Trace/Config.h"
 #include "TraceFilter.h"
 
 #include "ObjectTrace.generated.h"
 
-#if UE_TRACE_ENABLED && !IS_PROGRAM && !UE_BUILD_SHIPPING
-#define OBJECT_TRACE_ENABLED 1
-#else
-#define OBJECT_TRACE_ENABLED 0
-#endif
-
 // World subsystem used to track world info
 UCLASS()
 class UObjectTraceWorldSubsystem : public UWorldSubsystem
diff --git a/Engine/Source/Runtime/Engine/Public/ObjectTraceDefines.h b/Engine/Source/Runtime/Engine/Public/ObjectTraceDefines.h
new file mode 100644
--- /dev/null
+++ b/Engine/Source/Runtime/Engine/Public/ObjectTraceDefines.h
@@ -0,0 +1,15 @@
+// Copyright Epic Games, Inc. All Rights Reserved.
+
+#pragma once
+
+#include "Trace/Config.h"
+
+#ifndef OBJECT_TRACE_ENABLED
+
+#if UE_TRACE_ENABLED && !IS_PROGRAM && !UE_BUILD_SHIPPING
+#define OBJECT_TRACE_ENABLED 1
+#else
+#define OBJECT_TRACE_ENABLED 0
+#endif
+
+#endif
diff --git a/Engine/Source/Runtime/LiveLinkAnimationCore/Private/AnimNode_LiveLinkPose.cpp b/Engine/Source/Runtime/LiveLinkAnimationCore/Private/AnimNode_LiveLinkPose.cpp
--- a/Engine/Source/Runtime/LiveLinkAnimationCore/Private/AnimNode_LiveLinkPose.cpp
+++ b/Engine/Source/Runtime/LiveLinkAnimationCore/Private/AnimNode_LiveLinkPose.cpp
@@ -3,6 +3,7 @@
 #include "AnimNode_LiveLinkPose.h"
 
 #include "Animation/AnimInstanceProxy.h"
+#include "Animation/AnimInstance.h"
 #include "Animation/AnimTrace.h"
 #include "Features/IModularFeatures.h"
 #include "ILiveLinkClient.h"
