PoseSearch - making UPoseSearchInteractionSubsystem GC compliant

diff --git a/Engine/Plugins/Animation/PoseSearch/Source/Runtime/Private/PoseSearchInteractionIsland.cpp b/Engine/Plugins/Animation/PoseSearch/Source/Runtime/Private/PoseSearchInteractionIsland.cpp
--- a/Engine/Plugins/Animation/PoseSearch/Source/Runtime/Private/PoseSearchInteractionIsland.cpp
+++ b/Engine/Plugins/Animation/PoseSearch/Source/Runtime/Private/PoseSearchInteractionIsland.cpp
@@ -926,10 +926,9 @@ void FInteractionIsland::LogTickDependencies() const
 
 	for (const TWeakObjectPtr<UActorComponent>& TickActorComponentPtr : TickActorComponents)
 	{
+		const int32 InteractionIslandIndex = InteractionSubsystem->GetInteractionIslands().IndexOfByKey(this);
 		if (UActorComponent* TickActorComponent = TickActorComponentPtr.Get())
 		{
-			const int32 InteractionIslandIndex = InteractionSubsystem->GetInteractionIslands().IndexOfByKey(this);
-
 			UE_LOG(LogPoseSearch, Log, TEXT("============== %s (Island %d) =============="), *TickActorComponent->GetOwner()->GetName(), InteractionIslandIndex);
 			if (IInteractionIslandDependency* IslandDependency = FindCustomDependency(TickActorComponent))
 			{
@@ -949,6 +948,10 @@ void FInteractionIsland::LogTickDependencies() const
 				LogTickFunction(TickActorComponent->PrimaryComponentTick, Type, true, 1);
 			}
 		}
+		else
+		{
+			UE_LOG(LogPoseSearch, Log, TEXT("============== !!!Missing Actor!!! (Island %d) =============="), InteractionIslandIndex);
+		}
 	}
 
 }
diff --git a/Engine/Plugins/Animation/PoseSearch/Source/Runtime/Private/PoseSearchInteractionSubsystem.cpp b/Engine/Plugins/Animation/PoseSearch/Source/Runtime/Private/PoseSearchInteractionSubsystem.cpp
--- a/Engine/Plugins/Animation/PoseSearch/Source/Runtime/Private/PoseSearchInteractionSubsystem.cpp
+++ b/Engine/Plugins/Animation/PoseSearch/Source/Runtime/Private/PoseSearchInteractionSubsystem.cpp
@@ -24,36 +24,15 @@ static bool GVarPoseSearchInteractionLoglandsTickDependencies = false;
 static FAutoConsoleVariableRef CVarPoseSearchInteractionLoglandsTickDependencies(TEXT("a.PoseSearchInteraction.LoglandsTickDependencies"), GVarPoseSearchInteractionLoglandsTickDependencies, TEXT("Log islands tick dependencies"));
 #endif // !NO_CVARS
 
-FString FInteractionAvailability::GetPoseHistoryName() const
+struct FAnimContextInfo
 {
-	if (PoseHistory)
+	void Init(const FPoseSearchInteractionAnimContextAvailabilities& InAnimContextAvailabilities)
 	{
-		return "HistoryProvider";
+		check(InAnimContextAvailabilities.AnimContext && !InAnimContextAvailabilities.Availabilities.IsEmpty());
+		AnimContextAvailabilities = &InAnimContextAvailabilities;
+		Location = GetContextLocation(InAnimContextAvailabilities.AnimContext);
 	}
-	return PoseHistoryName.ToString();
-}
 
-const IPoseHistory* FInteractionAvailability::GetPoseHistory(const UObject* AnimContext) const
-{
-	if (PoseHistory)
-	{
-		return PoseHistory;
-	}
-		
-	if (const UAnimInstance* AnimInstance = Cast<UAnimInstance>(AnimContext))
-	{
-		if (const FAnimNode_PoseSearchHistoryCollector_Base* PoseSearchHistoryCollector = UPoseSearchLibrary::FindPoseHistoryNode(PoseHistoryName, AnimInstance))
-		{
-			return &PoseSearchHistoryCollector->GetPoseHistory();
-		}
-	}
-
-	unimplemented();
-	return nullptr;
-}
-
-struct FAnimContextInfo
-{
 	// performs broad phase analysis checking if at least one of the Availabilities associated to AnimContext can interact with OtherAnimContextInfo.
 	// This is a more relaxed analysis than the one performed in FRoledAnimContextInfo::CanInteractWith
 	bool CanInteractWith(const FAnimContextInfo& OtherAnimContextInfo) const
@@ -61,12 +40,13 @@ struct FAnimContextInfo
 		check(this != &OtherAnimContextInfo);
 		
 		// @todo: enable this code if we ended up requiring preventing interactions between the same actor!
-		// const AActor* AnimContextActor = GetContextOwningActor(AnimContext);
-		// const AActor* OtherAnimContextActor = GetContextOwningActor(OtherAnimContextInfo.AnimContext);
-		// if (AnimContextActor == OtherAnimContextActor)
-		// {
-		// 	return false;
-		// }
+		//check(AnimContextAvailabilities && OtherAnimContextInfo.AnimContextAvailabilities);
+		//const AActor* AnimContextActor = GetContextOwningActor(AnimContextAvailabilities->AnimContext);
+		//const AActor* OtherAnimContextActor = GetContextOwningActor(OtherAnimContextInfo.AnimContextAvailabilities->AnimContext);
+		//if (AnimContextActor == OtherAnimContextActor)
+		//{
+		//	return false;
+		//}
 
 		const FVector DeltaLocation = Location - OtherAnimContextInfo.Location;
 		const float DistanceSquared = DeltaLocation.SquaredLength();
@@ -78,8 +58,8 @@ struct FAnimContextInfo
 	float GetAvailabilitiesMaxBroadPhaseRadius() const
 	{
 		float AvailabilitiesMaxBroadPhaseRadius = 0.f;
-		check(Availabilities);
-		for (const FInteractionAvailability& Availability : *Availabilities)
+		check(AnimContextAvailabilities);
+		for (const FPoseSearchInteractionAvailabilityEx& Availability : AnimContextAvailabilities->Availabilities)
 		{
 			// @todo: optimize the AvailabilitiesMaxBroadPhaseRadius, since adding Availability.BroadPhaseRadiusIncrementOnInteraction is required ONLY if AnimContext is already part of an interaction
 			AvailabilitiesMaxBroadPhaseRadius = FMath::Max(AvailabilitiesMaxBroadPhaseRadius, Availability.BroadPhaseRadius + Availability.BroadPhaseRadiusIncrementOnInteraction);
@@ -87,9 +67,8 @@ struct FAnimContextInfo
 		return AvailabilitiesMaxBroadPhaseRadius;
 	}
 
-	const UObject* AnimContext = nullptr;
-	// array of all valid avalilabilities associated to AnimContext
-	const FInteractionAvailabilities* Availabilities = nullptr;
+	const FPoseSearchInteractionAnimContextAvailabilities* AnimContextAvailabilities = nullptr;
+
 	// cached AnimContext location
 	FVector Location = FVector::ZeroVector;
 	TArray<const FAnimContextInfo*, TInlineAllocator<8, TMemStackAllocator<>>> NearbyAnimContextInfos;
@@ -101,6 +80,38 @@ struct FTagToDatabases : TMap<FName, FDatabasesPerTag, TInlineSetAllocator<8, TM
 
 } // namespace UE::PoseSearch
 
+// FPoseSearchInteractionAvailabilityEx
+///////////////////////////////////////////////
+FString FPoseSearchInteractionAvailabilityEx::GetPoseHistoryName() const
+{
+	if (PoseHistory)
+	{
+		return "HistoryProvider";
+	}
+	return PoseHistoryName.ToString();
+}
+
+const UE::PoseSearch::IPoseHistory* FPoseSearchInteractionAvailabilityEx::GetPoseHistory(const UObject* AnimContext) const
+{
+	if (PoseHistory)
+	{
+		return PoseHistory;
+	}
+		
+	if (const UAnimInstance* AnimInstance = Cast<UAnimInstance>(AnimContext))
+	{
+		if (const FAnimNode_PoseSearchHistoryCollector_Base* PoseSearchHistoryCollector = UPoseSearchLibrary::FindPoseHistoryNode(PoseHistoryName, AnimInstance))
+		{
+			return &PoseSearchHistoryCollector->GetPoseHistory();
+		}
+	}
+
+	unimplemented();
+	return nullptr;
+}
+
+// UPoseSearchInteractionSubsystem
+///////////////////////////////////////////////
 UE::PoseSearch::FInteractionIsland& UPoseSearchInteractionSubsystem::CreateIsland()
 {
 	return *Islands.Add_GetRef(new UE::PoseSearch::FInteractionIsland(ToRawPtr(GetWorld()->PersistentLevel), this));
@@ -143,7 +154,7 @@ void UPoseSearchInteractionSubsystem::RegenerateAllIslands(float DeltaSeconds)
 
 	check(IsInGameThread());
 
-	// FScopeLock Lock(&AvailabilitiesMapMutex); is not necessary since UPoseSearchInteractionSubsystem gets ticked outside the parallel animation jobs
+	// FScopeLock Lock(&AnimContextsAvailabilitiesMutex); is not necessary since UPoseSearchInteractionSubsystem gets ticked outside the parallel animation jobs
 
 	// generating all the possible interaction tuples of AnimContext(s) with roles and pose histories (defined in FInteractionSearchContext)
 	FInteractionSearchContexts SearchContexts;
@@ -493,14 +504,16 @@ void UPoseSearchInteractionSubsystem::AddAvailabilities(const TArrayView<const F
 {
 	using namespace UE::PoseSearch;
 
+	check(AnimContext && AnimContext->GetWorld() && AnimContext->GetWorld() == GetWorld());
+
 	// collecting valid availabilities indexes here to minimize the time spent on the lock and avoid locking at all 
-	// if there're actually no valid availabilities (as well as avoiding polluting AvailabilitiesMap with an empty entry).
+	// if there're actually no valid availabilities (as well as avoiding polluting AnimContextsAvailabilities with an empty entry).
 	// This is unfortunately a common setup, and that's the main reason we optimize against it
 	TArray<int32, TInlineAllocator<64>> ValidAvailabilityIndexes;
 	for (int32 AvailabilityIndex = 0; AvailabilityIndex < Availabilities.Num(); ++AvailabilityIndex)
 	{
 		const FPoseSearchInteractionAvailability& Availability = Availabilities[AvailabilityIndex];
-		if (Availability.Database || Availability.IsTagValid())
+		if ((Availability.Database && Availability.Database->Schema) || Availability.IsTagValid())
 		{
 			ValidAvailabilityIndexes.Add(AvailabilityIndex);
 		}
@@ -508,15 +521,40 @@ void UPoseSearchInteractionSubsystem::AddAvailabilities(const TArrayView<const F
 
 	if (!ValidAvailabilityIndexes.IsEmpty())
 	{
-		FScopeLock Lock(&AvailabilitiesMapMutex);
-		FInteractionAvailabilities& InteractionAvailabilities = AvailabilitiesMap.FindOrAdd(AnimContext);
+		FScopeLock Lock(&AnimContextsAvailabilitiesMutex);
+
+		const int32 UpperBoundIndex = Algo::UpperBoundBy(AnimContextsAvailabilities, AnimContext,
+			[](const FPoseSearchInteractionAnimContextAvailabilities& AnimContextAvailabilities)
+			{
+				return AnimContextAvailabilities.AnimContext.Get();
+			},
+			[](const UObject* AnimContextA, const UObject* AnimContextB)
+			{
+				return AnimContextA < AnimContextB;
+			});
+
+		int32 AnimContextAvailabilityIndex;
+		if (UpperBoundIndex > 0 && AnimContextsAvailabilities[UpperBoundIndex - 1].AnimContext == AnimContext)
+		{
+			// we found the AnimContextsAvailabilities for AnimContext
+			AnimContextAvailabilityIndex = UpperBoundIndex - 1;
+		}
+		else
+		{
+			// we create a new AnimContextsAvailabilities for AnimContext, preserving AnimContextsAvailabilities sorting
+			AnimContextsAvailabilities.InsertDefaulted_GetRef(UpperBoundIndex).AnimContext = AnimContext;
+			AnimContextAvailabilityIndex = UpperBoundIndex;
+		}
+
 		for (const int32 AvailabilityIndex : ValidAvailabilityIndexes)
 		{
-			// Avoiding adding trivial duplicates. FInteractionAvailability could not be fully specified to understand if it's an actual duplicate in case the pose
+			// Avoiding adding trivial duplicates. FPoseSearchInteractionAvailabilityEx could not be fully specified to understand if it's an actual duplicate in case the pose
 			// history is passed by name or the Availability.Database is null and supposed to be resolved using other availabilities Database(s) with the same Availability.Tag.
 			// The duplicated availabilities are excluded when creating the combinations of possible interactions during FAnimContextInfoVisitor when 
 			// FRoledAnimContextInfos.AddRoledAnimContextInfos calls FRoledAnimContextInfos.AddUnique
-			InteractionAvailabilities.AddUnique(FInteractionAvailability(Availabilities[AvailabilityIndex], PoseHistoryName, PoseHistory));
+			FPoseSearchInteractionAvailabilityEx AvailabilityEx;
+			AvailabilityEx.Init(Availabilities[AvailabilityIndex], PoseHistoryName, PoseHistory);
+			AnimContextsAvailabilities[AnimContextAvailabilityIndex].Availabilities.AddUnique(AvailabilityEx);
 		}
 	}
 }
@@ -531,90 +569,36 @@ void UPoseSearchInteractionSubsystem::GenerateAnimContextInfosAndTagToDatabases(
 	check(SubsystemWorld);
 	check(AnimContextInfos.IsEmpty() && TagToDatabases.IsEmpty());
 
-	// collecting and sorting all the AnimContext(s) from the AvailabilitiesMap in AnimContextsWithAvailabilities, to have a stable and deterministic island creation
-	typedef TPair<const UObject*, const FInteractionAvailabilities*> FAnimContextWithAvailabilities;
-	TArray<FAnimContextWithAvailabilities, TInlineAllocator<512, TMemStackAllocator<>>> AnimContextsWithAvailabilities;
-	for (FAvailabilitiesMap::TConstIterator Iter = AvailabilitiesMap.CreateConstIterator(); Iter; ++Iter)
-	{
-		// making sure the weak pointer is still valid
-		const UObject* AnimContext = Iter->Key.ResolveObjectPtr();
-		if (!AnimContext)
-		{
-			UE_LOG(LogPoseSearch, Log, TEXT("UPoseSearchInteractionSubsystem::GenerateAnimContextInfosAndTagToDatabaseMap null anim context. The associated character got removed from the world"));
-		}
-		else if (!AnimContext->GetWorld())
-		{
-			UE_LOG(LogPoseSearch, Error, TEXT("UPoseSearchInteractionSubsystem::GenerateAnimContextInfosAndTagToDatabaseMap AnimContext %s is not in a world!"), *AnimContext->GetName());
-		}
-		else if (SubsystemWorld != AnimContext->GetWorld())
-		{
-			UE_LOG(LogPoseSearch, Error, TEXT("UPoseSearchInteractionSubsystem::GenerateAnimContextInfosAndTagToDatabaseMap AnimContext %s is from World %s, and supposed to be in World %s!"), *AnimContext->GetName(), *AnimContext->GetWorld()->GetName(), *SubsystemWorld->GetName());
-		}
-		else
+	check(Algo::IsSorted(AnimContextsAvailabilities, [](const FPoseSearchInteractionAnimContextAvailabilities& A, const FPoseSearchInteractionAnimContextAvailabilities& B)
 		{
-			bool bContinueEvaluatingAnimContext = !Iter->Value.IsEmpty();
-			if (bContinueEvaluatingAnimContext)
-			{
-				// adding AnimContext to AnimContextsWithAvailabilities only if at least one availability has a valid database or has a valid tag
-				for (const FInteractionAvailability& InteractionAvailability : Iter->Value)
-				{
-					if (const UPoseSearchDatabase* Database = InteractionAvailability.Database.Get())
-					{
-						if (!Database->Schema)
-						{
-							UE_LOG(LogPoseSearch, Error, TEXT("UPoseSearchInteractionSubsystem::GenerateAnimContextInfosAndTagToDatabaseMap null Schema for Database %s. Dropping ALL availabilities for AnimContext %s"), *Database->GetName(), *AnimContext->GetName());
-							bContinueEvaluatingAnimContext = false;
-							break;
-						}
-					}
-				}
-			}
+			return A.AnimContext < B.AnimContext;
+		}));
 
-			if (bContinueEvaluatingAnimContext)
+	for (const FPoseSearchInteractionAnimContextAvailabilities& AnimContextAvailabilities : AnimContextsAvailabilities)
+	{
+		check(AnimContextAvailabilities.AnimContext && AnimContextAvailabilities.AnimContext->GetWorld() && AnimContextAvailabilities.AnimContext->GetWorld() == GetWorld());
+		check(!AnimContextAvailabilities.Availabilities.IsEmpty());
+		
+		// adding AnimContext to AnimContextsWithAvailabilities only if at least one availability has a valid database or has a valid tag
+		for (const FPoseSearchInteractionAvailabilityEx& InteractionAvailabilityEx : AnimContextAvailabilities.Availabilities)
+		{
+			bool bAnyValidAvailability = false;
+			if (const UPoseSearchDatabase* Database = InteractionAvailabilityEx.Database.Get())
 			{
-				bool bAnyValidAvailability = false;
-
-				for (const FInteractionAvailability& InteractionAvailability : Iter->Value)
-				{
-					const bool bIsTagValid = InteractionAvailability.IsTagValid();
-					if (const UPoseSearchDatabase* Database = InteractionAvailability.Database.Get())
-					{
-						check(Database->Schema);
-						bAnyValidAvailability = true;
-						if (bIsTagValid)
-						{
-							TagToDatabases.FindOrAdd(InteractionAvailability.Tag).AddUnique(InteractionAvailability.Database);
-						}
-					}
-					else if (bIsTagValid)
-					{
-						bAnyValidAvailability = true;
-					}
-				}
-			
-				if (bAnyValidAvailability)
+				check(Database->Schema);
+				if (InteractionAvailabilityEx.IsTagValid())
 				{
-					AnimContextsWithAvailabilities.Add({ AnimContext, &Iter->Value });
+					TagToDatabases.FindOrAdd(InteractionAvailabilityEx.Tag).AddUnique(InteractionAvailabilityEx.Database);
 				}
 			}
 		}
 	}
 
-	AnimContextsWithAvailabilities.Sort([](const FAnimContextWithAvailabilities& A, const FAnimContextWithAvailabilities& B)
-		{
-			// @todo: do we need to sort this reliably for networking?
-			return A.Key < B.Key;
-		});
-
-	const int32 NumAnimContextInfos = AnimContextsWithAvailabilities.Num();
+	const int32 NumAnimContextInfos = AnimContextsAvailabilities.Num();
 	AnimContextInfos.SetNum(NumAnimContextInfos);
 	for (int32 AnimContextInfoIndex = 0; AnimContextInfoIndex < NumAnimContextInfos; ++AnimContextInfoIndex)
 	{
-		FAnimContextInfo& AnimContextInfo = AnimContextInfos[AnimContextInfoIndex];
-		AnimContextInfo.AnimContext = AnimContextsWithAvailabilities[AnimContextInfoIndex].Key;
-		AnimContextInfo.Availabilities = AnimContextsWithAvailabilities[AnimContextInfoIndex].Value;
-		AnimContextInfo.Location = GetContextLocation(AnimContextInfo.AnimContext);
-		check(AnimContextInfo.AnimContext && AnimContextInfo.Availabilities && !AnimContextInfo.Availabilities->IsEmpty());
+		AnimContextInfos[AnimContextInfoIndex].Init(AnimContextsAvailabilities[AnimContextInfoIndex]);
 	}
 
 	// solving the broad phase using the AnimContextInfos
@@ -644,7 +628,7 @@ void UPoseSearchInteractionSubsystem::GenerateSearchContexts(float DeltaSeconds,
 
 	struct FRoledAnimContextInfo
 	{
-		FRoledAnimContextInfo(const FInteractionAvailability& InAvailability, const FAnimContextInfo& InAnimContextInfo, const FRole InRole, const IPoseHistory* InPoseHistory, const UPoseSearchDatabase& InDatabase)
+		FRoledAnimContextInfo(const FPoseSearchInteractionAvailabilityEx& InAvailability, const FAnimContextInfo& InAnimContextInfo, const FRole InRole, const IPoseHistory* InPoseHistory, const UPoseSearchDatabase& InDatabase)
 			: Availability(&InAvailability)
 			, AnimContextInfo(&InAnimContextInfo)
 			, Role(InRole)
@@ -679,7 +663,7 @@ void UPoseSearchInteractionSubsystem::GenerateSearchContexts(float DeltaSeconds,
 		bool operator==(const FRoledAnimContextInfo& Other) const = default;
 
 		// Availability that spawned this FRoledAnimContextInfo
-		const FInteractionAvailability* Availability = nullptr;
+		const FPoseSearchInteractionAvailabilityEx* Availability = nullptr;
 		// AnimContextInfo containing all the information regarding the AnimContext that spawned this FRoledAnimContextInfo,
 		// including all the availabilities associated to the AnimContext as well as all the other AnimContext(s) it can potentially interact with
 		const FAnimContextInfo* AnimContextInfo = nullptr;
@@ -690,7 +674,7 @@ void UPoseSearchInteractionSubsystem::GenerateSearchContexts(float DeltaSeconds,
 	
 	struct FRoledAnimContextInfos : TArray<FRoledAnimContextInfo, TInlineAllocator<PreallocatedRolesNum, TMemStackAllocator<>>>
 	{
-		void AddRoledAnimContextInfos(const FInteractionAvailability& Availability, const FAnimContextInfo& AnimContextInfo, const IPoseHistory* PoseHistory, const UPoseSearchDatabase& Database)
+		void AddRoledAnimContextInfos(const FPoseSearchInteractionAvailabilityEx& Availability, const FAnimContextInfo& AnimContextInfo, const IPoseHistory* PoseHistory, const UPoseSearchDatabase& Database)
 		{
 			const UPoseSearchSchema* Schema = Database.Schema;
 			check(Schema);
@@ -742,7 +726,7 @@ void UPoseSearchInteractionSubsystem::GenerateSearchContexts(float DeltaSeconds,
 				{
 					// starting the evaluation of a new set of grouped AnimContext(s)
 
-					// processing the AvailabilitiesMap of the current AnimContextArray to fill up a map of Databases pointing to an array of all the AnimContexts with related roles
+					// processing the AnimContextsAvailabilities of the current AnimContextArray to fill up a map of Databases pointing to an array of all the AnimContexts with related roles
 					VisitRecursively(AnimContextInfo, OnNewAnimAnimContextInfoFound);
 
 					OnDoneGroupingAnimContexts();
@@ -790,10 +774,10 @@ void UPoseSearchInteractionSubsystem::GenerateSearchContexts(float DeltaSeconds,
 		{
 			// analyzing all the Availability(s) associated with this AnimContext and eventually generate the associated FRoledAnimContextInfos,
 			// inserted in a per database sorted data structure (DatabaseToRoledAnimContextInfos)
-			check(AnimContextInfo.Availabilities);
-			for (const FInteractionAvailability& Availability : *AnimContextInfo.Availabilities)
+			check(AnimContextInfo.AnimContextAvailabilities);
+			for (const FPoseSearchInteractionAvailabilityEx& Availability : AnimContextInfo.AnimContextAvailabilities->Availabilities)
 			{
-				if (const IPoseHistory* PoseHistory = Availability.GetPoseHistory(AnimContextInfo.AnimContext))
+				if (const IPoseHistory* PoseHistory = Availability.GetPoseHistory(AnimContextInfo.AnimContextAvailabilities->AnimContext))
 				{
 					if (const UPoseSearchDatabase* Database = Availability.Database.Get())
 					{
@@ -829,7 +813,7 @@ void UPoseSearchInteractionSubsystem::GenerateSearchContexts(float DeltaSeconds,
 				}
 				else
 				{
-					UE_LOG(LogPoseSearch, Error, TEXT("UPoseSearchInteractionSubsystem::GenerateSearchContexts couldn't find PoseHistory %s for AnimContext %s"), *Availability.GetPoseHistoryName(), *AnimContextInfo.AnimContext->GetName());
+					UE_LOG(LogPoseSearch, Error, TEXT("UPoseSearchInteractionSubsystem::GenerateSearchContexts couldn't find PoseHistory %s for AnimContext %s"), *Availability.GetPoseHistoryName(), *AnimContextInfo.AnimContextAvailabilities->AnimContext->GetName());
 				}
 			}
 		},
@@ -866,7 +850,7 @@ void UPoseSearchInteractionSubsystem::GenerateSearchContexts(float DeltaSeconds,
 				// sort RoledAnimContextInfos to generate deterministic SearchContext across multiple frames!
 				RoledAnimContextInfos.Sort([](const FRoledAnimContextInfo& RoledAnimContextInfoA, const FRoledAnimContextInfo& RoledAnimContextInfoB)
 					{
-						return RoledAnimContextInfoA.AnimContextInfo->AnimContext < RoledAnimContextInfoB.AnimContextInfo->AnimContext;
+						return RoledAnimContextInfoA.AnimContextInfo->AnimContextAvailabilities->AnimContext < RoledAnimContextInfoB.AnimContextInfo->AnimContextAvailabilities->AnimContext;
 					});
 
 				GenerateCombinations(RoledAnimContextInfos.Num(), CombinationCardinality,
@@ -883,7 +867,7 @@ void UPoseSearchInteractionSubsystem::GenerateSearchContexts(float DeltaSeconds,
 							const FRoledAnimContextInfo& RoledAnimContextInfo = RoledAnimContextInfos[RoledAnimContextIndex];
 							check(RoledAnimContextInfo.AnimContextInfo);
 							bool bIsAlreadyInSet;
-							UniqueAnimContexts.Add(RoledAnimContextInfo.AnimContextInfo->AnimContext, &bIsAlreadyInSet);
+							UniqueAnimContexts.Add(RoledAnimContextInfo.AnimContextInfo->AnimContextAvailabilities->AnimContext, &bIsAlreadyInSet);
 							if (bIsAlreadyInSet)
 							{
 								// we have a duplicate AnimContext. this combination is NOT valid
@@ -902,7 +886,7 @@ void UPoseSearchInteractionSubsystem::GenerateSearchContexts(float DeltaSeconds,
 							const FRoledAnimContextInfo& RoledAnimContextInfo = RoledAnimContextInfos[RoledAnimContextIndex];
 
 							check(RoledAnimContextInfo.PoseHistory && RoledAnimContextInfo.Availability);
-							SearchContext.Add(RoledAnimContextInfo.AnimContextInfo->AnimContext, RoledAnimContextInfo.PoseHistory, RoledAnimContextInfo.Role);
+							SearchContext.Add(RoledAnimContextInfo.AnimContextInfo->AnimContextAvailabilities->AnimContext, RoledAnimContextInfo.PoseHistory, RoledAnimContextInfo.Role);
 							SearchContext.bDisableCollisions |= RoledAnimContextInfo.Availability->bDisableCollisions;
 							SearchContext.TickPriorities.Add(RoledAnimContextInfo.Availability->TickPriority);
 							
@@ -1121,7 +1105,7 @@ void UPoseSearchInteractionSubsystem::Tick(float DeltaSeconds)
 	FMemMark Mark(FMemStack::Get());
 	UpdateValidInteractionSearches();
 
-	if (AvailabilitiesMap.IsEmpty())
+	if (AnimContextsAvailabilities.IsEmpty())
 	{
 		bool bAllUninjected = true;
 		for (FInteractionIsland* Island : Islands)
@@ -1143,8 +1127,8 @@ void UPoseSearchInteractionSubsystem::Tick(float DeltaSeconds)
 
 	RegenerateAllIslands(DeltaSeconds);
 
-	// clearing up AvailabilitiesMap for the next frame
-	AvailabilitiesMap.Reset();
+	// clearing up AnimContextsAvailabilities for the next frame
+	AnimContextsAvailabilities.Reset();
 
 #if DO_CHECK
 	check(ValidateAllIslands());
diff --git a/Engine/Plugins/Animation/PoseSearch/Source/Runtime/Private/PoseSearchInteractionValidator.cpp b/Engine/Plugins/Animation/PoseSearch/Source/Runtime/Private/PoseSearchInteractionValidator.cpp
--- a/Engine/Plugins/Animation/PoseSearch/Source/Runtime/Private/PoseSearchInteractionValidator.cpp
+++ b/Engine/Plugins/Animation/PoseSearch/Source/Runtime/Private/PoseSearchInteractionValidator.cpp
@@ -32,7 +32,7 @@ FInteractionValidator::FInteractionValidator(const UObject* AnimContext)
 	const AActor* MainActor = ValidatingIsland->GetMainActor();
 	if (!MainActor)
 	{
-		UE_LOG(LogPoseSearch, Error, TEXT("FInteractionValidator invalid MainActor! How did he died? Did you rebuild the animation blue print while PIE was running?"));
+		UE_LOG(LogPoseSearch, Error, TEXT("FInteractionValidator invalid MainActor! How did he die? Did you rebuild the animation blue print while PIE was running?"));
 		ValidatingIsland->LogTickDependencies();
 		return;
 	}
diff --git a/Engine/Plugins/Animation/PoseSearch/Source/Runtime/Public/PoseSearch/PoseSearchInteractionAvailability.h b/Engine/Plugins/Animation/PoseSearch/Source/Runtime/Public/PoseSearch/PoseSearchInteractionAvailability.h
--- a/Engine/Plugins/Animation/PoseSearch/Source/Runtime/Public/PoseSearch/PoseSearchInteractionAvailability.h
+++ b/Engine/Plugins/Animation/PoseSearch/Source/Runtime/Public/PoseSearch/PoseSearchInteractionAvailability.h
@@ -59,31 +59,38 @@ struct FPoseSearchInteractionAvailability
 	bool operator==(const FPoseSearchInteractionAvailability& Other) const = default;
 };
 
-namespace UE::PoseSearch
-{
-
 // Experimental, this feature might be removed without warning, not for production use
-struct FInteractionAvailability : public FPoseSearchInteractionAvailability
+USTRUCT(Experimental)
+struct FPoseSearchInteractionAvailabilityEx : public FPoseSearchInteractionAvailability
 {
-	FInteractionAvailability(const FPoseSearchInteractionAvailability& InAvailability, FName InPoseHistoryName, const IPoseHistory* InPoseHistory)
-		: FPoseSearchInteractionAvailability(InAvailability)
-		, PoseHistoryName(InPoseHistoryName)
-		, PoseHistory(InPoseHistory)
+	GENERATED_BODY()
+
+	void Init(const FPoseSearchInteractionAvailability& InAvailability, FName InPoseHistoryName, const UE::PoseSearch::IPoseHistory* InPoseHistory)
 	{
-		check(!PoseHistory || PoseHistory->AsWeak().IsValid());
+		check(!InPoseHistory || InPoseHistory->AsWeak().IsValid());
+
+		static_cast<FPoseSearchInteractionAvailability&>(*this) = InAvailability;
+		PoseHistoryName = InPoseHistoryName;
+		PoseHistory = InPoseHistory;
 	}
 
 	FString GetPoseHistoryName() const;
-	const IPoseHistory* GetPoseHistory(const UObject* AnimContext) const;
+	const UE::PoseSearch::IPoseHistory* GetPoseHistory(const UObject* AnimContext) const;
 
 private:
 	FName PoseHistoryName;
-	const IPoseHistory* PoseHistory = nullptr;
+	const UE::PoseSearch::IPoseHistory* PoseHistory = nullptr;
 };
 
 // Experimental, this feature might be removed without warning, not for production use
-typedef TArray<FInteractionAvailability> FInteractionAvailabilities;
-// Experimental, this feature might be removed without warning, not for production use
-typedef TMap<FObjectKey, FInteractionAvailabilities> FAvailabilitiesMap;
+USTRUCT(Experimental)
+struct FPoseSearchInteractionAnimContextAvailabilities
+{
+	GENERATED_BODY()
 
-} // namespace UE::PoseSearch
+	UPROPERTY(Transient)
+	TObjectPtr<const UObject> AnimContext;
+
+	UPROPERTY(Transient)
+	TArray<FPoseSearchInteractionAvailabilityEx> Availabilities;
+};
diff --git a/Engine/Plugins/Animation/PoseSearch/Source/Runtime/Public/PoseSearch/PoseSearchInteractionSubsystem.h b/Engine/Plugins/Animation/PoseSearch/Source/Runtime/Public/PoseSearch/PoseSearchInteractionSubsystem.h
--- a/Engine/Plugins/Animation/PoseSearch/Source/Runtime/Public/PoseSearch/PoseSearchInteractionSubsystem.h
+++ b/Engine/Plugins/Animation/PoseSearch/Source/Runtime/Public/PoseSearch/PoseSearchInteractionSubsystem.h
@@ -123,10 +123,12 @@ private:
 	static UE_API void OnInteractionContinuing(UE::PoseSearch::FValidInteractionSearch& ValidInteractionSearch);
 	static UE_API void OnInteractionEnd(UE::PoseSearch::FValidInteractionSearch& ValidInteractionSearch);
 
-	// @todo: use a thread local storage data structure to make publishing availabilities lockless to avoid the need of AvailabilitiesMapMutex
+	// @todo: use a thread local storage data structure to make publishing availabilities lockless to avoid the need of AnimContextsAvailabilitiesMutex
 	// it could contain duplicated availabilities
-	UE::PoseSearch::FAvailabilitiesMap AvailabilitiesMap;
-	FCriticalSection AvailabilitiesMapMutex;
+	UPROPERTY(Transient)
+	TArray<FPoseSearchInteractionAnimContextAvailabilities> AnimContextsAvailabilities;
+
+	FCriticalSection AnimContextsAvailabilitiesMutex;
 
 	// array of groups of characters that needs to be anaylzed together for possible interactions
 	TArray<UE::PoseSearch::FInteractionIsland*> Islands;
