PoseSearch - running FInteractionIsland::LogTickDependencies from the game thread

diff --git a/Engine/Plugins/Animation/PoseSearch/Source/Runtime/Private/PoseSearchInteractionIsland.cpp b/Engine/Plugins/Animation/PoseSearch/Source/Runtime/Private/PoseSearchInteractionIsland.cpp
--- a/Engine/Plugins/Animation/PoseSearch/Source/Runtime/Private/PoseSearchInteractionIsland.cpp
+++ b/Engine/Plugins/Animation/PoseSearch/Source/Runtime/Private/PoseSearchInteractionIsland.cpp
@@ -570,7 +570,7 @@ FInteractionIsland::~FInteractionIsland()
 	InteractionSubsystem = nullptr;
 }
 
-IInteractionIslandDependency* FInteractionIsland::FindCustomDependency(UActorComponent* InTickComponent) const
+IInteractionIslandDependency* FInteractionIsland::FindCustomDependency(UActorComponent* InTickComponent)
 {
 	IModularFeatures& ModularFeatures = IModularFeatures::Get();
 	const int32 NumFeatures = ModularFeatures.GetModularFeatureImplementationCount(IInteractionIslandDependency::FeatureName);
@@ -580,7 +580,7 @@ IInteractionIslandDependency* FInteractionIsland::FindCustomDependency(UActorCom
 	{
 		if (IInteractionIslandDependency* IslandDependency = static_cast<IInteractionIslandDependency*>(ModularFeatures.GetModularFeatureImplementation(IInteractionIslandDependency::FeatureName, FeatureIndex)))
 		{
-			if (IslandDependency->CanMakeDependency(InteractionSubsystem, InTickComponent))
+			if (IslandDependency->CanMakeDependency(nullptr, InTickComponent))
 			{
 				return IslandDependency;
 			}
@@ -920,13 +920,12 @@ const AActor* FInteractionIsland::GetMainActor() const
 }
 
 #if ENABLE_ANIM_DEBUG
-void FInteractionIsland::LogTickDependencies() const
+void FInteractionIsland::LogTickDependencies(const TConstArrayView<TWeakObjectPtr<UActorComponent>> TickActorComponents, int32 InteractionIslandIndex)
 {
-	const ENamedThreads::Type Type = IsInGameThread() ? ENamedThreads::GameThread : ENamedThreads::AnyThread;
+	check(IsInGameThread());
 
 	for (const TWeakObjectPtr<UActorComponent>& TickActorComponentPtr : TickActorComponents)
 	{
-		const int32 InteractionIslandIndex = InteractionSubsystem->GetInteractionIslands().IndexOfByKey(this);
 		if (UActorComponent* TickActorComponent = TickActorComponentPtr.Get())
 		{
 			UE_LOG(LogPoseSearch, Log, TEXT("============== %s (Island %d) =============="), *TickActorComponent->GetOwner()->GetName(), InteractionIslandIndex);
@@ -935,7 +934,7 @@ void FInteractionIsland::LogTickDependencies() const
 				// AnimNextComponent case
 				if (const FTickFunction* TickActorComponentTickFunction = IslandDependency->FindTickFunction(TickActorComponent))
 				{
-					LogTickFunction(*const_cast<FTickFunction*>(TickActorComponentTickFunction), Type, true, 1);
+					LogTickFunction(*const_cast<FTickFunction*>(TickActorComponentTickFunction), ENamedThreads::GameThread, true, 1);
 				}
 				else
 				{
@@ -945,7 +944,7 @@ void FInteractionIsland::LogTickDependencies() const
 			else
 			{
 				// SkeletalMeshComponent / AnimInstance case
-				LogTickFunction(TickActorComponent->PrimaryComponentTick, Type, true, 1);
+				LogTickFunction(TickActorComponent->PrimaryComponentTick, ENamedThreads::GameThread, true, 1);
 			}
 		}
 		else
@@ -953,7 +952,24 @@ void FInteractionIsland::LogTickDependencies() const
 			UE_LOG(LogPoseSearch, Log, TEXT("============== !!!Missing Actor!!! (Island %d) =============="), InteractionIslandIndex);
 		}
 	}
+}
+
+void FInteractionIsland::LogTickDependencies() const
+{
+	const int32 InteractionIslandIndex = InteractionSubsystem->GetInteractionIslands().IndexOfByKey(this);
 
+	if (IsInGameThread())
+	{
+		LogTickDependencies(TickActorComponents, InteractionIslandIndex);
+	}
+	else
+	{
+		TArray<TWeakObjectPtr<UActorComponent>> TickActorComponentsCopy = TickActorComponents;
+		FFunctionGraphTask::CreateAndDispatchWhenReady([TickActorComponentsCopy, InteractionIslandIndex]()
+		{
+			LogTickDependencies(TickActorComponentsCopy, InteractionIslandIndex);
+		}, TStatId(), nullptr, ENamedThreads::GameThread);
+	}
 }
 #endif // ENABLE_ANIM_DEBUG
 
diff --git a/Engine/Plugins/Animation/PoseSearch/Source/Runtime/Public/PoseSearch/PoseSearchInteractionIsland.h b/Engine/Plugins/Animation/PoseSearch/Source/Runtime/Public/PoseSearch/PoseSearchInteractionIsland.h
--- a/Engine/Plugins/Animation/PoseSearch/Source/Runtime/Public/PoseSearch/PoseSearchInteractionIsland.h
+++ b/Engine/Plugins/Animation/PoseSearch/Source/Runtime/Public/PoseSearch/PoseSearchInteractionIsland.h
@@ -154,9 +154,13 @@ struct FInteractionIsland
 #endif // ENABLE_ANIM_DEBUG
 
 private:
+#if ENABLE_ANIM_DEBUG
+	static void LogTickDependencies(const TConstArrayView<TWeakObjectPtr<UActorComponent>> TickActorComponents, int32 InteractionIslandIndex);
+#endif // ENABLE_ANIM_DEBUG
+
 	void AddTickDependencies(UActorComponent* TickActorComponent, bool bInIsMainActor);
 	void RemoveTickDependencies(bool bValidateTickDependencies);
-	IInteractionIslandDependency* FindCustomDependency(UActorComponent* InTickComponent) const;
+	static IInteractionIslandDependency* FindCustomDependency(UActorComponent* InTickComponent);
 
 	const UObject* GetMainAnimContext() const;
 	const AActor* GetMainActor() const;
