PoseSearch - posesearch debugger tracing objects fix

diff --git a/Engine/Plugins/Animation/PoseSearch/Source/Editor/Private/PoseSearchDebugger.cpp b/Engine/Plugins/Animation/PoseSearch/Source/Editor/Private/PoseSearchDebugger.cpp
--- a/Engine/Plugins/Animation/PoseSearch/Source/Editor/Private/PoseSearchDebugger.cpp
+++ b/Engine/Plugins/Animation/PoseSearch/Source/Editor/Private/PoseSearchDebugger.cpp
@@ -488,14 +488,16 @@ bool FDebuggerTrack::UpdateInternal()
 		}
 
 		TMap<int32, FText, TInlineSetAllocator<64>> SearchIdNames;
-		PoseSearchProvider->EnumerateMotionMatchingStateTimelines(ObjectId, [StartTime, EndTime, &SearchIdNames](const FTraceProvider::FMotionMatchingStateTimeline& InTimeline)
+		const IGameplayProvider* GameplayProvider = AnalysisSession->ReadProvider<IGameplayProvider>("GameplayProvider");
+
+		PoseSearchProvider->EnumerateMotionMatchingStateTimelines(ObjectId, [StartTime, EndTime, &SearchIdNames, GameplayProvider](const FTraceProvider::FMotionMatchingStateTimeline& InTimeline)
 			{
 				// this isn't very efficient, and it gets called every frame.  will need optimizing
-				InTimeline.EnumerateEvents(StartTime, EndTime, [StartTime, EndTime, &SearchIdNames](double InStartTime, double InEndTime, uint32 InDepth, const FTraceMotionMatchingStateMessage& InMessage)
+				InTimeline.EnumerateEvents(StartTime, EndTime, [StartTime, EndTime, &SearchIdNames, GameplayProvider](double InStartTime, double InEndTime, uint32 InDepth, const FTraceMotionMatchingStateMessage& InMessage)
 					{
 						if (!SearchIdNames.Find(InMessage.GetSearchId()) && InEndTime > StartTime && InStartTime < EndTime)
 						{
-							SearchIdNames.Add(InMessage.GetSearchId()) = InMessage.GenerateSearchName();
+							SearchIdNames.Add(InMessage.GetSearchId()) = GenerateSearchName(InMessage, GameplayProvider);
 						}
 						return TraceServices::EEventEnumerate::Continue;
 					});
diff --git a/Engine/Plugins/Animation/PoseSearch/Source/Editor/Private/PoseSearchDebuggerView.cpp b/Engine/Plugins/Animation/PoseSearch/Source/Editor/Private/PoseSearchDebuggerView.cpp
--- a/Engine/Plugins/Animation/PoseSearch/Source/Editor/Private/PoseSearchDebuggerView.cpp
+++ b/Engine/Plugins/Animation/PoseSearch/Source/Editor/Private/PoseSearchDebuggerView.cpp
@@ -1,8 +1,10 @@
 // Copyright Epic Games, Inc. All Rights Reserved.
 
 #include "PoseSearchDebuggerView.h"
-#include "Modules/ModuleManager.h"
+#include "Animation/TrajectoryTypes.h"
 #include "Editor.h"
+#include "IGameplayProvider.h"
+#include "Modules/ModuleManager.h"
 #include "PoseSearch/PoseSearchDatabase.h"
 #include "PoseSearch/PoseSearchDerivedData.h"
 #include "PoseSearch/PoseSearchSchema.h"
@@ -14,7 +16,6 @@
 #include "PoseSearchDebuggerViewModel.h"
 #include "PoseSearchEditor.h"
 #include "PropertyEditorModule.h"
-#include "Animation/TrajectoryTypes.h"
 #include "Subsystems/AssetEditorSubsystem.h"
 #include "Trace/PoseSearchTraceProvider.h"
 #include "Widgets/Input/SButton.h"
@@ -29,6 +30,67 @@
 namespace UE::PoseSearch
 {
 
+FText GenerateSearchName(const FTraceMotionMatchingStateMessage& MotionMatchingState, const IGameplayProvider* GameplayProvider)
+{
+	TStringBuilder<256> StringBuilder;
+	if (GameplayProvider)
+	{
+		bool bAddDatabasesSeparator = false;
+		for (const FTraceMotionMatchingStateDatabaseEntry& DbEntry : MotionMatchingState.DatabaseEntries)
+		{
+			if (bAddDatabasesSeparator)
+			{
+				StringBuilder.Append(" - ");
+			}
+
+			const FObjectInfo& DatabaseObjectInfo = GameplayProvider->GetObjectInfo(DbEntry.DatabaseId);
+			StringBuilder.Append(DatabaseObjectInfo.Name);
+			bAddDatabasesSeparator = true;
+		}
+
+		if (MotionMatchingState.Roles.Num() > 1)
+		{
+			if (MotionMatchingState.Roles.Num() != MotionMatchingState.SkeletalMeshComponentIds.Num())
+			{
+				StringBuilder.Append("Error!");
+			}
+			else
+			{
+				StringBuilder.Append(" [");
+
+				bool bAddRolesSeparator = false;
+
+				for (int32 RoleIndex = 0; RoleIndex < MotionMatchingState.Roles.Num(); ++RoleIndex)
+				{
+					if (bAddRolesSeparator)
+					{
+						StringBuilder.Append(" - ");
+					}
+
+					StringBuilder.Append(MotionMatchingState.Roles[RoleIndex].ToString());
+
+					StringBuilder.Append(": ");
+
+					const FObjectInfo& SkeletalMeshComponentObjectInfo = GameplayProvider->GetObjectInfo(MotionMatchingState.SkeletalMeshComponentIds[RoleIndex]);
+					const FObjectInfo& SkeletalMeshComponentOuterObjectInfo = GameplayProvider->GetObjectInfo(SkeletalMeshComponentObjectInfo.OuterId);
+
+					StringBuilder.Append(SkeletalMeshComponentOuterObjectInfo.Name);
+
+					bAddRolesSeparator = true;
+				}
+
+				StringBuilder.Append("]");
+			}
+		}
+	}
+	else
+	{
+		StringBuilder.Append("Error!");
+	}
+
+	return FText::FromString(StringBuilder.ToString());
+}
+
 class SDebuggerMessageBox : public SCompoundWidget
 {
 	SLATE_BEGIN_ARGS(SDebuggerMessageBox) {}
@@ -431,23 +493,32 @@ bool SDebuggerView::UpdateNodeSelection()
 
 	// Create selection view with buttons for each node, displaying the database name
 	SelectionView->ClearChildren();
-	for (const FTraceMotionMatchingStateMessage& MotionMatchingState : MotionMatchingStates)
+
+	if (!MotionMatchingStates.IsEmpty())
 	{
-		Model->OnUpdateSearchSelection(MotionMatchingState.GetSearchId());
-		SelectionView->AddSlot()
-		.HAlign(HAlign_Fill)
-		.VAlign(VAlign_Center)
-		.Padding(10.0f)
-		[
-			SNew(SButton)
-			.Text(MotionMatchingState.GenerateSearchName())
-			.HAlign(HAlign_Center)
-			.VAlign(VAlign_Center)
-			.ContentPadding(10.0f)
-			.OnClicked(this, &SDebuggerView::OnUpdateSearchSelection, MotionMatchingState.GetSearchId())
-		];
-	}
+		IRewindDebugger* RewindDebugger = IRewindDebugger::Instance();
+		const TraceServices::IAnalysisSession* AnalysisSession = RewindDebugger->GetAnalysisSession();
+		check(AnalysisSession);
+		const IGameplayProvider* GameplayProvider = AnalysisSession->ReadProvider<IGameplayProvider>("GameplayProvider");
+		TraceServices::FAnalysisSessionReadScope SessionReadScope(*AnalysisSession);
 
+		for (const FTraceMotionMatchingStateMessage& MotionMatchingState : MotionMatchingStates)
+		{
+			Model->OnUpdateSearchSelection(MotionMatchingState.GetSearchId());
+			SelectionView->AddSlot()
+				.HAlign(HAlign_Fill)
+				.VAlign(VAlign_Center)
+				.Padding(10.0f)
+				[
+					SNew(SButton)
+						.Text(GenerateSearchName(MotionMatchingState, GameplayProvider))
+						.HAlign(HAlign_Center)
+						.VAlign(VAlign_Center)
+						.ContentPadding(10.0f)
+						.OnClicked(this, &SDebuggerView::OnUpdateSearchSelection, MotionMatchingState.GetSearchId())
+				];
+		}
+	}
 	return false;
 }
 
diff --git a/Engine/Plugins/Animation/PoseSearch/Source/Editor/Private/PoseSearchDebuggerView.h b/Engine/Plugins/Animation/PoseSearch/Source/Editor/Private/PoseSearchDebuggerView.h
--- a/Engine/Plugins/Animation/PoseSearch/Source/Editor/Private/PoseSearchDebuggerView.h
+++ b/Engine/Plugins/Animation/PoseSearch/Source/Editor/Private/PoseSearchDebuggerView.h
@@ -6,9 +6,10 @@
 #include "PoseSearch/PoseSearchDefines.h"
 
 class IDetailsView;
-class SVerticalBox;
+class IGameplayProvider;
 class SHorizontalBox;
 class SSplitter;
+class SVerticalBox;
 class SWidgetSwitcher;
 class UPoseSearchDatabase;
 class UPoseSearchDebuggerReflection;
@@ -165,4 +166,6 @@ private:
 	static constexpr int32 ConsecutiveFramesUpdateThreshold = 10;
 };
 
+FText GenerateSearchName(const FTraceMotionMatchingStateMessage& MotionMatchingState, const IGameplayProvider* GameplayProvider);
+
 } // namespace UE::PoseSearch
diff --git a/Engine/Plugins/Animation/PoseSearch/Source/Runtime/Private/PoseSearchLibrary.cpp b/Engine/Plugins/Animation/PoseSearch/Source/Runtime/Private/PoseSearchLibrary.cpp
--- a/Engine/Plugins/Animation/PoseSearch/Source/Runtime/Private/PoseSearchLibrary.cpp
+++ b/Engine/Plugins/Animation/PoseSearch/Source/Runtime/Private/PoseSearchLibrary.cpp
@@ -474,6 +474,7 @@ void UPoseSearchLibrary::TraceMotionMatching(
 				return;
 			}
 			
+			TRACE_OBJECT(SkeletalMeshComponent);
 			TraceState.SkeletalMeshComponentIds[AnimInstanceIndex] = FObjectTrace::GetObjectId(SkeletalMeshComponent);
 		}
 	}
@@ -516,7 +517,8 @@ void UPoseSearchLibrary::TraceMotionMatching(
 
 		// if throttling is on, the continuing pose can be valid, but no actual search occurred, so the query will not be cached, and we need to build it
 		DbEntry.QueryVector = SearchContext.GetOrBuildQuery(Database->Schema);
-		DbEntry.DatabaseId = FTraceMotionMatchingStateMessage::GetIdFromObject(Database);
+		TRACE_OBJECT(Database);
+		DbEntry.DatabaseId = FObjectTrace::GetObjectId(Database);
 		DatabaseIds.Add(DbEntry.DatabaseId);
 
 		for (int32 CandidateIdx = 0; CandidateIdx < DatabaseBestPoseCandidates.Value.Num(); ++CandidateIdx)
@@ -617,7 +619,9 @@ void UPoseSearchLibrary::TraceMotionMatching(
 	// @todo: avoid publishing duplicated TraceState in ALL the AnimContexts! -currently necessary for multi character-
 	for (const FChooserEvaluationContext* Context : SearchContext.GetContexts())
 	{
-		TraceState.AnimInstanceId = FObjectTrace::GetObjectId(Context->GetFirstObjectParam());
+		const UObject* AnimContextObject = Context->GetFirstObjectParam();
+		TRACE_OBJECT(AnimContextObject);
+		TraceState.AnimInstanceId = FObjectTrace::GetObjectId(AnimContextObject);
 		TraceState.NodeId = SearchId;
 		TraceState.Output();
 	}
diff --git a/Engine/Plugins/Animation/PoseSearch/Source/Runtime/Public/PoseSearch/Trace/PoseSearchTraceLogger.h b/Engine/Plugins/Animation/PoseSearch/Source/Runtime/Public/PoseSearch/Trace/PoseSearchTraceLogger.h
--- a/Engine/Plugins/Animation/PoseSearch/Source/Runtime/Public/PoseSearch/Trace/PoseSearchTraceLogger.h
+++ b/Engine/Plugins/Animation/PoseSearch/Source/Runtime/Public/PoseSearch/Trace/PoseSearchTraceLogger.h
@@ -121,6 +121,7 @@ struct FTraceMotionMatchingStateMessage : public FTraceMessage
 		return nullptr;
 	}
 
+	UE_DEPRECATED(5.6, "Use FObjectTrace::GetObjectId instead")
 	static uint64 GetIdFromObject(const UObject* Object)
 	{
 #if OBJECT_TRACE_ENABLED
@@ -130,6 +131,7 @@ struct FTraceMotionMatchingStateMessage : public FTraceMessage
 #endif
 	}
 	
+	UE_DEPRECATED(5.6, "Use FText GenerateSearchName(const FTraceMotionMatchingStateMessage&, const IGameplayProvider*) instead")
 	UE_API FText GenerateSearchName() const;
 
 	static UE_API const FName Name;
