PoseSearch - fix for database indexing crash when using UAnimNotifyState_PoseSearchSamplingAttribute with bones missing from any schema channel

diff --git a/Engine/Plugins/Animation/PoseSearch/Source/Runtime/Private/PoseSearchDerivedData.cpp b/Engine/Plugins/Animation/PoseSearch/Source/Runtime/Private/PoseSearchDerivedData.cpp
--- a/Engine/Plugins/Animation/PoseSearch/Source/Runtime/Private/PoseSearchDerivedData.cpp
+++ b/Engine/Plugins/Animation/PoseSearch/Source/Runtime/Private/PoseSearchDerivedData.cpp
@@ -913,6 +913,86 @@ struct FSamplerMapKey
 	FVector BlendParameters = FVector::ZeroVector;
 };
 
+static void InitRoledBoneContainers(TMap<FRole, FBoneContainer>& RoledBoneContainers, const UPoseSearchDatabase* DatabaseToLookForAssets, const UPoseSearchSchema* Schema)
+{
+	auto AddBoneIndex = [](int32 BoneIndex, TArray<uint16>& BoneIndicesWithParents, const UMirrorDataTable* MirrorDataTable)
+	{
+		BoneIndicesWithParents.AddUnique(BoneIndex);
+
+		// adding mirrored BoneIndex if there's a valid MirrorDataTable
+		if (MirrorDataTable)
+		{
+			if (MirrorDataTable->BoneToMirrorBoneIndex.IsValidIndex(BoneIndex))
+			{
+				const FSkeletonPoseBoneIndex MirroredBoneIndex = MirrorDataTable->BoneToMirrorBoneIndex[BoneIndex];
+				if (MirroredBoneIndex.IsValid())
+				{
+					BoneIndicesWithParents.AddUnique(MirroredBoneIndex.GetInt());
+				}
+			}
+			else
+			{
+				UE_LOG(LogPoseSearch, Warning, TEXT("InitBoneContainersFromRoledSkeleton: MirrorDataTable %s doesn't contain bone with index %d."), *MirrorDataTable->GetName(), BoneIndex);
+			}
+		}
+	};
+
+	RoledBoneContainers.Reset();
+	RoledBoneContainers.Reserve(Schema->GetRoledSkeletons().Num());
+
+	// filling up BoneIndicesWithParents with all the bone indexes from the bones in the 
+	// schema roled skeletons and from UAnimNotifyState_PoseSearchSamplingAttribute
+	TArray<uint16> BoneIndicesWithParents;
+	BoneIndicesWithParents.Reserve(128);
+
+	for (const FPoseSearchRoledSkeleton& RoledSkeleton : Schema->GetRoledSkeletons())
+	{
+		BoneIndicesWithParents.Reset();
+
+		FBoneContainer& RoledBoneContainer = RoledBoneContainers.Add(RoledSkeleton.Role);
+		// Add a curve filter to our bone container to only eval curves actually used by the schema.
+		const UE::Anim::FCurveFilterSettings CurveFilterSettings(UE::Anim::ECurveFilterMode::AllowOnlyFiltered, &RoledSkeleton.RequiredCurves);
+
+		// Initialize references to obtain bone indices and fill out bone index array
+		for (const FBoneReference& BoneRef : RoledSkeleton.BoneReferences)
+		{
+			check(BoneRef.HasValidSetup());
+			AddBoneIndex(BoneRef.BoneIndex, BoneIndicesWithParents, RoledSkeleton.MirrorDataTable);
+		}
+
+		for (int32 AnimationAssetIndex = 0; AnimationAssetIndex < DatabaseToLookForAssets->GetNumAnimationAssets(); ++AnimationAssetIndex)
+		{
+			if (const FPoseSearchDatabaseAnimationAssetBase* DatabaseAsset = DatabaseToLookForAssets->GetDatabaseAnimationAsset<FPoseSearchDatabaseAnimationAssetBase>(AnimationAssetIndex))
+			{
+				if (!DatabaseAsset->IsEnabled() || !DatabaseAsset->GetAnimationAsset())
+				{
+					continue;
+				}
+
+				const FAnimationAssetSampler Sampler(DatabaseAsset->GetAnimationAssetForRole(RoledSkeleton.Role), FTransform::Identity, FVector::ZeroVector, FAnimationAssetSampler::DefaultRootTransformSamplingRate, false, true);
+				for (const FAnimNotifyEvent& AnimNotifyEvent : Sampler.GetAllAnimNotifyEvents())
+				{
+					if (const UAnimNotifyState_PoseSearchSamplingAttribute* SamplingAttribute = Cast<UAnimNotifyState_PoseSearchSamplingAttribute>(AnimNotifyEvent.NotifyStateClass))
+					{
+						FBoneReference TempBoneReference = SamplingAttribute->Bone;
+						TempBoneReference.Initialize(RoledSkeleton.Skeleton);
+						if (TempBoneReference.HasValidSetup())
+						{
+							AddBoneIndex(TempBoneReference.BoneIndex, BoneIndicesWithParents, RoledSkeleton.MirrorDataTable);
+						}
+					}
+				}
+			}
+		}
+
+		// Sort bone indexes and add eventual missing parent bone indexes
+		BoneIndicesWithParents.Sort();
+		FAnimationRuntime::EnsureParentsPresent(BoneIndicesWithParents, RoledSkeleton.Skeleton->GetReferenceSkeleton());
+
+		RoledBoneContainer.InitializeTo(BoneIndicesWithParents, CurveFilterSettings, *RoledSkeleton.Skeleton);
+	}
+}
+
 static bool IndexDatabase(FSearchIndexBase& SearchIndexBase, const UPoseSearchDatabase* DatabaseToLookForAssets, const UPoseSearchSchema* Schema, const FAssetSamplingContext& SamplingContext, const FFloatInterval& AdditionalExtrapolationTime, UE::DerivedData::FRequestOwner& Owner)
 {
 	if (!ensure(Schema))
@@ -962,7 +1042,7 @@ static bool IndexDatabase(FSearchIndexBase& SearchIndexBase, const UPoseSearchDa
 	Indexers.Reserve(SearchIndexBase.Assets.Num());
 
 	TMap<FRole, FBoneContainer> RoledBoneContainers;
-	Schema->InitBoneContainersFromRoledSkeleton(RoledBoneContainers);
+	InitRoledBoneContainers(RoledBoneContainers, DatabaseToLookForAssets, Schema);
 
 	TMap<FRole, FMirrorDataCache> RoledMirrorDataCaches;
 	RoledMirrorDataCaches.Reserve(RoledBoneContainers.Num());
@@ -2422,7 +2502,7 @@ void FAsyncPoseSearchDatabasesManagement::Tick(float DeltaTime)
 				if (Database->Schema)
 				{
 					TMap<FRole, FBoneContainer> RoledBoneContainers;
-					Database->Schema->InitBoneContainersFromRoledSkeleton(RoledBoneContainers);
+					InitRoledBoneContainers(RoledBoneContainers, Database, Database->Schema);
 
 					for (int32 AnimationAssetIndex = 0; AnimationAssetIndex < Database->GetNumAnimationAssets(); ++AnimationAssetIndex)
 					{
diff --git a/Engine/Plugins/Animation/PoseSearch/Source/Runtime/Private/PoseSearchSchema.cpp b/Engine/Plugins/Animation/PoseSearch/Source/Runtime/Private/PoseSearchSchema.cpp
--- a/Engine/Plugins/Animation/PoseSearch/Source/Runtime/Private/PoseSearchSchema.cpp
+++ b/Engine/Plugins/Animation/PoseSearch/Source/Runtime/Private/PoseSearchSchema.cpp
@@ -87,7 +87,9 @@ void UPoseSearchSchema::InitBoneContainersFromRoledSkeleton(TMap<FName, FBoneCon
 		// Add a curve filter to our bone container to only eval curves actually used by the schema.
 		const UE::Anim::FCurveFilterSettings CurveFilterSettings(UE::Anim::ECurveFilterMode::AllowOnlyFiltered, &RoledSkeleton.RequiredCurves);
 
-		RoledBoneContainer.InitializeTo(RoledSkeleton.BoneIndicesWithParents, CurveFilterSettings, *RoledSkeleton.Skeleton);
+		PRAGMA_DISABLE_DEPRECATION_WARNINGS
+		RoledBoneContainer.InitializeTo(RoledSkeleton.BoneIndicesWithParents_DEPRECATED, CurveFilterSettings, *RoledSkeleton.Skeleton);
+		PRAGMA_ENABLE_DEPRECATION_WARNINGS
 	}
 }
 
@@ -242,7 +244,9 @@ void UPoseSearchSchema::ResetFinalize()
 	for (FPoseSearchRoledSkeleton& RoledSkeleton : Skeletons)
 	{
 		RoledSkeleton.BoneReferences.Reset();
-		RoledSkeleton.BoneIndicesWithParents.Reset();
+		PRAGMA_DISABLE_DEPRECATION_WARNINGS
+		RoledSkeleton.BoneIndicesWithParents_DEPRECATED.Reset();
+		PRAGMA_ENABLE_DEPRECATION_WARNINGS
 	}
 
 	FinalizedChannels.Reset();
@@ -331,13 +335,14 @@ void UPoseSearchSchema::Finalize()
 		}
 	}
 
+	PRAGMA_DISABLE_DEPRECATION_WARNINGS
 	// Initialize references to obtain bone indices and fill out bone index array
 	for (FPoseSearchRoledSkeleton& RoledSkeleton : Skeletons)
 	{
 		for (FBoneReference& BoneRef : RoledSkeleton.BoneReferences)
 		{
 			check(BoneRef.HasValidSetup());
-			RoledSkeleton.BoneIndicesWithParents.AddUnique(BoneRef.BoneIndex);
+			RoledSkeleton.BoneIndicesWithParents_DEPRECATED.AddUnique(BoneRef.BoneIndex);
 
 			if (RoledSkeleton.MirrorDataTable)
 			{
@@ -346,7 +351,7 @@ void UPoseSearchSchema::Finalize()
 					const FSkeletonPoseBoneIndex MirroredBoneIndex = RoledSkeleton.MirrorDataTable->BoneToMirrorBoneIndex[BoneRef.BoneIndex];
 					if (MirroredBoneIndex.IsValid())
 					{
-						RoledSkeleton.BoneIndicesWithParents.AddUnique(MirroredBoneIndex.GetInt());
+						RoledSkeleton.BoneIndicesWithParents_DEPRECATED.AddUnique(MirroredBoneIndex.GetInt());
 					}
 				}
 				else
@@ -357,10 +362,11 @@ void UPoseSearchSchema::Finalize()
 		}
 
 		// Build separate index array with parent indices guaranteed to be present. Sort for EnsureParentsPresent.
-		check(!RoledSkeleton.BoneIndicesWithParents.IsEmpty());
-		RoledSkeleton.BoneIndicesWithParents.Sort();
-		FAnimationRuntime::EnsureParentsPresent(RoledSkeleton.BoneIndicesWithParents, RoledSkeleton.Skeleton->GetReferenceSkeleton());
+		check(!RoledSkeleton.BoneIndicesWithParents_DEPRECATED.IsEmpty());
+		RoledSkeleton.BoneIndicesWithParents_DEPRECATED.Sort();
+		FAnimationRuntime::EnsureParentsPresent(RoledSkeleton.BoneIndicesWithParents_DEPRECATED, RoledSkeleton.Skeleton->GetReferenceSkeleton());
 	}
+	PRAGMA_ENABLE_DEPRECATION_WARNINGS
 }
 
 void UPoseSearchSchema::PostLoad()
diff --git a/Engine/Plugins/Animation/PoseSearch/Source/Runtime/Public/PoseSearch/PoseSearchSchema.h b/Engine/Plugins/Animation/PoseSearch/Source/Runtime/Public/PoseSearch/PoseSearchSchema.h
--- a/Engine/Plugins/Animation/PoseSearch/Source/Runtime/Public/PoseSearch/PoseSearchSchema.h
+++ b/Engine/Plugins/Animation/PoseSearch/Source/Runtime/Public/PoseSearch/PoseSearchSchema.h
@@ -50,8 +50,9 @@ struct FPoseSearchRoledSkeleton
 	UPROPERTY(Transient)
 	TArray<FBoneReference> BoneReferences;
 
+	UE_DEPRECATED(5.6, "no longer necessary property")
 	UPROPERTY(Transient)
-	TArray<uint16> BoneIndicesWithParents;
+	TArray<uint16> BoneIndicesWithParents_DEPRECATED;
 
 	UPROPERTY(Transient)
 	TArray<FName> RequiredCurves;
@@ -175,10 +176,9 @@ public:
 	UE_API bool AreSkeletonsCompatible(const UPoseSearchSchema* Other) const;
 	UE_API void AddDefaultChannels();
 
-	// we avoid exposing GetRoledSkeleton(RoleIndex) and all others method referencing Skeletons by index
-	// to avoid confusing the Role to index mapping in the Schema versus the Role to index mapping in other assets
-	// please use GetRoledSkeleton(Role)
+	UE_DEPRECATED(5.6, "No longer supported API")
 	UE_API void InitBoneContainersFromRoledSkeleton(TMap<FName, FBoneContainer>& RoledBoneContainers) const;
+
 	UE_DEPRECATED(5.6, "access the Skeletons via GetRoledSkeletons() and query for their MirrorDataTable instead")
 	UE_API bool AllRoledSkeletonHaveMirrorDataTable() const;
 	UE_API const FPoseSearchRoledSkeleton* GetRoledSkeleton(const UE::PoseSearch::FRole& Role) const;
