Fix deprecation warnings due to implicit assignment operator
#rnx

diff --git a/Engine/Plugins/Animation/PoseSearch/Source/Runtime/Private/PoseSearchSchema.cpp b/Engine/Plugins/Animation/PoseSearch/Source/Runtime/Private/PoseSearchSchema.cpp
--- a/Engine/Plugins/Animation/PoseSearch/Source/Runtime/Private/PoseSearchSchema.cpp
+++ b/Engine/Plugins/Animation/PoseSearch/Source/Runtime/Private/PoseSearchSchema.cpp
@@ -12,6 +12,10 @@
 #include "PoseSearch/PoseSearchFeatureChannel_Trajectory.h"
 #include "UObject/ObjectSaveContext.h"
 
+PRAGMA_DISABLE_DEPRECATION_WARNINGS
+FPoseSearchRoledSkeleton& FPoseSearchRoledSkeleton::operator=(const FPoseSearchRoledSkeleton&) = default;
+PRAGMA_ENABLE_DEPRECATION_WARNINGS
+
 void UPoseSearchSchema::AddChannel(UPoseSearchFeatureChannel* Channel)
 {
 	Channels.Add(Channel);
diff --git a/Engine/Plugins/Animation/PoseSearch/Source/Runtime/Public/PoseSearch/PoseSearchSchema.h b/Engine/Plugins/Animation/PoseSearch/Source/Runtime/Public/PoseSearch/PoseSearchSchema.h
--- a/Engine/Plugins/Animation/PoseSearch/Source/Runtime/Public/PoseSearch/PoseSearchSchema.h
+++ b/Engine/Plugins/Animation/PoseSearch/Source/Runtime/Public/PoseSearch/PoseSearchSchema.h
@@ -36,6 +36,9 @@ struct FPoseSearchRoledSkeleton
 {
 	GENERATED_BODY()
 
+	// Declared to be able to disable deprecation warnings
+	FPoseSearchRoledSkeleton& operator=(const FPoseSearchRoledSkeleton&);
+
 	// Skeleton Reference for Motion Matching Database assets. Must be set to a compatible skeleton to the animation data in the database.
 	UPROPERTY(EditAnywhere, Category = "Schema")
 	TObjectPtr<USkeleton> Skeleton;
