Also declare and implement defaulted copy constructor to avoid deprecation warnings
#rnx

diff --git a/Engine/Plugins/Animation/PoseSearch/Source/Runtime/Private/PoseSearchSchema.cpp b/Engine/Plugins/Animation/PoseSearch/Source/Runtime/Private/PoseSearchSchema.cpp
--- a/Engine/Plugins/Animation/PoseSearch/Source/Runtime/Private/PoseSearchSchema.cpp
+++ b/Engine/Plugins/Animation/PoseSearch/Source/Runtime/Private/PoseSearchSchema.cpp
@@ -13,6 +13,7 @@
 #include "UObject/ObjectSaveContext.h"
 
 PRAGMA_DISABLE_DEPRECATION_WARNINGS
+FPoseSearchRoledSkeleton::FPoseSearchRoledSkeleton(const FPoseSearchRoledSkeleton&) = default;
 FPoseSearchRoledSkeleton& FPoseSearchRoledSkeleton::operator=(const FPoseSearchRoledSkeleton&) = default;
 PRAGMA_ENABLE_DEPRECATION_WARNINGS
 
diff --git a/Engine/Plugins/Animation/PoseSearch/Source/Runtime/Public/PoseSearch/PoseSearchSchema.h b/Engine/Plugins/Animation/PoseSearch/Source/Runtime/Public/PoseSearch/PoseSearchSchema.h
--- a/Engine/Plugins/Animation/PoseSearch/Source/Runtime/Public/PoseSearch/PoseSearchSchema.h
+++ b/Engine/Plugins/Animation/PoseSearch/Source/Runtime/Public/PoseSearch/PoseSearchSchema.h
@@ -38,6 +38,7 @@ struct FPoseSearchRoledSkeleton
 
 	// Declared to be able to disable deprecation warnings
 	FPoseSearchRoledSkeleton& operator=(const FPoseSearchRoledSkeleton&);
+	FPoseSearchRoledSkeleton(const FPoseSearchRoledSkeleton&);
 
 	// Skeleton Reference for Motion Matching Database assets. Must be set to a compatible skeleton to the animation data in the database.
 	UPROPERTY(EditAnywhere, Category = "Schema")
