Explicitly declare default constructor since declaring copy constructor prevents the implicit
Export the defaulted copy constructor and assignment operator

diff --git a/Engine/Plugins/Animation/PoseSearch/Source/Runtime/Public/PoseSearch/PoseSearchSchema.h b/Engine/Plugins/Animation/PoseSearch/Source/Runtime/Public/PoseSearch/PoseSearchSchema.h
--- a/Engine/Plugins/Animation/PoseSearch/Source/Runtime/Public/PoseSearch/PoseSearchSchema.h
+++ b/Engine/Plugins/Animation/PoseSearch/Source/Runtime/Public/PoseSearch/PoseSearchSchema.h
@@ -37,8 +37,9 @@ struct FPoseSearchRoledSkeleton
 	GENERATED_BODY()
 
 	// Declared to be able to disable deprecation warnings
-	FPoseSearchRoledSkeleton& operator=(const FPoseSearchRoledSkeleton&);
-	FPoseSearchRoledSkeleton(const FPoseSearchRoledSkeleton&);
+	FPoseSearchRoledSkeleton() = default;
+	UE_API FPoseSearchRoledSkeleton(const FPoseSearchRoledSkeleton&);
+	UE_API FPoseSearchRoledSkeleton& operator=(const FPoseSearchRoledSkeleton&);
 
 	// Skeleton Reference for Motion Matching Database assets. Must be set to a compatible skeleton to the animation data in the database.
 	UPROPERTY(EditAnywhere, Category = "Schema")
