[PoseSearch]
* Added missing include to fix non unity build error

diff --git a/Engine/Plugins/Animation/PoseSearch/Source/Runtime/Private/PoseSearchLibrary.cpp b/Engine/Plugins/Animation/PoseSearch/Source/Runtime/Private/PoseSearchLibrary.cpp
--- a/Engine/Plugins/Animation/PoseSearch/Source/Runtime/Private/PoseSearchLibrary.cpp
+++ b/Engine/Plugins/Animation/PoseSearch/Source/Runtime/Private/PoseSearchLibrary.cpp
@@ -14,6 +14,7 @@
 #include "Animation/AnimSubsystem_Tag.h"
 #include "Animation/BlendSpace.h"
 #include "Animation/AnimTrace.h"
+#include "Components/SkeletalMeshComponent.h"
 #include "GameFramework/Character.h"
 #include "PoseSearch/AnimNode_MotionMatching.h"
 #include "PoseSearch/AnimNode_PoseSearchHistoryCollector.h"
