PoseSearch - removing 5.5 deprecations

diff --git a/Engine/Plugins/Animation/PoseSearch/Source/Runtime/Private/PoseSearchLibrary.cpp b/Engine/Plugins/Animation/PoseSearch/Source/Runtime/Private/PoseSearchLibrary.cpp
--- a/Engine/Plugins/Animation/PoseSearch/Source/Runtime/Private/PoseSearchLibrary.cpp
+++ b/Engine/Plugins/Animation/PoseSearch/Source/Runtime/Private/PoseSearchLibrary.cpp
@@ -1368,28 +1368,4 @@ const FAnimNode_PoseSearchHistoryCollector_Base* UPoseSearchLibrary::FindPoseHis
 	return nullptr;
 }
 
-///////////////////////////////////////////////////////////////////////////////////////////
-// Begin deprecated signatures
-UE::PoseSearch::FSearchResult UPoseSearchLibrary::MotionMatch(
-	const TArrayView<UAnimInstance*> AnimInstances,
-	const TArrayView<const UE::PoseSearch::FRole> Roles,
-	const TArrayView<const UE::PoseSearch::IPoseHistory*> PoseHistories,
-	const TArrayView<const UObject*> AssetsToSearch,
-	const FPoseSearchContinuingProperties& ContinuingProperties,
-	const FPoseSearchFutureProperties& Future)
-{
-	TArray<const UObject*, TInlineAllocator<UE::PoseSearch::PreallocatedRolesNum>> AnimContexts;
-	AnimContexts.Reserve(AnimInstances.Num());
-
-	for (UAnimInstance* AnimInstance : AnimInstances)
-	{
-		AnimContexts.Add(AnimInstance);
-	}
-
-	return MotionMatch(AnimContexts, Roles, PoseHistories, AssetsToSearch, ContinuingProperties, Future, FPoseSearchEvent());
-}
-
-// End deprecated signatures
-///////////////////////////////////////////////////////////////////////////////////////////
-
 #undef LOCTEXT_NAMESPACE
diff --git a/Engine/Plugins/Animation/PoseSearch/Source/Runtime/Private/PoseSearchSchema.cpp b/Engine/Plugins/Animation/PoseSearch/Source/Runtime/Private/PoseSearchSchema.cpp
--- a/Engine/Plugins/Animation/PoseSearch/Source/Runtime/Private/PoseSearchSchema.cpp
+++ b/Engine/Plugins/Animation/PoseSearch/Source/Runtime/Private/PoseSearchSchema.cpp
@@ -378,27 +378,6 @@ void UPoseSearchSchema::PostLoad()
 {
 	Super::PostLoad();
 
-	PRAGMA_DISABLE_DEPRECATION_WARNINGS
-	if (Skeleton_DEPRECATED)
-	{
-		Skeletons.AddDefaulted_GetRef().Skeleton = Skeleton_DEPRECATED;
-		Skeleton_DEPRECATED = nullptr;
-	}
-
-	if (MirrorDataTable_DEPRECATED)
-	{
-		if (Skeletons.IsEmpty())
-		{
-			Skeletons.AddDefaulted_GetRef().MirrorDataTable = MirrorDataTable_DEPRECATED;
-		}
-		else
-		{
-			Skeletons[0].MirrorDataTable = MirrorDataTable_DEPRECATED;
-		}
-		MirrorDataTable_DEPRECATED = nullptr;
-	}
-	PRAGMA_ENABLE_DEPRECATION_WARNINGS
-
 	for (FPoseSearchRoledSkeleton& Skeleton : Skeletons)
 	{
 		if (Skeleton.MirrorDataTable)
diff --git a/Engine/Plugins/Animation/PoseSearch/Source/Runtime/Private/PoseSearchTrajectoryLibrary.cpp b/Engine/Plugins/Animation/PoseSearch/Source/Runtime/Private/PoseSearchTrajectoryLibrary.cpp
--- a/Engine/Plugins/Animation/PoseSearch/Source/Runtime/Private/PoseSearchTrajectoryLibrary.cpp
+++ b/Engine/Plugins/Animation/PoseSearch/Source/Runtime/Private/PoseSearchTrajectoryLibrary.cpp
@@ -175,21 +175,6 @@ FVector FPoseSearchTrajectoryData::StepCharacterMovementGroundPrediction(
 
 PRAGMA_DISABLE_DEPRECATION_WARNINGS
 
-// Deprecated
-void UPoseSearchTrajectoryLibrary::InitTrajectorySamples(
-	FPoseSearchQueryTrajectory& Trajectory,
-	const FPoseSearchTrajectoryData& TrajectoryData,
-	const FPoseSearchTrajectoryData::FDerived& TrajectoryDataDerived,
-	const FPoseSearchTrajectoryData::FSampling& TrajectoryDataSampling,
-	float DeltaTime)
-{
-	InitTrajectorySamples(
-		Trajectory,
-		TrajectoryDataDerived.Position, TrajectoryDataDerived.Facing,
-		TrajectoryDataSampling, 
-		DeltaTime );
-}
-
 // Deprecated
 void UPoseSearchTrajectoryLibrary::InitTrajectorySamples(
 	FPoseSearchQueryTrajectory& Trajectory,
@@ -393,28 +378,11 @@ FVector UPoseSearchTrajectoryLibrary::RemapVectorMagnitudeWithCurve(
 	return Vector;
 }
 
-PRAGMA_DISABLE_DEPRECATION_WARNINGS
-
-// Deprecated
-void UPoseSearchTrajectoryLibrary::UpdatePrediction_SimulateCharacterMovement(
-	FPoseSearchQueryTrajectory& Trajectory,
-	const FPoseSearchTrajectoryData& TrajectoryData,
-	const FPoseSearchTrajectoryData::FDerived& TrajectoryDataDerived,
-	const FPoseSearchTrajectoryData::FSampling& TrajectoryDataSampling,
-	float DeltaTime)
-{
-	FTransformTrajectory TransformTrajectory = Trajectory;
-	UpdatePrediction_SimulateCharacterMovement(TransformTrajectory, TrajectoryData, TrajectoryDataDerived, TrajectoryDataSampling, DeltaTime);
-	Trajectory = TransformTrajectory;
-}
-
-PRAGMA_ENABLE_DEPRECATION_WARNINGS
-
 void UPoseSearchTrajectoryLibrary::UpdatePrediction_SimulateCharacterMovement(FTransformTrajectory& Trajectory,
 	const FPoseSearchTrajectoryData& TrajectoryData, const FPoseSearchTrajectoryData::FDerived& TrajectoryDataDerived,
 	const FPoseSearchTrajectoryData::FSampling& TrajectoryDataSampling, float DeltaTime)
 {
-		FVector CurrentPositionWS = TrajectoryDataDerived.Position;
+	FVector CurrentPositionWS = TrajectoryDataDerived.Position;
 	FVector CurrentVelocityWS = RemapVectorMagnitudeWithCurve(TrajectoryDataDerived.Velocity, TrajectoryData.bUseSpeedRemappingCurve, TrajectoryData.SpeedRemappingCurve);
 	FVector CurrentAccelerationWS = RemapVectorMagnitudeWithCurve(TrajectoryDataDerived.Acceleration, TrajectoryData.bUseAccelerationRemappingCurve, TrajectoryData.AccelerationRemappingCurve);
 
diff --git a/Engine/Plugins/Animation/PoseSearch/Source/Runtime/Public/PoseSearch/PoseSearchLibrary.h b/Engine/Plugins/Animation/PoseSearch/Source/Runtime/Public/PoseSearch/PoseSearchLibrary.h
--- a/Engine/Plugins/Animation/PoseSearch/Source/Runtime/Public/PoseSearch/PoseSearchLibrary.h
+++ b/Engine/Plugins/Animation/PoseSearch/Source/Runtime/Public/PoseSearch/PoseSearchLibrary.h
@@ -180,20 +180,6 @@ class UPoseSearchLibrary : public UBlueprintFunctionLibrary
 		TraceMotionMatching(SearchContext, CurrentState.CurrentSearchResult, CurrentState.ElapsedPoseSearchTime, DeltaTime, bSearch, CurrentState.WantedPlayRate, EPoseSearchInterruptMode::DoNotInterrupt);
 		PRAGMA_ENABLE_DEPRECATION_WARNINGS
 	}
-	
-	UE_DEPRECATED(5.5, "Use TraceMotionMatching with different signature instead.")
-	static void TraceMotionMatching(
-		UE::PoseSearch::FSearchContext& SearchContext,
-		const UE::PoseSearch::FSearchResult& CurrentResult,
-		float ElapsedPoseSearchTime,
-		const FTransform& RootMotionTransformDelta,
-		float DeltaTime,
-		bool bSearch,
-		float RecordingTime)
-	{
-		TraceMotionMatching(SearchContext, CurrentResult, ElapsedPoseSearchTime, DeltaTime, bSearch, 1.f, EPoseSearchInterruptMode::DoNotInterrupt);
-	}
-
 #endif // UE_POSE_SEARCH_TRACE_ENABLED
 
 public:
@@ -332,15 +318,6 @@ public:
 		const float DesiredPermutationTimeOffset,
 		const FPoseSearchEvent& EventToSearch);
 
-	UE_DEPRECATED(5.5, "Use other MotionMatch signatures instead")
-	static UE_API UE::PoseSearch::FSearchResult MotionMatch(
-		const TArrayView<UAnimInstance*> AnimInstances,
-		const TArrayView<const UE::PoseSearch::FRole> Roles,
-		const TArrayView<const UE::PoseSearch::IPoseHistory*> PoseHistories, 
-		const TArrayView<const UObject*> AssetsToSearch,
-		const FPoseSearchContinuingProperties& ContinuingProperties,
-		const FPoseSearchFutureProperties& Future);
-
 	static UE_API const FAnimNode_PoseSearchHistoryCollector_Base* FindPoseHistoryNode(
 		const FName PoseHistoryName,
 		const UAnimInstance* AnimInstance);
diff --git a/Engine/Plugins/Animation/PoseSearch/Source/Runtime/Public/PoseSearch/PoseSearchSchema.h b/Engine/Plugins/Animation/PoseSearch/Source/Runtime/Public/PoseSearch/PoseSearchSchema.h
--- a/Engine/Plugins/Animation/PoseSearch/Source/Runtime/Public/PoseSearch/PoseSearchSchema.h
+++ b/Engine/Plugins/Animation/PoseSearch/Source/Runtime/Public/PoseSearch/PoseSearchSchema.h
@@ -72,10 +72,6 @@ class UPoseSearchSchema : public UDataAsset
 	GENERATED_BODY()
 
 public:
-	UE_DEPRECATED(5.5, "Use Skeletons[i].Skeleton instead")
-	UPROPERTY()
-	TObjectPtr<USkeleton> Skeleton_DEPRECATED;
-	
 	// The update rate at which we sample the animation data in the database. The higher the SampleRate the more refined your searches will be, but the more memory will be required
 	UPROPERTY(EditAnywhere, Category = "Schema", meta = (DisplayPriority = 3, ClampMin = "1", ClampMax = "240"))
 	int32 SampleRate = 30;
@@ -93,10 +89,6 @@ private:
 	TArray<TObjectPtr<UPoseSearchFeatureChannel>> FinalizedChannels;
 
 public:
-	UE_DEPRECATED(5.5, "Use Skeletons[i].MirrorDataTable instead")
-	UPROPERTY()
-	TObjectPtr<UMirrorDataTable> MirrorDataTable_DEPRECATED;
-
 #if WITH_EDITORONLY_DATA
 	// Type of operation performed to the full pose features dataset
 	UPROPERTY(EditAnywhere, Category = "Schema", meta = (DisplayPriority = 2))
diff --git a/Engine/Plugins/Animation/PoseSearch/Source/Runtime/Public/PoseSearch/PoseSearchTrajectoryLibrary.h b/Engine/Plugins/Animation/PoseSearch/Source/Runtime/Public/PoseSearch/PoseSearchTrajectoryLibrary.h
--- a/Engine/Plugins/Animation/PoseSearch/Source/Runtime/Public/PoseSearch/PoseSearchTrajectoryLibrary.h
+++ b/Engine/Plugins/Animation/PoseSearch/Source/Runtime/Public/PoseSearch/PoseSearchTrajectoryLibrary.h
@@ -114,9 +114,6 @@ public:
 	///////////////////////////////////////////////////////////////////////////
 	
 	PRAGMA_DISABLE_DEPRECATION_WARNINGS
-	UE_DEPRECATED(5.5, "Use the InitTrajectorySamples function that takes DefaultPosition and DefaultFacing arguments instead of TrajectoryDataDerived")
-	static UE_API void InitTrajectorySamples(FPoseSearchQueryTrajectory& Trajectory, const FPoseSearchTrajectoryData& TrajectoryData, const FPoseSearchTrajectoryData::FDerived& TrajectoryDataDerived, const FPoseSearchTrajectoryData::FSampling& TrajectoryDataSampling, float DeltaTime);
-
 	// Initialize history and predicted samples based on sampling settings and a default state
 	UE_DEPRECATED(5.6, "Use the InitTrajectorySamples function that doesn't take FPoseSearchTrajectoryData")
 	static UE_API void InitTrajectorySamples(FPoseSearchQueryTrajectory& Trajectory, const FPoseSearchTrajectoryData& TrajectoryData, FVector DefaultPosition, FQuat DefaultFacing, const FPoseSearchTrajectoryData::FSampling& TrajectoryDataSampling, float DeltaTime);
@@ -125,9 +122,6 @@ public:
 	UE_DEPRECATED(5.6, "Use the IniTrajctorySamples function that takes in FTransformTrajectory instead of FPoseSearchQueryTrajectory.")
 	static UE_API void InitTrajectorySamples(FPoseSearchQueryTrajectory& Trajectory, FVector DefaultPosition, FQuat DefaultFacing, const FPoseSearchTrajectoryData::FSampling& TrajectoryDataSampling, float DeltaTime);
 
-	UE_DEPRECATED(5.5, "Use the UpdateHistory_TransformHistory function that takes CurrentPosition and CurrentVelocity arguments instead of TrajectoryDataDerived")
-	static UE_API void UpdateHistory_TransformHistory(FPoseSearchQueryTrajectory& Trajectory, const FPoseSearchTrajectoryData& TrajectoryData, const FPoseSearchTrajectoryData::FDerived& TrajectoryDataDerived, const FPoseSearchTrajectoryData::FSampling& TrajectoryDataSampling, float DeltaTime);
-
 	UE_DEPRECATED(5.6, "Use the UpdateHistory_TransformHistory function that doesn't take FPoseSearchTrajectoryData")
 	static UE_API void UpdateHistory_TransformHistory(FPoseSearchQueryTrajectory& Trajectory, const FPoseSearchTrajectoryData& TrajectoryData, FVector CurrentPosition, FVector CurrentVelocity, const FPoseSearchTrajectoryData::FSampling& TrajectoryDataSampling, float DeltaTime);
 	
