PoseSearch - adding more check(s) to enforce APIs contracts

diff --git a/Engine/Plugins/Animation/PoseSearch/Source/Runtime/Private/PoseSearchAssetIndexer.cpp b/Engine/Plugins/Animation/PoseSearch/Source/Runtime/Private/PoseSearchAssetIndexer.cpp
--- a/Engine/Plugins/Animation/PoseSearch/Source/Runtime/Private/PoseSearchAssetIndexer.cpp
+++ b/Engine/Plugins/Animation/PoseSearch/Source/Runtime/Private/PoseSearchAssetIndexer.cpp
@@ -855,11 +855,16 @@ int32 FAssetIndexer::GetNumIndexedPoses() const
 
 int32 FAssetIndexer::GetVectorIdx(int32 SampleIdx) const
 {
-	return SampleIdx - GetBeginSampleIdx();
+	check(SampleIdx >= 0);
+	const int32 BeginSampleIdx = GetBeginSampleIdx();
+	check(BeginSampleIdx >= 0);
+	check(SampleIdx >= BeginSampleIdx);
+	return SampleIdx - BeginSampleIdx;
 }
 
 TArrayView<float> FAssetIndexer::GetPoseVector(int32 SampleIdx) const
 {
+	check(Schema.SchemaCardinality > 0);
 	return MakeArrayView(&FeatureVectorTable[GetVectorIdx(SampleIdx) * Schema.SchemaCardinality], Schema.SchemaCardinality);
 }
 
diff --git a/Engine/Plugins/Animation/PoseSearch/Source/Runtime/Private/PoseSearchContext.cpp b/Engine/Plugins/Animation/PoseSearch/Source/Runtime/Private/PoseSearchContext.cpp
--- a/Engine/Plugins/Animation/PoseSearch/Source/Runtime/Private/PoseSearchContext.cpp
+++ b/Engine/Plugins/Animation/PoseSearch/Source/Runtime/Private/PoseSearchContext.cpp
@@ -1104,7 +1104,12 @@ TArrayView<float> FSearchContext::EditFeatureVector()
 {
 	// CachedQueries.Last is the query we're building 
 	check(!CachedQueries.IsEmpty());
-	return CachedQueries.Last().EditValues();
+	FCachedQuery& CachedQuery = CachedQueries.Last();
+	TArrayView<float> FeatureVector = CachedQuery.EditValues();
+	check(CachedQuery.GetSchema());
+	check(FeatureVector.GetData());
+	check(CachedQuery.GetSchema()->SchemaCardinality == FeatureVector.Num());
+	return FeatureVector;
 }
 
 const UPoseSearchFeatureChannel* FSearchContext::GetCachedChannelData(uint32 ChannelUniqueIdentifier, const UPoseSearchFeatureChannel* Channel, TConstArrayView<float>& CachedChannelData)
diff --git a/Engine/Plugins/Animation/PoseSearch/Source/Runtime/Private/PoseSearchFeatureChannel.cpp b/Engine/Plugins/Animation/PoseSearch/Source/Runtime/Private/PoseSearchFeatureChannel.cpp
--- a/Engine/Plugins/Animation/PoseSearch/Source/Runtime/Private/PoseSearchFeatureChannel.cpp
+++ b/Engine/Plugins/Animation/PoseSearch/Source/Runtime/Private/PoseSearchFeatureChannel.cpp
@@ -28,24 +28,29 @@ int32 FFeatureVectorHelper::GetVectorCardinality(EComponentStrippingVector Compo
 
 void FFeatureVectorHelper::EncodeVector(TArrayView<float> Values, int32 DataOffset, const FVector& Vector, EComponentStrippingVector ComponentStrippingVector, bool bNormalize)
 {
+	check(Values.GetData());
+
 	if (bNormalize)
 	{
 		FVector NormalizedVector = Vector;
 		switch (ComponentStrippingVector)
 		{
 		case EComponentStrippingVector::None:
+			check(Values.Num() >= DataOffset + 2);
 			NormalizedVector.Normalize();
 			Values[DataOffset + 0] = NormalizedVector.X;
 			Values[DataOffset + 1] = NormalizedVector.Y;
 			Values[DataOffset + 2] = NormalizedVector.Z;
 			break;
 		case EComponentStrippingVector::StripXY:
+			check(Values.Num() >= DataOffset + 0);			
 			NormalizedVector.X = 0.f;
 			NormalizedVector.Y = 0.f;
 			NormalizedVector.Normalize();
 			Values[DataOffset + 0] = NormalizedVector.Z;
 			break;
 		case EComponentStrippingVector::StripZ:
+			check(Values.Num() >= DataOffset + 1);
 			NormalizedVector.Z = 0.f;
 			NormalizedVector.Normalize();
 			Values[DataOffset + 0] = NormalizedVector.X;
@@ -61,14 +66,17 @@ void FFeatureVectorHelper::EncodeVector(TArrayView<float> Values, int32 DataOffs
 		switch (ComponentStrippingVector)
 		{
 		case EComponentStrippingVector::None:
+			check(Values.Num() >= DataOffset + 2);
 			Values[DataOffset + 0] = Vector.X;
 			Values[DataOffset + 1] = Vector.Y;
 			Values[DataOffset + 2] = Vector.Z;
 			break;
 		case EComponentStrippingVector::StripXY:
+			check(Values.Num() >= DataOffset + 0);
 			Values[DataOffset + 0] = Vector.Z;
 			break;
 		case EComponentStrippingVector::StripZ:
+			check(Values.Num() >= DataOffset + 1);
 			Values[DataOffset + 0] = Vector.X;
 			Values[DataOffset + 1] = Vector.Y;
 			break;
@@ -84,10 +92,13 @@ FVector FFeatureVectorHelper::DecodeVector(TConstArrayView<float> Values, int32
 	switch (ComponentStrippingVector)
 	{
 	case EComponentStrippingVector::None:
+		check(Values.Num() >= DataOffset + 2);
 		return FVector(Values[DataOffset + 0], Values[DataOffset + 1], Values[DataOffset + 2]);
 	case EComponentStrippingVector::StripXY:
+		check(Values.Num() >= DataOffset + 0);
 		return FVector(0, 0, Values[DataOffset + 0]);
 	case EComponentStrippingVector::StripZ:
+		check(Values.Num() >= DataOffset + 1);
 		return FVector(Values[DataOffset + 0], Values[DataOffset + 1], 0);
 	default:
 		checkNoEntry();
@@ -97,27 +108,32 @@ FVector FFeatureVectorHelper::DecodeVector(TConstArrayView<float> Values, int32
 
 void FFeatureVectorHelper::EncodeVector2D(TArrayView<float> Values, int32 DataOffset, const FVector2D& Vector2D)
 {
+	check(Values.Num() >= DataOffset + 1);
 	Values[DataOffset + 0] = Vector2D.X;
 	Values[DataOffset + 1] = Vector2D.Y;
 }
 
 FVector2D FFeatureVectorHelper::DecodeVector2D(TConstArrayView<float> Values, int32 DataOffset)
 {
+	check(Values.Num() >= DataOffset + 1);
 	return FVector2D(Values[DataOffset + 0], Values[DataOffset + 1]);
 }
 
 void FFeatureVectorHelper::EncodeFloat(TArrayView<float> Values, int32 DataOffset, const float Value)
 {
+	check(Values.Num() >= DataOffset + 0);
 	Values[DataOffset + 0] = Value;
 }
 
 float FFeatureVectorHelper::DecodeFloat(TConstArrayView<float> Values, int32 DataOffset)
 {
+	check(Values.Num() >= DataOffset + 0);
 	return Values[DataOffset];
 }
 
 void FFeatureVectorHelper::Copy(TArrayView<float> Values, int32 DataOffset, int32 DataCardinality, TConstArrayView<float> OriginValues)
 {
+	check(Values.GetData() && OriginValues.GetData());
 	check(Values.Num() == OriginValues.Num() && DataOffset + DataCardinality <= Values.Num());
 
 	float* RESTRICT ValuesData = Values.GetData() + DataOffset;
diff --git a/Engine/Plugins/Animation/PoseSearch/Source/Runtime/Public/PoseSearch/PoseSearchIndex.h b/Engine/Plugins/Animation/PoseSearch/Source/Runtime/Public/PoseSearch/PoseSearchIndex.h
--- a/Engine/Plugins/Animation/PoseSearch/Source/Runtime/Public/PoseSearch/PoseSearchIndex.h
+++ b/Engine/Plugins/Animation/PoseSearch/Source/Runtime/Public/PoseSearch/PoseSearchIndex.h
@@ -540,9 +540,6 @@ struct FSearchIndexBase
 
 	bool operator==(const FSearchIndexBase& Other) const;
 
-	// Experimental, this feature might be removed without warning, not for production use
-	TArray<int32> GetPosesWithEvent(const FGameplayTag& GameplayTag) const;
-
 	friend FArchive& operator<<(FArchive& Ar, FSearchIndexBase& Index);
 };
 
