PoseSearch - removing invalid error logging in UPoseSearchLibrary::UpdateMotionMatchingState

diff --git a/Engine/Plugins/Animation/PoseSearch/Source/Runtime/Private/PoseSearchLibrary.cpp b/Engine/Plugins/Animation/PoseSearch/Source/Runtime/Private/PoseSearchLibrary.cpp
--- a/Engine/Plugins/Animation/PoseSearch/Source/Runtime/Private/PoseSearchLibrary.cpp
+++ b/Engine/Plugins/Animation/PoseSearch/Source/Runtime/Private/PoseSearchLibrary.cpp
@@ -731,28 +731,6 @@ void UPoseSearchLibrary::UpdateMotionMatchingState(
 	else
 	{
 		InternalSearchResult.InitFrom(InOutMotionMatchingState.SearchResult);
-
-#if DO_CHECK
-		if (InternalSearchResult.PoseIdx != INDEX_NONE)
-		{
-			FPoseSearchBlueprintResult TestSearchResult;
-			TestSearchResult.InitFrom(InternalSearchResult, InOutMotionMatchingState.SearchResult.WantedPlayRate);
-
-			if (InOutMotionMatchingState.SearchResult.SelectedAnim != TestSearchResult.SelectedAnim ||
-				InOutMotionMatchingState.SearchResult.SelectedTime != TestSearchResult.SelectedTime ||
-				InOutMotionMatchingState.SearchResult.bIsContinuingPoseSearch != TestSearchResult.bIsContinuingPoseSearch ||
-				InOutMotionMatchingState.SearchResult.WantedPlayRate != TestSearchResult.WantedPlayRate ||
-				InOutMotionMatchingState.SearchResult.bLoop != TestSearchResult.bLoop ||
-				InOutMotionMatchingState.SearchResult.bIsMirrored != TestSearchResult.bIsMirrored ||
-				InOutMotionMatchingState.SearchResult.BlendParameters != TestSearchResult.BlendParameters ||
-				InOutMotionMatchingState.SearchResult.SelectedDatabase != TestSearchResult.SelectedDatabase ||
-				InOutMotionMatchingState.SearchResult.SearchCost != TestSearchResult.SearchCost ||
-				InOutMotionMatchingState.SearchResult.bIsInteraction != TestSearchResult.bIsInteraction)
-			{
-				UE_LOG(LogPoseSearch, Error, TEXT("Error converting FPoseSearchBlueprintResult to UE::PoseSearch::FSearchResult!"));
-			}
-		}
-#endif // DO_CHECK
 	}
 
 	const FPoseSearchEvent PlayRateOverriddenEvent = EventToSearch.GetPlayRateOverriddenEvent(PlayRate);
