PoseSearch - fix for MM tracing during throttled or skipped searches

diff --git a/Engine/Plugins/Animation/PoseSearch/Source/Runtime/Private/PoseSearchDatabase.cpp b/Engine/Plugins/Animation/PoseSearch/Source/Runtime/Private/PoseSearchDatabase.cpp
--- a/Engine/Plugins/Animation/PoseSearch/Source/Runtime/Private/PoseSearchDatabase.cpp
+++ b/Engine/Plugins/Animation/PoseSearch/Source/Runtime/Private/PoseSearchDatabase.cpp
@@ -1658,6 +1658,13 @@ UE::PoseSearch::FSearchResult UPoseSearchDatabase::SearchContinuingPose(UE::Pose
 
 #if UE_POSE_SEARCH_TRACE_ENABLED
 		SearchContext.Track(this, ContinuingPoseIdx, EPoseCandidateFlags::Valid_ContinuingPose, Result.PoseCost);
+#endif // UE_POSE_SEARCH_TRACE_ENABLED
+	}
+	else
+	{
+#if UE_POSE_SEARCH_TRACE_ENABLED
+		// in case we skipped the search, we still have to track we requested to evaluate this database, so we keep track of this
+		SearchContext.Track(this);
 #endif // UE_POSE_SEARCH_TRACE_ENABLED
 	}
 
diff --git a/Engine/Plugins/Animation/PoseSearch/Source/Runtime/Private/PoseSearchLibrary.cpp b/Engine/Plugins/Animation/PoseSearch/Source/Runtime/Private/PoseSearchLibrary.cpp
--- a/Engine/Plugins/Animation/PoseSearch/Source/Runtime/Private/PoseSearchLibrary.cpp
+++ b/Engine/Plugins/Animation/PoseSearch/Source/Runtime/Private/PoseSearchLibrary.cpp
@@ -838,6 +838,22 @@ void UPoseSearchLibrary::UpdateMotionMatchingState(
 		// @todo: for blendspaces ElapsedPoseSearchTime should be incremented by a normalized DeltaTime, also didn't we already synchronized the result???
 		InOutMotionMatchingState.ElapsedPoseSearchTime += DeltaTime;
 		InternalSearchResult.bIsContinuingPoseSearch = true;
+
+#if UE_POSE_SEARCH_TRACE_ENABLED
+		// in case we skipped the search, we still have to track we would have requested to evaluate Databases and CurrentResultDatabase
+		for (const TObjectPtr<const UPoseSearchDatabase>& Database : Databases)
+		{
+			if (Database)
+			{
+				SearchContext.Track(Database);
+			}
+		}
+
+		if (CurrentResultDatabase)
+		{
+			SearchContext.Track(CurrentResultDatabase);
+		}
+#endif // UE_POSE_SEARCH_TRACE_ENABLED
 	}
 
 	// @todo: consider moving this into if (bSearch) to avoid calling SearchContext.GetCachedQuery if no search is required
