PoseSearch - UpdateMotionMatchingState deprecation WIP in favor of MotionMatch

diff --git a/Engine/Plugins/Animation/PoseSearch/Source/Runtime/Private/AnimNode_MotionMatching.cpp b/Engine/Plugins/Animation/PoseSearch/Source/Runtime/Private/AnimNode_MotionMatching.cpp
--- a/Engine/Plugins/Animation/PoseSearch/Source/Runtime/Private/AnimNode_MotionMatching.cpp
+++ b/Engine/Plugins/Animation/PoseSearch/Source/Runtime/Private/AnimNode_MotionMatching.cpp
@@ -46,6 +46,9 @@ static FAutoConsoleVariableRef CVarAnimNodeMotionMatchingDrawInfoVerbose(TEXT("a
 static float GVarAnimNodeMotionMatchingDrawInfoHeight = 50.f;
 static FAutoConsoleVariableRef CVarAnimNodeMotionMatchingDrawInfoHeight(TEXT("a.AnimNode.MotionMatching.DebugDrawInfoHeight"), GVarAnimNodeMotionMatchingDrawInfoHeight, TEXT("Vertical offset for DebugDrawInfo"));
 
+static int32 GVarAnimNodeMotionMatchingCompareUpdateMotionMatchingStateAgainstMotionMatch = 0;
+static FAutoConsoleVariableRef CVarAnimNodeMotionMatchingCompareUpdateMotionMatchingStateAgainstMotionMatch(TEXT("a.AnimNode.MotionMatching.CompareUpdateMotionMatchingStateAgainstMotionMatch"), GVarAnimNodeMotionMatchingCompareUpdateMotionMatchingStateAgainstMotionMatch, TEXT("Compare UpdateMotionMatchingState Against MotionMatch. 0:disabled 1:warn 2:debug restore MM state"));
+
 static int32 GVarAnimNodeMotionMatchingPlayRateEnabled = UE::Private::EPlayRateState::Enabled;
 static FAutoConsoleVariableRef CVarAnimNodeMotionMatchingPlayRateEnabled(
 	TEXT("a.AnimNode.MotionMatching.DebugPlayRateEnabled"), 
@@ -157,13 +160,77 @@ void FAnimNode_MotionMatching::UpdateAssetPlayer(const FAnimationUpdateContext&
 	PoseSearchPlayRate = bPoseSearchPlayRateEnabled ? PoseSearchPlayRate : FFloatInterval(1.f, 1.f);
 #endif // ENABLE_ANIM_DEBUG
 
-	FChooserEvaluationContext EvaluationContext(Context.AnimInstanceProxy->GetAnimInstanceObject());
+	const float DeltaTime = Context.GetDeltaTime();
+
+#if ENABLE_ANIM_DEBUG
+	// debug code to validate UPoseSearchLibrary::MotionMatch against UPoseSearchLibrary::UpdateMotionMatchingState to deprecate the latter in the near future,
+	// and expose FSearchContext to be able to perform multiple searches with the same cached queries (to speed up MM interactions)
+	FMotionMatchingState MotionMatchingStateBackup;
+	FMotionMatchingState MotionMatchingStateCopy;
+	if (GVarAnimNodeMotionMatchingCompareUpdateMotionMatchingStateAgainstMotionMatch)
+	{
+		MotionMatchingStateBackup = MotionMatchingState;
+		MotionMatchingStateCopy = MotionMatchingState;
+
+		const FPoseSearchEvent PlayRateOverriddenEvent = EventToSearch.GetPlayRateOverriddenEvent(PoseSearchPlayRate);
+		FSearchContext SearchContext(0.f, &MotionMatchingStateCopy.PoseIndicesHistory, PoseJumpThresholdTime, PlayRateOverriddenEvent);
+		FChooserEvaluationContext ChooserEvaluationContext(const_cast<UObject*>(Context.AnimInstanceProxy->GetAnimInstanceObject()));
+		SearchContext.AddRole(GetCommonDefaultRole(DatabasesToSearch), &ChooserEvaluationContext, PoseHistory);
+
+		// If we can't advance or enough time has elapsed since the last pose jump then search
+		const bool bCanAdvance = SearchContext.GetContinuingPoseSearchResult().PoseIdx != INDEX_NONE;
+		const bool bSearch = !bCanAdvance || (MotionMatchingStateCopy.ElapsedPoseSearchTime >= SearchThrottleTime);
+		FSearchResult SearchResult;
+		if (bSearch)
+		{
+			MotionMatchingStateCopy.ElapsedPoseSearchTime = 0.f;
+			
+			TArrayView<const UObject*> AssetsToSearch((const UObject**)FMemory_Alloca(DatabasesToSearch.Num() * sizeof(UObject*)), (int32)DatabasesToSearch.Num());
+			for (int32 DatabaseToSearchIndex = 0; DatabaseToSearchIndex < DatabasesToSearch.Num(); ++DatabaseToSearchIndex)
+			{
+				AssetsToSearch[DatabaseToSearchIndex] = DatabasesToSearch[DatabaseToSearchIndex].Get();
+			}
+
+			FPoseSearchContinuingProperties ContinuingProperties;
+			ContinuingProperties.InitFrom(MotionMatchingStateCopy.SearchResult, NextUpdateInterruptMode);
+			SearchResult = UPoseSearchLibrary::MotionMatch(SearchContext, AssetsToSearch, ContinuingProperties);
+		}
+		else
+		{
+			MotionMatchingStateCopy.ElapsedPoseSearchTime += DeltaTime;
+
+			// if we didn't perform a search we carry on with the continuing pose search result
+			SearchResult = SearchContext.GetContinuingPoseSearchResult();
+			SearchResult.bIsContinuingPoseSearch = true;
 
+	#if UE_POSE_SEARCH_TRACE_ENABLED
+			// in case we skipped the search, we still have to track we would have requested to evaluate DatabasesToSearch and SearchResult.Database
+			for (const TObjectPtr<const UPoseSearchDatabase>& DatabaseToSearch : DatabasesToSearch)
+			{
+				SearchContext.Track(DatabaseToSearch.Get());
+			}
+
+			SearchContext.Track(SearchResult.Database.Get());
+	#endif // UE_POSE_SEARCH_TRACE_ENABLED
+		}
+
+		const float WantedPlayRate = CalculateWantedPlayRate(SearchResult, SearchContext, PoseSearchPlayRate, PoseHistory ? PoseHistory->GetTrajectorySpeedMultiplier() : 1.f, EventToSearch);
+		MotionMatchingStateCopy.SearchResult.InitFrom(SearchResult, WantedPlayRate);
+		MotionMatchingStateCopy.PoseIndicesHistory.Update(SearchResult, DeltaTime, PoseReselectHistory);
+
+		PRAGMA_DISABLE_DEPRECATION_WARNINGS
+		MotionMatchingStateCopy.CurrentSearchResult = SearchResult;
+		MotionMatchingStateCopy.WantedPlayRate = WantedPlayRate;
+		PRAGMA_ENABLE_DEPRECATION_WARNINGS
+	}
+#endif // ENABLE_ANIM_DEBUG
+
+	FChooserEvaluationContext EvaluationContext(Context.AnimInstanceProxy->GetAnimInstanceObject());
 	UPoseSearchLibrary::UpdateMotionMatchingState(
 		&EvaluationContext,
 		PoseHistory,
 		DatabasesToSearch,
-		Context.GetDeltaTime(),
+		DeltaTime,
 		PoseJumpThresholdTime,
 		PoseReselectHistory,
 		bShouldSearch ? SearchThrottleTime : UE_BIG_NUMBER,
@@ -181,6 +248,52 @@ void FAnimNode_MotionMatching::UpdateAssetPlayer(const FAnimationUpdateContext&
 		, EventToSearch
 	);
 
+#if ENABLE_ANIM_DEBUG
+	if (GVarAnimNodeMotionMatchingCompareUpdateMotionMatchingStateAgainstMotionMatch)
+	{
+		bool bUpdateMotionMatchingDiffersFromMotionMatch = false;
+		PRAGMA_DISABLE_DEPRECATION_WARNINGS
+		if (MotionMatchingState.CurrentSearchResult != MotionMatchingStateCopy.CurrentSearchResult)
+		{
+			UE_LOG(LogPoseSearch, Warning, TEXT("MotionMatchingState.CurrentSearchResult != MotionMatchingStateCopy.CurrentSearchResult"));
+			bUpdateMotionMatchingDiffersFromMotionMatch = true;
+		}
+
+		if (MotionMatchingState.SearchResult != MotionMatchingStateCopy.SearchResult)
+		{
+			UE_LOG(LogPoseSearch, Warning, TEXT("MotionMatchingState.SearchResult != MotionMatchingStateCopy.SearchResult"));
+			bUpdateMotionMatchingDiffersFromMotionMatch = true;
+		}
+
+		if (MotionMatchingState.ElapsedPoseSearchTime != MotionMatchingStateCopy.ElapsedPoseSearchTime)
+		{
+			UE_LOG(LogPoseSearch, Warning, TEXT("MotionMatchingState.ElapsedPoseSearchTime != MotionMatchingStateCopy.ElapsedPoseSearchTime"));
+			bUpdateMotionMatchingDiffersFromMotionMatch = true;
+		}
+
+		if (MotionMatchingState.WantedPlayRate != MotionMatchingStateCopy.WantedPlayRate)
+		{
+			UE_LOG(LogPoseSearch, Warning, TEXT("MotionMatchingState.CurrentSearchResult != MotionMatchingStateCopy.CurrentSearchResult"));
+			bUpdateMotionMatchingDiffersFromMotionMatch = true;
+		}
+
+		if (MotionMatchingState.PoseIndicesHistory != MotionMatchingStateCopy.PoseIndicesHistory)
+		{
+			UE_LOG(LogPoseSearch, Warning, TEXT("MotionMatchingState.PoseIndicesHistory != MotionMatchingStateCopy.PoseIndicesHistory"));
+			bUpdateMotionMatchingDiffersFromMotionMatch = true;
+		}
+		PRAGMA_ENABLE_DEPRECATION_WARNINGS
+
+		if (GVarAnimNodeMotionMatchingCompareUpdateMotionMatchingStateAgainstMotionMatch > 1 && bUpdateMotionMatchingDiffersFromMotionMatch)
+		{
+			MotionMatchingStateCopy = MotionMatchingStateBackup;
+			MotionMatchingState = MotionMatchingStateBackup;
+
+			// the idea is that now you can run MM searches again and analyze why they differs
+		}
+	}
+#endif // ENABLE_ANIM_DEBUG
+
 	const bool bJumpToPose = MotionMatchingState.SearchResult.SelectedAnim && !MotionMatchingState.SearchResult.bIsContinuingPoseSearch;
 
 	PRAGMA_DISABLE_DEPRECATION_WARNINGS
diff --git a/Engine/Plugins/Animation/PoseSearch/Source/Runtime/Private/PoseSearchContext.cpp b/Engine/Plugins/Animation/PoseSearch/Source/Runtime/Private/PoseSearchContext.cpp
--- a/Engine/Plugins/Animation/PoseSearch/Source/Runtime/Private/PoseSearchContext.cpp
+++ b/Engine/Plugins/Animation/PoseSearch/Source/Runtime/Private/PoseSearchContext.cpp
@@ -720,12 +720,22 @@ FCachedQuery::FCachedQuery(const UPoseSearchSchema* InSchema)
 // FSearchContext
 FSearchContext::FSearchContext(float InDesiredPermutationTimeOffset, const FPoseIndicesHistory* InPoseIndicesHistory,
 	const FSearchResult& InCurrentResult, const FFloatInterval& InPoseJumpThresholdTime, bool bInUseCachedChannelData)
+PRAGMA_DISABLE_DEPRECATION_WARNINGS
 : FSearchContext(InDesiredPermutationTimeOffset, InPoseIndicesHistory, InCurrentResult, InPoseJumpThresholdTime, FPoseSearchEvent())
+PRAGMA_ENABLE_DEPRECATION_WARNINGS
+{
+}
+
+FSearchContext::FSearchContext(float InDesiredPermutationTimeOffset, const FPoseIndicesHistory* InPoseIndicesHistory,
+	const FFloatInterval& InPoseJumpThresholdTime, const FPoseSearchEvent& InEventToSearch)
+PRAGMA_DISABLE_DEPRECATION_WARNINGS
+: FSearchContext(InDesiredPermutationTimeOffset, InPoseIndicesHistory, FSearchResult(), InPoseJumpThresholdTime, FPoseSearchEvent())
+PRAGMA_ENABLE_DEPRECATION_WARNINGS
 {
 }
 
 FSearchContext::FSearchContext(float InDesiredPermutationTimeOffset, const FPoseIndicesHistory* InPoseIndicesHistory,
-	const FSearchResult& InCurrentResult, const FFloatInterval& InPoseJumpThresholdTime, const FPoseSearchEvent& InEventToSearch)
+	const FSearchResult& InContinuingPoseSearchResult, const FFloatInterval& InPoseJumpThresholdTime, const FPoseSearchEvent& InEventToSearch)
 : AnimContexts()
 , PoseHistories()
 , RoleToIndex()
@@ -733,11 +743,12 @@ FSearchContext::FSearchContext(float InDesiredPermutationTimeOffset, const FPose
 , EventToSearch(InEventToSearch)
 , DesiredPermutationTimeOffset(InDesiredPermutationTimeOffset)
 , PoseIndicesHistory(InPoseIndicesHistory)
-, CurrentResult(InCurrentResult)
+, ContinuingPoseSearchResult(FSearchResult())
+, ContinuingPoseValuesDatabasePoseIdx(FDatabasePoseIdx())
 , PoseJumpThresholdTime(InPoseJumpThresholdTime)
 , bUseCachedChannelData(false)
 {
-	UpdateCurrentResultPoseVector();
+	UpdateContinuingPoseSearchResult(InContinuingPoseSearchResult, InContinuingPoseSearchResult);
 
 #if UE_POSE_SEARCH_TRACE_ENABLED
 	BestPoseCandidatesMap.Reserve(16);
@@ -757,24 +768,56 @@ void FSearchContext::AddRole(const FRole& Role, FChooserEvaluationContext* AnimC
 
 void FSearchContext::UpdateCurrentResultPoseVector()
 {
-	if (CurrentResult.IsValid())
+	UpdateContinuingPoseValues();
+}
+
+void FSearchContext::UpdateContinuingPoseValues()
+{
+	if (ContinuingPoseValuesDatabasePoseIdx.IsValid())
 	{
-		const FSearchIndex& SearchIndex = CurrentResult.Database->GetSearchIndex();
+		const FSearchIndex& SearchIndex = ContinuingPoseValuesDatabasePoseIdx.Database->GetSearchIndex();
 		if (SearchIndex.IsValuesEmpty())
 		{
-			const int32 NumDimensions = CurrentResult.Database->Schema->SchemaCardinality;
-			CurrentResultPoseVectorData.SetNumUninitialized(NumDimensions, EAllowShrinking::No);
-			CurrentResultPoseVector = SearchIndex.GetReconstructedPoseValues(CurrentResult.PoseIdx, MakeArrayView(CurrentResultPoseVectorData).Slice(NumDimensions, NumDimensions));
+			const int32 NumDimensions = ContinuingPoseValuesDatabasePoseIdx.Database->Schema->SchemaCardinality;
+			ContinuingPoseValuesData.SetNumUninitialized(NumDimensions, EAllowShrinking::No);
+			ContinuingPoseValues = SearchIndex.GetReconstructedPoseValues(ContinuingPoseValuesDatabasePoseIdx.PoseIdx, MakeArrayView(ContinuingPoseValuesData).Slice(NumDimensions, NumDimensions));
 		}
 		else
 		{
-			CurrentResultPoseVector = SearchIndex.GetPoseValues(CurrentResult.PoseIdx);
+			ContinuingPoseValues = SearchIndex.GetPoseValues(ContinuingPoseValuesDatabasePoseIdx.PoseIdx);
 		}
 	}
 	else
 	{
-		CurrentResultPoseVector = TStackAlignedArray<float>();
+		ContinuingPoseValues = TStackAlignedArray<float>();
+	}
+}
+
+void FSearchContext::UpdateContinuingPoseSearchResult(const FSearchResult& InContinuingPoseSearchResult, const FDatabasePoseIdx& InContinuingPoseValuesDatabasePoseIdx)
+{
+#if DO_CHECK
+	// checking the validity of the continuing pose search values result, that has to be valid
+	// if the result used to gather continuing pose search properties is valid
+	if (InContinuingPoseSearchResult.IsValid())
+	{
+		check(InContinuingPoseValuesDatabasePoseIdx.IsValid());
+		check(InContinuingPoseSearchResult.Database->Schema == InContinuingPoseValuesDatabasePoseIdx.Database->Schema);
+	}
+#endif //DO_CHECK
+
+	const bool bNeedToUpdateContinuingPoseValues =
+		ContinuingPoseValuesDatabasePoseIdx.PoseIdx != InContinuingPoseValuesDatabasePoseIdx.PoseIdx ||
+		ContinuingPoseValuesDatabasePoseIdx.Database != InContinuingPoseValuesDatabasePoseIdx.Database;
+
+	ContinuingPoseSearchResult = InContinuingPoseSearchResult;
+	ContinuingPoseValuesDatabasePoseIdx = InContinuingPoseValuesDatabasePoseIdx;
+	
+	if (bNeedToUpdateContinuingPoseValues)
+	{
+		UpdateContinuingPoseValues();
 	}
+
+	check(!ContinuingPoseValuesDatabasePoseIdx.IsValid() || ContinuingPoseValues.Num() == ContinuingPoseValuesDatabasePoseIdx.Database->Schema->SchemaCardinality);
 }
 
 float FSearchContext::GetSampleCurveValue(float SampleTimeOffset, const FName& CurveName, const FRole& SampleRole)
@@ -1178,14 +1221,26 @@ TConstArrayView<float> FSearchContext::GetOrBuildQuery(const UPoseSearchSchema*
 
 bool FSearchContext::IsCurrentResultFromDatabase(const UPoseSearchDatabase* Database) const
 {
-	return CurrentResult.IsValid() && CurrentResult.Database == Database;
+	return IsContinuingPoseDatabase(Database);
+}
+
+bool FSearchContext::IsContinuingPoseDatabase(const UPoseSearchDatabase* Database) const
+{
+	return ContinuingPoseSearchResult.IsValid() && ContinuingPoseSearchResult.Database == Database;
 }
 
 bool FSearchContext::CanUseCurrentResult() const
+{
+	return CanUseContinuingPoseValues();
+}
+
+bool FSearchContext::CanUseContinuingPoseValues() const
 {
 	// CachedQueries.Last is the query we're building 
 	check(!CachedQueries.IsEmpty());
-	return CurrentResult.IsValid() && CurrentResult.Database->Schema == CachedQueries.Last().GetSchema();
+	const bool bCanUseContinuingPoseValues = ContinuingPoseSearchResult.IsValid() && ContinuingPoseSearchResult.Database->Schema == CachedQueries.Last().GetSchema();
+	check(!bCanUseContinuingPoseValues || ContinuingPoseValues.Num() == ContinuingPoseSearchResult.Database->Schema->SchemaCardinality);
+	return bCanUseContinuingPoseValues;
 }
 
 } // namespace UE::PoseSearch
\ No newline at end of file
diff --git a/Engine/Plugins/Animation/PoseSearch/Source/Runtime/Private/PoseSearchDatabase.cpp b/Engine/Plugins/Animation/PoseSearch/Source/Runtime/Private/PoseSearchDatabase.cpp
--- a/Engine/Plugins/Animation/PoseSearch/Source/Runtime/Private/PoseSearchDatabase.cpp
+++ b/Engine/Plugins/Animation/PoseSearch/Source/Runtime/Private/PoseSearchDatabase.cpp
@@ -1442,9 +1442,9 @@ void UPoseSearchDatabase::PopulateNonSelectableIdx(FNonSelectableIdx& NonSelecta
 #endif // UE_POSE_SEARCH_TRACE_ENABLED
 
 	NonSelectableIdx.Reset();
-	if (SearchContext.IsCurrentResultFromDatabase(this))
+	if (SearchContext.IsContinuingPoseDatabase(this))
 	{
-		if (const FSearchIndexAsset* CurrentIndexAsset = SearchContext.GetCurrentResult().GetSearchIndexAsset(true))
+		if (const FSearchIndexAsset* CurrentIndexAsset = SearchContext.GetContinuingPoseSearchResult().GetSearchIndexAsset(true))
 		{
 			if (CurrentIndexAsset->IsDisableReselection())
 			{
@@ -1472,9 +1472,9 @@ void UPoseSearchDatabase::PopulateNonSelectableIdx(FNonSelectableIdx& NonSelecta
 			}
 			else if (!FMath::IsNearlyEqual(SearchContext.GetPoseJumpThresholdTime().Min, SearchContext.GetPoseJumpThresholdTime().Max))
 			{
-				const int32 CurrentResultPoseIdx = SearchContext.GetCurrentResult().PoseIdx;
-				const int32 UnboundMinPoseIdx = CurrentResultPoseIdx + FMath::FloorToInt(SearchContext.GetPoseJumpThresholdTime().Min * Schema->SampleRate);
-				const int32 UnboundMaxPoseIdx = CurrentResultPoseIdx + FMath::CeilToInt(SearchContext.GetPoseJumpThresholdTime().Max * Schema->SampleRate);
+				const int32 ContinuingPoseIdx = SearchContext.GetContinuingPoseSearchResult().PoseIdx;
+				const int32 UnboundMinPoseIdx = ContinuingPoseIdx + FMath::FloorToInt(SearchContext.GetPoseJumpThresholdTime().Min * Schema->SampleRate);
+				const int32 UnboundMaxPoseIdx = ContinuingPoseIdx + FMath::CeilToInt(SearchContext.GetPoseJumpThresholdTime().Max * Schema->SampleRate);
 				const int32 CurrentIndexAssetFirstPoseIdx = CurrentIndexAsset->GetFirstPoseIdx();
 				const int32 CurrentIndexAssetNumPoses = CurrentIndexAsset->GetNumPoses();
 				const bool bIsLooping = CurrentIndexAsset->IsLooping();
@@ -1575,7 +1575,7 @@ UE::PoseSearch::FSearchResult UPoseSearchDatabase::SearchContinuingPose(UE::Pose
 
 	using namespace UE::PoseSearch;
 
-	check(SearchContext.GetCurrentResult().Database.Get() == this);
+	check(SearchContext.GetContinuingPoseSearchResult().Database.Get() == this);
 
 	FSearchResult Result;
 	Result.bIsContinuingPoseSearch = true;
@@ -1590,13 +1590,13 @@ UE::PoseSearch::FSearchResult UPoseSearchDatabase::SearchContinuingPose(UE::Pose
 
 	// extracting notifies from the database animation asset at time SampleTime to search for UAnimNotifyState_PoseSearchOverrideContinuingPoseCostBias eventually overriding the database ContinuingPoseCostBias
 	const FSearchIndex& SearchIndex = GetSearchIndex();
-	const int32 PoseIdx = SearchContext.GetCurrentResult().PoseIdx;
-	const FSearchIndexAsset& SearchIndexAsset = SearchIndex.GetAssetForPose(PoseIdx);
+	const int32 ContinuingPoseIdx = SearchContext.GetContinuingPoseSearchResult().PoseIdx;
+	const FSearchIndexAsset& SearchIndexAsset = SearchIndex.GetAssetForPose(ContinuingPoseIdx);
 	const FPoseSearchDatabaseAnimationAssetBase* DatabaseAnimationAssetBase = GetDatabaseAnimationAsset<FPoseSearchDatabaseAnimationAssetBase>(SearchIndexAsset);
 	check(DatabaseAnimationAssetBase);
 
 	float ContinuingPoseCostAddend = ContinuingPoseCostBias;
-	const float SampleTime = GetRealAssetTime(PoseIdx);
+	const float SampleTime = GetRealAssetTime(ContinuingPoseIdx);
 	FAnimNotifyContext PreAllocatedNotifyContext;
 	for (int32 RoleIndex = 0; RoleIndex < DatabaseAnimationAssetBase->GetNumRoles(); ++RoleIndex)
 	{
@@ -1634,13 +1634,12 @@ UE::PoseSearch::FSearchResult UPoseSearchDatabase::SearchContinuingPose(UE::Pose
 		// FMemory_Alloca is forced 16 bytes aligned
 		TArrayView<float> ReconstructedPoseValuesBuffer((float*)FMemory_Alloca(NumDimensions * sizeof(float)), NumDimensions);
 		check(IsAligned(ReconstructedPoseValuesBuffer.GetData(), alignof(VectorRegister4Float)));
-		const TConstArrayView<float> PoseValues = SearchIndex.IsValuesEmpty() ? SearchIndex.GetReconstructedPoseValues(PoseIdx, ReconstructedPoseValuesBuffer) : SearchIndex.GetPoseValues(PoseIdx);
+		const TConstArrayView<float> PoseValues = SearchIndex.IsValuesEmpty() ? SearchIndex.GetReconstructedPoseValues(ContinuingPoseIdx, ReconstructedPoseValuesBuffer) : SearchIndex.GetPoseValues(ContinuingPoseIdx);
 
 		// @todo: perhaps store the DynamicWeightsSqrt into the GetOrBuildQuery to share between the continuing pose and the full search
 		TArrayView<float> DynamicWeightsSqrtBuffer((float*)FMemory_Alloca(NumDimensions * sizeof(float)), NumDimensions);
 		const TConstArrayView<float> DynamicWeightsSqrt = CalculateDynamicWeightsSqrt(DynamicWeightsSqrtBuffer);
 
-		const int32 ContinuingPoseIdx = SearchContext.GetCurrentResult().PoseIdx;
 		// is the data padded at 16 bytes (and 16 bytes aligned by construction)?
 		if (NumDimensions % 4 == 0)
 		{
@@ -1652,8 +1651,8 @@ UE::PoseSearch::FSearchResult UPoseSearchDatabase::SearchContinuingPose(UE::Pose
 			Result.PoseCost = FPoseSearchCost(CompareFeatureVectors<false>(PoseValues, SearchContext.GetOrBuildQuery(Schema), DynamicWeightsSqrt), SearchIndex.PoseMetadata[ContinuingPoseIdx].GetCostAddend(), ContinuingPoseCostAddend, ContinuingInteractionCostAddend);
 		}
 
-		Result.AssetTime = SearchContext.GetCurrentResult().AssetTime;
-		Result.PoseIdx = PoseIdx;
+		Result.AssetTime = SearchContext.GetContinuingPoseSearchResult().AssetTime;
+		Result.PoseIdx = ContinuingPoseIdx;
 		Result.Database = this;
 
 #if UE_POSE_SEARCH_TRACE_ENABLED
diff --git a/Engine/Plugins/Animation/PoseSearch/Source/Runtime/Private/PoseSearchFeatureChannel_Curve.cpp b/Engine/Plugins/Animation/PoseSearch/Source/Runtime/Private/PoseSearchFeatureChannel_Curve.cpp
--- a/Engine/Plugins/Animation/PoseSearch/Source/Runtime/Private/PoseSearchFeatureChannel_Curve.cpp
+++ b/Engine/Plugins/Animation/PoseSearch/Source/Runtime/Private/PoseSearchFeatureChannel_Curve.cpp
@@ -82,13 +82,13 @@ void UPoseSearchFeatureChannel_Curve::BuildQuery(UE::PoseSearch::FSearchContext&
 		}
 	}
 
-	const bool bCanUseCurrentResult = SearchContext.CanUseCurrentResult();
-	const bool bSkip = InputQueryPose != EInputQueryPose::UseCharacterPose && bCanUseCurrentResult;
+	const bool bCanUseContinuingPoseValues = SearchContext.CanUseContinuingPoseValues();
+	const bool bSkip = InputQueryPose != EInputQueryPose::UseCharacterPose && bCanUseContinuingPoseValues;
 	if (bSkip || (!SearchContext.ArePoseHistoriesValid()))
 	{
-		if (bCanUseCurrentResult)
+		if (bCanUseContinuingPoseValues)
 		{
-			FFeatureVectorHelper::Copy(SearchContext.EditFeatureVector(), ChannelDataOffset, ChannelCardinality, SearchContext.GetCurrentResultPoseVector());
+			FFeatureVectorHelper::Copy(SearchContext.EditFeatureVector(), ChannelDataOffset, ChannelCardinality, SearchContext.GetContinuingPoseValues());
 			return;
 		}
 
diff --git a/Engine/Plugins/Animation/PoseSearch/Source/Runtime/Private/PoseSearchFeatureChannel_Distance.cpp b/Engine/Plugins/Animation/PoseSearch/Source/Runtime/Private/PoseSearchFeatureChannel_Distance.cpp
--- a/Engine/Plugins/Animation/PoseSearch/Source/Runtime/Private/PoseSearchFeatureChannel_Distance.cpp
+++ b/Engine/Plugins/Animation/PoseSearch/Source/Runtime/Private/PoseSearchFeatureChannel_Distance.cpp
@@ -125,14 +125,14 @@ void UPoseSearchFeatureChannel_Distance::BuildQuery(UE::PoseSearch::FSearchConte
 		}
 	}
 
-	const bool bCanUseCurrentResult = SearchContext.CanUseCurrentResult();
-	const bool bSkip = InputQueryPose != EInputQueryPose::UseCharacterPose && bCanUseCurrentResult && SampleRole == OriginRole;
+	const bool bCanUseContinuingPoseValues = SearchContext.CanUseContinuingPoseValues();
+	const bool bSkip = InputQueryPose != EInputQueryPose::UseCharacterPose && bCanUseContinuingPoseValues && SampleRole == OriginRole;
 	const bool bIsRootBone = SchemaBoneIdx == RootSchemaBoneIdx;
 	if (bSkip || (!SearchContext.ArePoseHistoriesValid() && !bIsRootBone))
 	{
-		if (bCanUseCurrentResult)
+		if (bCanUseContinuingPoseValues)
 		{
-			FFeatureVectorHelper::Copy(SearchContext.EditFeatureVector(), ChannelDataOffset, ChannelCardinality, SearchContext.GetCurrentResultPoseVector());
+			FFeatureVectorHelper::Copy(SearchContext.EditFeatureVector(), ChannelDataOffset, ChannelCardinality, SearchContext.GetContinuingPoseValues());
 			return;
 		}
 
diff --git a/Engine/Plugins/Animation/PoseSearch/Source/Runtime/Private/PoseSearchFeatureChannel_FilterCrashingLegs.cpp b/Engine/Plugins/Animation/PoseSearch/Source/Runtime/Private/PoseSearchFeatureChannel_FilterCrashingLegs.cpp
--- a/Engine/Plugins/Animation/PoseSearch/Source/Runtime/Private/PoseSearchFeatureChannel_FilterCrashingLegs.cpp
+++ b/Engine/Plugins/Animation/PoseSearch/Source/Runtime/Private/PoseSearchFeatureChannel_FilterCrashingLegs.cpp
@@ -95,13 +95,13 @@ void UPoseSearchFeatureChannel_FilterCrashingLegs::BuildQuery(UE::PoseSearch::FS
 	}
 
 	// trying to get the BuildQuery data from the continuing pose
-	const bool bCanUseCurrentResult = SearchContext.CanUseCurrentResult();
-	const bool bSkip = InputQueryPose != EInputQueryPose::UseCharacterPose && bCanUseCurrentResult;
+	const bool bCanUseContinuingPoseValues = SearchContext.CanUseContinuingPoseValues();
+	const bool bSkip = InputQueryPose != EInputQueryPose::UseCharacterPose && bCanUseContinuingPoseValues;
 	if (bSkip || !SearchContext.ArePoseHistoriesValid())
 	{
-		if (bCanUseCurrentResult)
+		if (bCanUseContinuingPoseValues)
 		{
-			FFeatureVectorHelper::Copy(SearchContext.EditFeatureVector(), ChannelDataOffset, ChannelCardinality, SearchContext.GetCurrentResultPoseVector());
+			FFeatureVectorHelper::Copy(SearchContext.EditFeatureVector(), ChannelDataOffset, ChannelCardinality, SearchContext.GetContinuingPoseValues());
 			return;
 		}
 
diff --git a/Engine/Plugins/Animation/PoseSearch/Source/Runtime/Private/PoseSearchFeatureChannel_Heading.cpp b/Engine/Plugins/Animation/PoseSearch/Source/Runtime/Private/PoseSearchFeatureChannel_Heading.cpp
--- a/Engine/Plugins/Animation/PoseSearch/Source/Runtime/Private/PoseSearchFeatureChannel_Heading.cpp
+++ b/Engine/Plugins/Animation/PoseSearch/Source/Runtime/Private/PoseSearchFeatureChannel_Heading.cpp
@@ -167,13 +167,13 @@ void UPoseSearchFeatureChannel_Heading::BuildQuery(UE::PoseSearch::FSearchContex
 		}
 	}
 
-	const bool bCanUseCurrentResult = SearchContext.CanUseCurrentResult();
-	const bool bSkip = InputQueryPose != EInputQueryPose::UseCharacterPose && bCanUseCurrentResult && SampleRole == OriginRole;
+	const bool bCanUseContinuingPoseValues = SearchContext.CanUseContinuingPoseValues();
+	const bool bSkip = InputQueryPose != EInputQueryPose::UseCharacterPose && bCanUseContinuingPoseValues && SampleRole == OriginRole;
 	if (bSkip || (!SearchContext.ArePoseHistoriesValid() && !bIsRootBone))
 	{
-		if (bCanUseCurrentResult)
+		if (bCanUseContinuingPoseValues)
 		{
-			FFeatureVectorHelper::Copy(SearchContext.EditFeatureVector(), ChannelDataOffset, ChannelCardinality, SearchContext.GetCurrentResultPoseVector());
+			FFeatureVectorHelper::Copy(SearchContext.EditFeatureVector(), ChannelDataOffset, ChannelCardinality, SearchContext.GetContinuingPoseValues());
 			return;
 		}
 
diff --git a/Engine/Plugins/Animation/PoseSearch/Source/Runtime/Private/PoseSearchFeatureChannel_Phase.cpp b/Engine/Plugins/Animation/PoseSearch/Source/Runtime/Private/PoseSearchFeatureChannel_Phase.cpp
--- a/Engine/Plugins/Animation/PoseSearch/Source/Runtime/Private/PoseSearchFeatureChannel_Phase.cpp
+++ b/Engine/Plugins/Animation/PoseSearch/Source/Runtime/Private/PoseSearchFeatureChannel_Phase.cpp
@@ -337,13 +337,13 @@ void UPoseSearchFeatureChannel_Phase::BuildQuery(UE::PoseSearch::FSearchContext&
 		}
 	}
 
-	const bool bCanUseCurrentResult = SearchContext.CanUseCurrentResult();
-	const bool bSkip = InputQueryPose != EInputQueryPose::UseCharacterPose && bCanUseCurrentResult;
+	const bool bCanUseContinuingPoseValues = SearchContext.CanUseContinuingPoseValues();
+	const bool bSkip = InputQueryPose != EInputQueryPose::UseCharacterPose && bCanUseContinuingPoseValues;
 	if (bSkip || !SearchContext.ArePoseHistoriesValid())
 	{
-		if (bCanUseCurrentResult)
+		if (bCanUseContinuingPoseValues)
 		{
-			FFeatureVectorHelper::Copy(SearchContext.EditFeatureVector(), ChannelDataOffset, ChannelCardinality, SearchContext.GetCurrentResultPoseVector());
+			FFeatureVectorHelper::Copy(SearchContext.EditFeatureVector(), ChannelDataOffset, ChannelCardinality, SearchContext.GetContinuingPoseValues());
 			return;
 		}
 		
diff --git a/Engine/Plugins/Animation/PoseSearch/Source/Runtime/Private/PoseSearchFeatureChannel_Position.cpp b/Engine/Plugins/Animation/PoseSearch/Source/Runtime/Private/PoseSearchFeatureChannel_Position.cpp
--- a/Engine/Plugins/Animation/PoseSearch/Source/Runtime/Private/PoseSearchFeatureChannel_Position.cpp
+++ b/Engine/Plugins/Animation/PoseSearch/Source/Runtime/Private/PoseSearchFeatureChannel_Position.cpp
@@ -169,14 +169,14 @@ void UPoseSearchFeatureChannel_Position::BuildQuery(UE::PoseSearch::FSearchConte
 		}
 	}
 
-	const bool bCanUseCurrentResult = SearchContext.CanUseCurrentResult();
-	const bool bSkip = InputQueryPose != EInputQueryPose::UseCharacterPose && bCanUseCurrentResult && SampleRole == OriginRole;
+	const bool bCanUseContinuingPoseValues = SearchContext.CanUseContinuingPoseValues();
+	const bool bSkip = InputQueryPose != EInputQueryPose::UseCharacterPose && bCanUseContinuingPoseValues && SampleRole == OriginRole;
 	const bool bIsRootBone = SchemaBoneIdx == RootSchemaBoneIdx;
 	if (bSkip || (!SearchContext.ArePoseHistoriesValid() && !bIsRootBone))
 	{
-		if (bCanUseCurrentResult)
+		if (bCanUseContinuingPoseValues)
 		{
-			FFeatureVectorHelper::Copy(SearchContext.EditFeatureVector(), ChannelDataOffset, ChannelCardinality, SearchContext.GetCurrentResultPoseVector());
+			FFeatureVectorHelper::Copy(SearchContext.EditFeatureVector(), ChannelDataOffset, ChannelCardinality, SearchContext.GetContinuingPoseValues());
 			return;
 		}
 
diff --git a/Engine/Plugins/Animation/PoseSearch/Source/Runtime/Private/PoseSearchFeatureChannel_Velocity.cpp b/Engine/Plugins/Animation/PoseSearch/Source/Runtime/Private/PoseSearchFeatureChannel_Velocity.cpp
--- a/Engine/Plugins/Animation/PoseSearch/Source/Runtime/Private/PoseSearchFeatureChannel_Velocity.cpp
+++ b/Engine/Plugins/Animation/PoseSearch/Source/Runtime/Private/PoseSearchFeatureChannel_Velocity.cpp
@@ -115,13 +115,13 @@ void UPoseSearchFeatureChannel_Velocity::BuildQuery(UE::PoseSearch::FSearchConte
 		}
 	}
 
-	const bool bCanUseCurrentResult = SearchContext.CanUseCurrentResult();
-	const bool bSkip = InputQueryPose != EInputQueryPose::UseCharacterPose && bCanUseCurrentResult && SampleRole == OriginRole;
+	const bool bCanUseContinuingPoseValues = SearchContext.CanUseContinuingPoseValues();
+	const bool bSkip = InputQueryPose != EInputQueryPose::UseCharacterPose && bCanUseContinuingPoseValues && SampleRole == OriginRole;
 	if (bSkip || (!SearchContext.ArePoseHistoriesValid() && !bIsRootBone))
 	{
-		if (bCanUseCurrentResult)
+		if (bCanUseContinuingPoseValues)
 		{
-			FFeatureVectorHelper::Copy(SearchContext.EditFeatureVector(), ChannelDataOffset, ChannelCardinality, SearchContext.GetCurrentResultPoseVector());
+			FFeatureVectorHelper::Copy(SearchContext.EditFeatureVector(), ChannelDataOffset, ChannelCardinality, SearchContext.GetContinuingPoseValues());
 			return;
 		}
 
diff --git a/Engine/Plugins/Animation/PoseSearch/Source/Runtime/Private/PoseSearchHistory.cpp b/Engine/Plugins/Animation/PoseSearch/Source/Runtime/Private/PoseSearchHistory.cpp
--- a/Engine/Plugins/Animation/PoseSearch/Source/Runtime/Private/PoseSearchHistory.cpp
+++ b/Engine/Plugins/Animation/PoseSearch/Source/Runtime/Private/PoseSearchHistory.cpp
@@ -1367,4 +1367,10 @@ void FPoseIndicesHistory::Update(const FSearchResult& SearchResult, float DeltaT
 	}
 }
 
+
+bool FPoseIndicesHistory::operator==(const FPoseIndicesHistory& Other) const
+{
+	return IndexToTime.OrderIndependentCompareEqual(Other.IndexToTime);
+}
+
 } // namespace UE::PoseSearch
diff --git a/Engine/Plugins/Animation/PoseSearch/Source/Runtime/Private/PoseSearchLibrary.cpp b/Engine/Plugins/Animation/PoseSearch/Source/Runtime/Private/PoseSearchLibrary.cpp
--- a/Engine/Plugins/Animation/PoseSearch/Source/Runtime/Private/PoseSearchLibrary.cpp
+++ b/Engine/Plugins/Animation/PoseSearch/Source/Runtime/Private/PoseSearchLibrary.cpp
@@ -42,25 +42,71 @@ static FAutoConsoleVariableRef CVarAnimMotionMatchDrawMatchEnable(TEXT("a.Motion
 
 namespace UE::PoseSearch
 {
-	// budgeting some stack allocations for simple use cases. bigger requests of AnimationAssets containing UAnimNotifyState_PoseSearchBranchIn 
-	// referencing multiple databases will default to a slower TMemStackAllocator (that hides heap allocations)
-	enum { MAX_STACK_ALLOCATED_ANIMATIONS = 16 };
-	enum { MAX_STACK_ALLOCATED_SETS = 2 };
-	typedef	TArray<const UObject*, TInlineAllocator<MAX_STACK_ALLOCATED_ANIMATIONS, TMemStackAllocator<>>> FAssetsToSearch;
-	// an empty FAssetsToSearch associated to Database means we need to search ALL the assets
-	typedef TMap<const UPoseSearchDatabase*, FAssetsToSearch, TInlineSetAllocator<MAX_STACK_ALLOCATED_SETS, TMemStackSetAllocator<>>> FAssetsToSearchPerDatabaseMap;
-	typedef TPair<const UPoseSearchDatabase*, FAssetsToSearch> TAssetsToSearchPerDatabasePair;
-	typedef TMap<const UPoseSearchDatabase*, FSearchResult, TInlineSetAllocator<MAX_STACK_ALLOCATED_SETS, TMemStackSetAllocator<>>> FReconstructedPreviousSearchBestResultMap;
+	// an empty FAssetsToSearch in any of the FAssetsToSearchPerDatabasePair entries means we need to search ALL the assets for the associated Database
+	typedef	TArray<const UObject*, TInlineAllocator<8, TMemStackAllocator<>>> FAssetsToSearch;
+	typedef TPair<const UPoseSearchDatabase*, FAssetsToSearch> FAssetsToSearchPerDatabasePair;
+	struct FAssetsToSearchPerDatabase : public TArray<FAssetsToSearchPerDatabasePair, TInlineAllocator<8, TMemStackAllocator<>>>
+	{
+		FAssetsToSearchPerDatabasePair* Find(const UPoseSearchDatabase* Database)
+		{
+			return FindByPredicate([Database](const FAssetsToSearchPerDatabasePair& Pair) { return Pair.Key == Database; });
+		}
+
+		bool Contains(const UPoseSearchDatabase* Database) const
+		{
+			return nullptr != FindByPredicate([Database](const FAssetsToSearchPerDatabasePair& Pair) { return Pair.Key == Database; });
+		}
+	};
+	
+	// used to cache the continuing pose search results
+	struct FCachedContinuingPoseSearchResults : public TMap<const UObject*, FSearchResult, TInlineSetAllocator<16, TMemStackSetAllocator<>>>
+	{
+		const FSearchResult& CheckedAdd(const UObject* Object, const FSearchResult& SearchResult)
+		{
+			check(Object);
+			check(!Find(Object));
+
+			FSearchResult& NewSearchResult = Add(Object);
+			NewSearchResult = SearchResult;
+			return NewSearchResult;
+		}
+
+		const FSearchResult& FindOrAdd(const UObject* Object, const FSearchResult& SearchResult)
+		{
+			check(Object);
+			if (const FSearchResult* FoundSearchResult = Find(Object))
+			{
+				return *FoundSearchResult;
+			}
+
+			FSearchResult& NewSearchResult = Add(Object);
+			NewSearchResult = SearchResult;
+			return NewSearchResult;
+		}
+
+		const FSearchResult& FindOrDefault(const UObject* Object) const
+		{
+			check(Object);
+			if (const FSearchResult* FoundSearchResult = Find(Object))
+			{
+				return *FoundSearchResult;
+			}
+
+			static const FSearchResult DefaultSearchResult;
+			return DefaultSearchResult;
+		}
+	};
 
 	// this function adds AssetToSearch to the search of Database
 	// returns bAsyncBuildIndexInProgress
-	static bool AddToSearchForDatabase(FAssetsToSearchPerDatabaseMap& AssetsToSearchPerDatabaseMap, const UObject* AssetToSearch, const UPoseSearchDatabase* Database, bool bContainsIsMandatory)
+	static bool AddToSearchForDatabase(FAssetsToSearchPerDatabase& AssetsToSearchPerDatabase, const UObject* AssetToSearch, const UPoseSearchDatabase* Database, bool bContainsIsMandatory)
 	{
-		FAssetsToSearch* AssetsToSearch = AssetsToSearchPerDatabaseMap.Find(Database);
+		// making sure AssetToSearch is not a databases! later on we could add support for nested databases, but currently we don't support that
+		check(Cast<const UPoseSearchDatabase>(AssetToSearch) == nullptr);
 
 #if WITH_EDITOR
-		// no need to check if Database is indexing if found into AssetsToSearchPerDatabaseMap, since it already passed RequestAsyncBuildIndex successfully in a previous AddToSearchForDatabase call
-		if (!AssetsToSearch && (EAsyncBuildIndexResult::Success != FAsyncPoseSearchDatabasesManagement::RequestAsyncBuildIndex(Database, ERequestAsyncBuildFlag::ContinueRequest)))
+		// no need to check if Database is indexing if found into AssetsToSearchPerDatabase, since it already passed RequestAsyncBuildIndex successfully in a previous AddToSearchForDatabase call
+		if (EAsyncBuildIndexResult::Success != FAsyncPoseSearchDatabasesManagement::RequestAsyncBuildIndex(Database, ERequestAsyncBuildFlag::ContinueRequest))
 		{
 			// database is still indexing... moving on
 			return true;
@@ -76,21 +122,21 @@ namespace UE::PoseSearch
 			return false;
 		}
 
-		// making sure AssetToSearch is not a databases! later on we could add support for nested databases, but currently we don't support that
-		check(Cast<const UPoseSearchDatabase>(AssetToSearch) == nullptr);
-
-		if (AssetsToSearch)
+		if (FAssetsToSearchPerDatabasePair* AssetsToSearchPerDatabasePair = AssetsToSearchPerDatabase.Find(Database))
 		{
 			// an empty FAssetsToSearch associated to Database means we need to search ALL the assets, so we don't need to add this AssetToSearch
-			if (!AssetsToSearch->IsEmpty())
+			FAssetsToSearch& AssetsToSearch = AssetsToSearchPerDatabasePair->Value;
+			if (!AssetsToSearch.IsEmpty())
 			{
-				AssetsToSearch->AddUnique(AssetToSearch);
+				AssetsToSearch.AddUnique(AssetToSearch);
 			}
 		}
 		else
 		{
 			// no need to AddUnique since it's the first one
-			AssetsToSearchPerDatabaseMap.Add(Database).Add(AssetToSearch);
+			FAssetsToSearchPerDatabasePair& NewAssetsToSearchPerDatabasePair = AssetsToSearchPerDatabase.AddDefaulted_GetRef();
+			NewAssetsToSearchPerDatabasePair.Key = Database;
+			NewAssetsToSearchPerDatabasePair.Value.Add(AssetToSearch);
 		}
 
 		return false;
@@ -100,96 +146,126 @@ namespace UE::PoseSearch
 	// if AssetToSearch is a database search it ALL,
 	// if it's a sequence containing UAnimNotifyState_PoseSearchBranchIn, we add to the search of the database UAnimNotifyState_PoseSearchBranchIn::Database the asset AssetToSearch
 	// returns bAsyncBuildIndexInProgress
-	static bool AddToSearch(FAssetsToSearchPerDatabaseMap& AssetsToSearchPerDatabaseMap, const UObject* AssetToSearch)
+	static bool AddToSearch(FAssetsToSearchPerDatabase& AssetsToSearchPerDatabase, const UObject* AssetToSearch, bool bUsePoseSearchBranchIn)
 	{
-		bool bAsyncBuildIndexInProgress = false;
-		if (const UAnimSequenceBase* SequenceBase = Cast<const UAnimSequenceBase>(AssetToSearch))
+		if (bUsePoseSearchBranchIn)
 		{
-			for (const FAnimNotifyEvent& NotifyEvent : SequenceBase->Notifies)
+			if (const UAnimSequenceBase* SequenceBase = Cast<const UAnimSequenceBase>(AssetToSearch))
 			{
-				if (const UAnimNotifyState_PoseSearchBranchIn* PoseSearchBranchIn = Cast<UAnimNotifyState_PoseSearchBranchIn>(NotifyEvent.NotifyStateClass))
+				bool bAsyncBuildIndexInProgress = false;
+				for (const FAnimNotifyEvent& NotifyEvent : SequenceBase->Notifies)
 				{
-					if (!PoseSearchBranchIn->Database)
+					if (const UAnimNotifyState_PoseSearchBranchIn* PoseSearchBranchIn = Cast<UAnimNotifyState_PoseSearchBranchIn>(NotifyEvent.NotifyStateClass))
 					{
-						UE_LOG(LogPoseSearch, Error, TEXT("improperly setup UAnimNotifyState_PoseSearchBranchIn with null Database in %s"), *SequenceBase->GetName());
-						continue;
-					}
-					
-					// we just skip indexing databases to keep the experience as smooth as possible
-					if (AddToSearchForDatabase(AssetsToSearchPerDatabaseMap, SequenceBase, PoseSearchBranchIn->Database, true))
-					{
-						bAsyncBuildIndexInProgress = true;
+						if (!PoseSearchBranchIn->Database)
+						{
+							UE_LOG(LogPoseSearch, Error, TEXT("improperly setup UAnimNotifyState_PoseSearchBranchIn with null Database in %s"), *SequenceBase->GetName());
+							continue;
+						}
+
+						// we just skip indexing databases to keep the experience as smooth as possible
+						if (AddToSearchForDatabase(AssetsToSearchPerDatabase, SequenceBase, PoseSearchBranchIn->Database, true))
+						{
+							bAsyncBuildIndexInProgress = true;
+						}
 					}
 				}
+				return bAsyncBuildIndexInProgress;
 			}
 		}
-		else if (const UPoseSearchDatabase* Database = Cast<UPoseSearchDatabase>(AssetToSearch))
+
+		if (const UPoseSearchDatabase* Database = Cast<UPoseSearchDatabase>(AssetToSearch))
 		{
-			// we already added Database to AssetsToSearchPerDatabaseMap, so it already successfully passed RequestAsyncBuildIndex
-			if (FAssetsToSearch* AssetsToSearch = AssetsToSearchPerDatabaseMap.Find(Database))
-			{
-				// an empty FAssetsToSearch associated to Database means we need to search ALL the assets
-				AssetsToSearch->Reset();
-			}
-			else
 #if WITH_EDITOR
 			if (EAsyncBuildIndexResult::Success != FAsyncPoseSearchDatabasesManagement::RequestAsyncBuildIndex(Database, ERequestAsyncBuildFlag::ContinueRequest))
 			{
-				bAsyncBuildIndexInProgress = true;
+				return true;
 			}
-			else
 #endif // WITH_EDITOR
+
+			// we already added Database to AssetsToSearchPerDatabase, so it already successfully passed RequestAsyncBuildIndex
+			if (FAssetsToSearchPerDatabasePair* AssetsToSearchPerDatabasePair = AssetsToSearchPerDatabase.Find(Database))
+			{
+				// an empty FAssetsToSearch associated to Database means we need to search ALL the assets
+				FAssetsToSearch& AssetsToSearch = AssetsToSearchPerDatabasePair->Value;
+				AssetsToSearch.Reset();
+			}
+			else
 			{
 				// an empty FAssetsToSearch associated to Database means we need to search ALL the assets
-				AssetsToSearchPerDatabaseMap.Add(Database);
+				AssetsToSearchPerDatabase.AddDefaulted_GetRef().Key = Database;
 			}
+			return false;
 		}
 
-		return bAsyncBuildIndexInProgress;
+		return false;
 	}
 
-	static void PopulateContinuingPoseSearches(const UObject* PlayingAnimationAsset, const TArrayView<const UObject*> AssetsToSearch, FSearchContext& SearchContext, FAssetsToSearchPerDatabaseMap& ContinuingPoseAssetsToSearchPerDatabaseMap)
+	static void PopulateContinuingPoseSearches(const FPoseSearchContinuingProperties& ContinuingProperties, const TArrayView<const UObject*> AssetsToSearch, FSearchContext& SearchContext, FAssetsToSearchPerDatabase& ContinuingPoseAssetsToSearchPerDatabase, bool bUsePoseSearchBranchIn)
 	{
-		// checking if PlayingAnimationAsset has an associated database
-		if (AddToSearch(ContinuingPoseAssetsToSearchPerDatabaseMap, PlayingAnimationAsset))
+		if (const UObject* PlayingAnimationAsset = ContinuingProperties.PlayingAsset.Get())
 		{
+			// checking if any of the AssetsToSearch (databases) or ContinuingProperties.PlayingAssetDatabase contain PlayingAnimationAsset
+			if (ContinuingProperties.PlayingAssetDatabase)
+			{
+				if (AddToSearchForDatabase(ContinuingPoseAssetsToSearchPerDatabase, PlayingAnimationAsset, ContinuingProperties.PlayingAssetDatabase, false))
+				{
 #if WITH_EDITOR
-			SearchContext.SetAsyncBuildIndexInProgress();
+					SearchContext.SetAsyncBuildIndexInProgress();
 #endif // WITH_EDITOR
-		}
-		
-		// checking if any of the AssetsToSearch (database) contains PlayingAnimationAsset
-		for (const UObject* AssetToSearch : AssetsToSearch)
-		{
-			if (const UPoseSearchDatabase* Database = Cast<UPoseSearchDatabase>(AssetToSearch))
+				}
+			}
+
+			for (const UObject* AssetToSearch : AssetsToSearch)
 			{
-				// since it cannot be a database we can directly add it to ContinuingPoseAssetsToSearchPerDatabaseMap
-				if (AddToSearchForDatabase(ContinuingPoseAssetsToSearchPerDatabaseMap, PlayingAnimationAsset, Database, false))
+				if (const UPoseSearchDatabase* Database = Cast<UPoseSearchDatabase>(AssetToSearch))
 				{
+					// since it cannot be a database we can directly add it to ContinuingPoseAssetsToSearchPerDatabase
+					if (AddToSearchForDatabase(ContinuingPoseAssetsToSearchPerDatabase, PlayingAnimationAsset, Database, false))
+					{
 #if WITH_EDITOR
-					SearchContext.SetAsyncBuildIndexInProgress();
+						SearchContext.SetAsyncBuildIndexInProgress();
 #endif // WITH_EDITOR
+					}
 				}
 			}
+
+			// checking if PlayingAnimationAsset has an associated database
+			if (AddToSearch(ContinuingPoseAssetsToSearchPerDatabase, PlayingAnimationAsset, bUsePoseSearchBranchIn))
+			{
+#if WITH_EDITOR
+				SearchContext.SetAsyncBuildIndexInProgress();
+#endif // WITH_EDITOR
+			}
 		}
 	}
 	
-	static void PopulateSearches(const TArrayView<const UObject*> AssetsToSearch, FSearchContext& SearchContext, FAssetsToSearchPerDatabaseMap& AssetsToSearchPerDatabaseMap)
+	static void PopulateSearches(const TArrayView<const UObject*> AssetsToSearch, FSearchContext& SearchContext, FAssetsToSearchPerDatabase& AssetsToSearchPerDatabase, bool bUsePoseSearchBranchIn)
 	{
-		if (!AssetsToSearch.IsEmpty())
+		for (const UObject* AssetToSearch : AssetsToSearch)
 		{
-			for (const UObject* AssetToSearch : AssetsToSearch)
+			if (AddToSearch(AssetsToSearchPerDatabase, AssetToSearch, bUsePoseSearchBranchIn))
 			{
-				if (AddToSearch(AssetsToSearchPerDatabaseMap, AssetToSearch))
-				{
 #if WITH_EDITOR
-					SearchContext.SetAsyncBuildIndexInProgress();
+				SearchContext.SetAsyncBuildIndexInProgress();
 #endif // WITH_EDITOR
-				}
 			}
 		}
 	}
 
+	// @todo: refine this logic. Currently if AssetsToSearch contains ONLY UPoseSearchDatabase we don't have to look for other databases referenced by other assets UAnimNotifyState_PoseSearchBranchIn(s)
+	bool ShouldUsePoseSearchBranchIn(const TArrayView<const UObject*> AssetsToSearch)
+	{
+		for (const UObject* AssetToSearch : AssetsToSearch)
+		{
+			if (!Cast<UPoseSearchDatabase>(AssetToSearch))
+			{
+				return true;
+			}
+		}
+		return false;
+	}
+
 	template <typename DatabasesContainer>
 	static bool IsForceInterrupt(EPoseSearchInterruptMode InterruptMode, const UPoseSearchDatabase* CurrentResultDatabase, const DatabasesContainer& Databases)
 	{
@@ -264,7 +340,7 @@ namespace UE::PoseSearch
 		return false;
 	}
 
-	static FRole GetCommonDefaultRole(const TConstArrayView<TObjectPtr<const UPoseSearchDatabase>> Databases)
+	FRole GetCommonDefaultRole(const TConstArrayView<TObjectPtr<const UPoseSearchDatabase>> Databases)
 	{
 		FRole Role = DefaultRole;
 
@@ -299,21 +375,21 @@ namespace UE::PoseSearch
 		return Role;
 	}
 
-	static float CalculateWantedPlayRate(const UE::PoseSearch::FSearchResult& CurrentSearchResult, const UE::PoseSearch::FSearchContext& SearchContext, const FFloatInterval& PlayRate, float TrajectorySpeedMultiplier, const FPoseSearchEvent& EventToSearch)
+	float CalculateWantedPlayRate(const FSearchResult& SearchResult, const FSearchContext& SearchContext, const FFloatInterval& PlayRate, float TrajectorySpeedMultiplier, const FPoseSearchEvent& EventToSearch)
 	{
 		float WantedPlayRate = 1.f;
 
-		if (CurrentSearchResult.IsValid())
+		if (SearchResult.IsValid())
 		{
-			if (CurrentSearchResult.IsEventSearchResult())
+			if (SearchResult.IsEventSearchResult())
 			{
-				// checking if CurrentSearchResult.EventPoseIdx is part of the EventToSearch.EventTag.
+				// checking if SearchResult.EventPoseIdx is part of the EventToSearch.EventTag.
 				// If not, it's an event from a continuing pose search that hasn't been interrupted,
 				// so we keep the previously calculated WantedPlayRate
-				const bool bIsEventSearchFromTag = CurrentSearchResult.IsEventSearchFromTag(EventToSearch.EventTag);
+				const bool bIsEventSearchFromTag = SearchResult.IsEventSearchFromTag(EventToSearch.EventTag);
 				if (bIsEventSearchFromTag)
 				{
-					const float TimeToEvent = CurrentSearchResult.CalculateTimeToEvent();
+					const float TimeToEvent = SearchResult.CalculateTimeToEvent();
 					if (TimeToEvent > UE_KINDA_SMALL_NUMBER && EventToSearch.TimeToEvent > UE_KINDA_SMALL_NUMBER)
 					{
 						// EventToSearch.TimeToEvent is the desired time to event, and TimeToEvent is the actually current time to event. we calculate WantedPlayRate as ratio between the two
@@ -329,12 +405,12 @@ namespace UE::PoseSearch
 			}
 			else if (!FMath::IsNearlyEqual(PlayRate.Min, PlayRate.Max, UE_KINDA_SMALL_NUMBER))
 			{
-				TConstArrayView<float> QueryData = SearchContext.GetCachedQuery(CurrentSearchResult.Database->Schema);
+				TConstArrayView<float> QueryData = SearchContext.GetCachedQuery(SearchResult.Database->Schema);
 				if (!QueryData.IsEmpty())
 				{
-					if (const UPoseSearchFeatureChannel_Trajectory* TrajectoryChannel = CurrentSearchResult.Database->Schema->FindFirstChannelOfType<UPoseSearchFeatureChannel_Trajectory>())
+					if (const UPoseSearchFeatureChannel_Trajectory* TrajectoryChannel = SearchResult.Database->Schema->FindFirstChannelOfType<UPoseSearchFeatureChannel_Trajectory>())
 					{
-						TConstArrayView<float> ResultData = CurrentSearchResult.Database->GetSearchIndex().GetPoseValues(CurrentSearchResult.PoseIdx);
+						TConstArrayView<float> ResultData = SearchResult.Database->GetSearchIndex().GetPoseValues(SearchResult.PoseIdx);
 						const float EstimatedSpeedRatio = TrajectoryChannel->GetEstimatedSpeedRatio(QueryData, ResultData);
 
 						WantedPlayRate = FMath::Clamp(EstimatedSpeedRatio, PlayRate.Min, PlayRate.Max);
@@ -343,7 +419,7 @@ namespace UE::PoseSearch
 					{
 						UE_LOG(LogPoseSearch, Warning,
 							TEXT("Couldn't update the WantedPlayRate in CalculateWantedPlayRate, because Schema '%s' couldn't find a UPoseSearchFeatureChannel_Trajectory channel"),
-							*GetNameSafe(CurrentSearchResult.Database->Schema));
+							*GetNameSafe(SearchResult.Database->Schema));
 					}
 				}
 			}
@@ -416,6 +492,22 @@ PRAGMA_DISABLE_DEPRECATION_WARNINGS
 PRAGMA_ENABLE_DEPRECATION_WARNINGS
 }
 
+//////////////////////////////////////////////////////////////////////////
+// FPoseSearchContinuingProperties
+
+void FPoseSearchContinuingProperties::InitFrom(const FPoseSearchBlueprintResult& SearchResult, EPoseSearchInterruptMode InInterruptMode)
+{
+	PlayingAsset = SearchResult.SelectedAnim;
+	PlayingAssetAccumulatedTime = SearchResult.SelectedTime;
+	bIsPlayingAssetMirrored = SearchResult.bIsMirrored;
+	PlayingAssetBlendParameters = SearchResult.BlendParameters;
+	InterruptMode = InInterruptMode;
+	PlayingAssetDatabase = SearchResult.SelectedDatabase;
+}
+
+//////////////////////////////////////////////////////////////////////////
+// UPoseSearchLibrary
+
 #if UE_POSE_SEARCH_TRACE_ENABLED
 void UPoseSearchLibrary::TraceMotionMatching(
 		UE::PoseSearch::FSearchContext& SearchContext,
@@ -721,7 +813,9 @@ void UPoseSearchLibrary::UpdateMotionMatchingState(
 
 	FMemMark Mark(FMemStack::Get());
 
-	UE::PoseSearch::FSearchResult InternalSearchResult;
+	const FPoseSearchEvent PlayRateOverriddenEvent = EventToSearch.GetPlayRateOverriddenEvent(PlayRate);
+	FSearchContext SearchContext(0.f, &InOutMotionMatchingState.PoseIndicesHistory, PoseJumpThresholdTime, PlayRateOverriddenEvent);
+	SearchContext.AddRole(GetCommonDefaultRole(Databases), AnimContext, PoseHistory);
 
 	const UPoseSearchDatabase* CurrentResultDatabase = InOutMotionMatchingState.SearchResult.SelectedDatabase.Get();
 	if (IsInvalidatingContinuingPose(InterruptMode, CurrentResultDatabase, Databases))
@@ -730,14 +824,14 @@ void UPoseSearchLibrary::UpdateMotionMatchingState(
 	}
 	else
 	{
-		InternalSearchResult.InitFrom(InOutMotionMatchingState.SearchResult);
+		FSearchResult SearchResult;
+		SearchResult.InitFrom(InOutMotionMatchingState.SearchResult);
+		SearchContext.UpdateContinuingPoseSearchResult(SearchResult, SearchResult);
 	}
 
-	const FPoseSearchEvent PlayRateOverriddenEvent = EventToSearch.GetPlayRateOverriddenEvent(PlayRate);
-	FSearchContext SearchContext(0.f, &InOutMotionMatchingState.PoseIndicesHistory, InternalSearchResult, PoseJumpThresholdTime, PlayRateOverriddenEvent);
-	SearchContext.AddRole(GetCommonDefaultRole(Databases), AnimContext, PoseHistory);
+	FSearchResult SearchResult;
 
-	const bool bCanAdvance = InternalSearchResult.PoseIdx != INDEX_NONE;
+	const bool bCanAdvance = SearchContext.GetContinuingPoseSearchResult().PoseIdx != INDEX_NONE;
 
 	// If we can't advance or enough time has elapsed since the last pose jump then search
 	const bool bSearch = !bCanAdvance || (InOutMotionMatchingState.ElapsedPoseSearchTime >= SearchThrottleTime);
@@ -750,7 +844,6 @@ void UPoseSearchLibrary::UpdateMotionMatchingState(
 		// calculating if it's worth bUseCachedChannelData (if we potentially have to build query with multiple schemas)
 		SearchContext.SetUseCachedChannelData(bShouldUseCachedChannelData && ShouldUseCachedChannelData(bSearchContinuingPose ? CurrentResultDatabase : nullptr, Databases));
 
-		FSearchResult SearchResult;
 		// Evaluate continuing pose
 		if (bSearchContinuingPose)
 		{
@@ -780,14 +873,6 @@ void UPoseSearchLibrary::UpdateMotionMatchingState(
 			}
 		}
 
-#if WITH_EDITOR
-		// resetting CurrentSearchResult if any DDC indexing on the requested databases is still in progress
-		if (SearchContext.IsAsyncBuildIndexInProgress())
-		{
-			InternalSearchResult.Reset();
-		}
-#endif // WITH_EDITOR
-
 #if !NO_LOGGING
 		if (!SearchResult.IsValid())
 		{
@@ -829,53 +914,41 @@ void UPoseSearchLibrary::UpdateMotionMatchingState(
 			}
 		}
 #endif // !NO_LOGGING
-
-		// Remember which pose and sequence we're playing from the database
-		InternalSearchResult = SearchResult;
 	}
 	else
 	{
-		// @todo: for blendspaces ElapsedPoseSearchTime should be incremented by a normalized DeltaTime, also didn't we already synchronized the result???
 		InOutMotionMatchingState.ElapsedPoseSearchTime += DeltaTime;
-		InternalSearchResult.bIsContinuingPoseSearch = true;
+		
+		SearchResult = SearchContext.GetContinuingPoseSearchResult();
+		SearchResult.bIsContinuingPoseSearch = true;
 
 #if UE_POSE_SEARCH_TRACE_ENABLED
 		// in case we skipped the search, we still have to track we would have requested to evaluate Databases and CurrentResultDatabase
 		for (const TObjectPtr<const UPoseSearchDatabase>& Database : Databases)
 		{
-			if (Database)
-			{
-				SearchContext.Track(Database);
-			}
-		}
-
-		if (CurrentResultDatabase)
-		{
-			SearchContext.Track(CurrentResultDatabase);
+			SearchContext.Track(Database);
 		}
+		SearchContext.Track(CurrentResultDatabase);
 #endif // UE_POSE_SEARCH_TRACE_ENABLED
 	}
 
-	// @todo: consider moving this into if (bSearch) to avoid calling SearchContext.GetCachedQuery if no search is required
-	const float WantedPlayRate = CalculateWantedPlayRate(InternalSearchResult, SearchContext, PlayRate, PoseHistory ? PoseHistory->GetTrajectorySpeedMultiplier() : 1.f, EventToSearch);
-
-	InOutMotionMatchingState.PoseIndicesHistory.Update(InternalSearchResult, DeltaTime, PoseReselectHistory);
-	InOutMotionMatchingState.SearchResult.InitFrom(InternalSearchResult, WantedPlayRate);
+	const float WantedPlayRate = CalculateWantedPlayRate(SearchResult, SearchContext, PlayRate, PoseHistory ? PoseHistory->GetTrajectorySpeedMultiplier() : 1.f, EventToSearch);
+	InOutMotionMatchingState.PoseIndicesHistory.Update(SearchResult, DeltaTime, PoseReselectHistory);
+	InOutMotionMatchingState.SearchResult.InitFrom(SearchResult, WantedPlayRate);
 
 PRAGMA_DISABLE_DEPRECATION_WARNINGS
 	InOutMotionMatchingState.WantedPlayRate = WantedPlayRate;
-	InOutMotionMatchingState.CurrentSearchResult = InternalSearchResult;
+	InOutMotionMatchingState.CurrentSearchResult = SearchResult;
 PRAGMA_ENABLE_DEPRECATION_WARNINGS
 
 #if UE_POSE_SEARCH_TRACE_ENABLED
-	TraceMotionMatching(SearchContext, InternalSearchResult, InOutMotionMatchingState.ElapsedPoseSearchTime, DeltaTime, bSearch, InOutMotionMatchingState.SearchResult.WantedPlayRate, InterruptMode);
+	TraceMotionMatching(SearchContext, SearchResult, InOutMotionMatchingState.ElapsedPoseSearchTime, DeltaTime, bSearch, InOutMotionMatchingState.SearchResult.WantedPlayRate, InterruptMode);
 #endif // UE_POSE_SEARCH_TRACE_ENABLED
 
 #if WITH_EDITORONLY_DATA && ENABLE_ANIM_DEBUG
-	const FSearchResult& CurResult = InternalSearchResult;
 	if (bDebugDrawQuery || bDebugDrawCurResult)
 	{
-		const UPoseSearchDatabase* CurResultDatabase = CurResult.Database.Get();
+		const UPoseSearchDatabase* CurResultDatabase = SearchResult.Database.Get();
 
 #if WITH_EDITOR
 		if (EAsyncBuildIndexResult::Success == FAsyncPoseSearchDatabasesManagement::RequestAsyncBuildIndex(CurResultDatabase, ERequestAsyncBuildFlag::ContinueRequest))
@@ -884,7 +957,7 @@ PRAGMA_ENABLE_DEPRECATION_WARNINGS
 			FDebugDrawParams DrawParams(SearchContext.GetContexts(), SearchContext.GetPoseHistories(), SearchContext.GetRoleToIndex(), CurResultDatabase);
 			if (bDebugDrawCurResult)
 			{
-				DrawParams.DrawFeatureVector(CurResult.PoseIdx);
+				DrawParams.DrawFeatureVector(SearchResult.PoseIdx);
 			}
 
 			if (bDebugDrawQuery)
@@ -1146,60 +1219,68 @@ UE::PoseSearch::FSearchResult UPoseSearchLibrary::MotionMatch(
 UE::PoseSearch::FSearchResult UPoseSearchLibrary::MotionMatch(
 	const TArrayView<FChooserEvaluationContext> Contexts,
 	const TArrayView<const UE::PoseSearch::FRole> Roles,
-	const TArrayView<const UE::PoseSearch::IPoseHistory*> PoseHistories, 
+	const TArrayView<const UE::PoseSearch::IPoseHistory*> PoseHistories,
 	const TArrayView<const UObject*> AssetsToSearch,
 	const FPoseSearchContinuingProperties& ContinuingProperties,
 	const float DesiredPermutationTimeOffset,
 	const FPoseSearchEvent& EventToSearch)
 {
 	using namespace UE::PoseSearch;
-
-	FSearchResult SearchResult;
-
-	FReconstructedPreviousSearchBestResultMap ReconstructedPreviousSearchBestResultMap;
-	FSearchResult ReconstructedPreviousSearchResult;
-	FSearchContext SearchContext(DesiredPermutationTimeOffset, nullptr, ReconstructedPreviousSearchResult, FFloatInterval(0.f, 0.f), EventToSearch);
-	SearchContext.SetIsContinuingInteraction(ContinuingProperties.bIsContinuingInteraction);
-
+	
+	FSearchContext SearchContext(DesiredPermutationTimeOffset, nullptr, FFloatInterval(0.f, 0.f), EventToSearch);
 	for (int32 RoleIndex = 0; RoleIndex < Roles.Num(); ++RoleIndex)
 	{
 		SearchContext.AddRole(Roles[RoleIndex], &Contexts[RoleIndex], PoseHistories[RoleIndex]);
 	}
-	
-	// collecting all the possible continuing pose search (it could be multiple searches, but most likely only one)
-	float DeltaSeconds = FiniteDelta;
-	if (Contexts[0].ObjectParams.Num() > 0)
-	{
-		if (const UAnimInstance* AnimInstance = Cast<UAnimInstance>(Contexts[0].GetFirstObjectParam()))
-		{
-			DeltaSeconds = AnimInstance->GetDeltaSeconds();
-		}
-	}
 
-	// collecting all the databases searches in AssetsToSearchPerDatabaseMap
-	// and all the continuing pose searches in ContinuingPoseAssetsToSearchPerDatabaseMap
-	FAssetsToSearchPerDatabaseMap AssetsToSearchPerDatabaseMap;
-	FAssetsToSearchPerDatabaseMap ContinuingPoseAssetsToSearchPerDatabaseMap;
-	PopulateSearches(AssetsToSearch, SearchContext, AssetsToSearchPerDatabaseMap);
-	PopulateContinuingPoseSearches(ContinuingProperties.PlayingAsset.Get(), AssetsToSearch, SearchContext, ContinuingPoseAssetsToSearchPerDatabaseMap);
+	return MotionMatch(SearchContext, AssetsToSearch, ContinuingProperties);
+}
 
-	for (const TAssetsToSearchPerDatabasePair& AssetsToSearchPerDatabasePair : ContinuingPoseAssetsToSearchPerDatabaseMap)
+UE::PoseSearch::FSearchResult UPoseSearchLibrary::MotionMatch(
+	UE::PoseSearch::FSearchContext& SearchContext,
+	const TArrayView<const UObject*> AssetsToSearch,
+	const FPoseSearchContinuingProperties& ContinuingProperties)
+{
+	using namespace UE::PoseSearch;
+
+	FSearchResult SearchResult;
+
+	SearchContext.SetIsContinuingInteraction(ContinuingProperties.bIsContinuingInteraction);
+
+	// collecting all the databases searches in AssetsToSearchPerDatabase
+	// and all the continuing pose searches in ContinuingPoseAssetsToSearchPerDatabase
+	const bool bUsePoseSearchBranchIn = ShouldUsePoseSearchBranchIn(AssetsToSearch);
+	FAssetsToSearchPerDatabase AssetsToSearchPerDatabase;
+	FAssetsToSearchPerDatabase ContinuingPoseAssetsToSearchPerDatabase;
+	PopulateSearches(AssetsToSearch, SearchContext, AssetsToSearchPerDatabase, bUsePoseSearchBranchIn);
+	PopulateContinuingPoseSearches(ContinuingProperties, AssetsToSearch, SearchContext, ContinuingPoseAssetsToSearchPerDatabase, bUsePoseSearchBranchIn);
+
+	FCachedContinuingPoseSearchResults CachedContinuingPoseSearchResults;
+	for (const FAssetsToSearchPerDatabasePair& AssetsToSearchPerDatabasePair : ContinuingPoseAssetsToSearchPerDatabase)
 	{
 		const UPoseSearchDatabase* Database = AssetsToSearchPerDatabasePair.Key;
 		check(Database);
 
-		const bool bInvalidatingContinuingPose = IsInvalidatingContinuingPose(ContinuingProperties.InterruptMode, Database, AssetsToSearchPerDatabaseMap);
+		const bool bInvalidatingContinuingPose = IsInvalidatingContinuingPose(ContinuingProperties.InterruptMode, Database, AssetsToSearchPerDatabase);
 		if (!bInvalidatingContinuingPose)
 		{
-			ReconstructedPreviousSearchResult.AssetTime = ContinuingProperties.PlayingAssetAccumulatedTime;
-			ReconstructedPreviousSearchResult.PoseIdx = Database->GetPoseIndex(ContinuingProperties.PlayingAsset.Get(), ContinuingProperties.PlayingAssetAccumulatedTime, ContinuingProperties.bIsPlayingAssetMirrored, ContinuingProperties.PlayingAssetBlendParameters);
-			ReconstructedPreviousSearchResult.Database = Database;
-
-			const bool bForceInterrupt = IsForceInterrupt(ContinuingProperties.InterruptMode, Database, AssetsToSearchPerDatabaseMap);
-			const bool bCanAdvance = ReconstructedPreviousSearchResult.PoseIdx != INDEX_NONE;
+			// reconstructing and caching all the required continuing pose search results
+			FSearchResult DatabaseContinuingPoseSearchResult;
+			DatabaseContinuingPoseSearchResult.AssetTime = ContinuingProperties.PlayingAssetAccumulatedTime;
+			DatabaseContinuingPoseSearchResult.PoseIdx = Database->GetPoseIndex(ContinuingProperties.PlayingAsset.Get(), ContinuingProperties.PlayingAssetAccumulatedTime, ContinuingProperties.bIsPlayingAssetMirrored, ContinuingProperties.PlayingAssetBlendParameters);
+			DatabaseContinuingPoseSearchResult.Database = Database;
+			CachedContinuingPoseSearchResults.CheckedAdd(Database, DatabaseContinuingPoseSearchResult);
+
+			// adding the continuing pose search redult relative to the schema - first instance of the DatabaseContinuingPoseSearchResult 
+			// (used to gather the continuing pose search values adopted to create the MM query - relaive to the schema, NOT the database)
+			const FSearchResult& SchemaContinuingPoseSearchResult = CachedContinuingPoseSearchResults.FindOrAdd(Database->Schema, DatabaseContinuingPoseSearchResult);
+
+			const bool bForceInterrupt = IsForceInterrupt(ContinuingProperties.InterruptMode, Database, AssetsToSearchPerDatabase);
+			const bool bCanAdvance = DatabaseContinuingPoseSearchResult.PoseIdx != INDEX_NONE;
 			if (bCanAdvance && !bForceInterrupt)
 			{
-				SearchContext.UpdateCurrentResultPoseVector();
+				SearchContext.UpdateContinuingPoseSearchResult(DatabaseContinuingPoseSearchResult, SchemaContinuingPoseSearchResult);
+
 				const FSearchResult NewSearchResult = Database->SearchContinuingPose(SearchContext);
 
 #if WITH_EDITOR && ENABLE_ANIM_DEBUG && UE_POSE_SEARCH_TRACE_ENABLED
@@ -1215,63 +1296,23 @@ UE::PoseSearch::FSearchResult UPoseSearchLibrary::MotionMatch(
 #if WITH_EDITOR && ENABLE_ANIM_DEBUG && UE_POSE_SEARCH_TRACE_ENABLED
 				SearchResult.BruteForcePoseCost = BestBruteForcePoseCost;
 #endif // WITH_EDITOR && ENABLE_ANIM_DEBUG && UE_POSE_SEARCH_TRACE_ENABLED
-
-				if (NewSearchResult.IsValid())
-				{
-					// keeping track of the best ReconstructedPreviousSearchResult, as best continuing pose search result for the full Database search happening after
-					if (FSearchResult* FoundReconstructedPreviousSearchBestResult = ReconstructedPreviousSearchBestResultMap.Find(Database))
-					{
-						if (NewSearchResult.PoseCost == FoundReconstructedPreviousSearchBestResult->PoseCost)
-						{
-							check(NewSearchResult.PoseIdx != FoundReconstructedPreviousSearchBestResult->PoseIdx);
-							if (NewSearchResult.PoseIdx < FoundReconstructedPreviousSearchBestResult->PoseIdx)
-							{
-								// choosing to update with the lowest PoseIdx to avoid indeterminism 
-								// since ReconstructedPreviousSearchBestResultMap is an unsorted Map
-								*FoundReconstructedPreviousSearchBestResult = ReconstructedPreviousSearchResult;
-								FoundReconstructedPreviousSearchBestResult->PoseCost = NewSearchResult.PoseCost;
-							}
-						}
-						else if (NewSearchResult.PoseCost < FoundReconstructedPreviousSearchBestResult->PoseCost)
-						{
-							*FoundReconstructedPreviousSearchBestResult = ReconstructedPreviousSearchResult;
-							FoundReconstructedPreviousSearchBestResult->PoseCost = NewSearchResult.PoseCost;
-						}
-					}
-					else
-					{
-						FSearchResult& NewReconstructedPreviousSearchBestResult = ReconstructedPreviousSearchBestResultMap.Add(Database);
-						NewReconstructedPreviousSearchBestResult = ReconstructedPreviousSearchResult;
-						NewReconstructedPreviousSearchBestResult.PoseCost = NewSearchResult.PoseCost;
-					}
-				}
 			}
 		}
 	}
 
 	// performing all the other databases searches
-	for (const TAssetsToSearchPerDatabasePair& AssetsToSearchPerDatabasePair : AssetsToSearchPerDatabaseMap)
+	for (const FAssetsToSearchPerDatabasePair& AssetsToSearchPerDatabasePair : AssetsToSearchPerDatabase)
 	{
 		const UPoseSearchDatabase* Database = AssetsToSearchPerDatabasePair.Key;
 		check(Database);
 
-		// setting up the best continuing pose environment from the continuing pose searches we already performed
-		if (FSearchResult* ReconstructedPreviousSearchBestResult = ReconstructedPreviousSearchBestResultMap.Find(Database))
-		{
-			ReconstructedPreviousSearchResult = *ReconstructedPreviousSearchBestResult;
-		}
-		else
-		{
-			ReconstructedPreviousSearchResult.Reset();
-		}
-
-		SearchContext.SetAssetsToConsider(AssetsToSearchPerDatabasePair.Value);
-
 		// in case we haven't searched the continuing pose for this Database, we haven't created and cached the query yet,
 		// but if we didn't invalidated the continuing pose (when IsInvalidatingContinuingPose is true), we still can reuse
-		// the updated ReconstructedPreviousSearchResult data, and by calling UpdateCurrentResultPoseVector we set the 
+		// the updated FirstInstanceOfReconstructedContinuingPoseSearchResult data, and by calling UpdateContinuingPoseSearchResult we set the 
 		// SearchContext to be able to create a query for Database using the continuing pose data.
-		SearchContext.UpdateCurrentResultPoseVector();
+		SearchContext.UpdateContinuingPoseSearchResult(CachedContinuingPoseSearchResults.FindOrDefault(Database), CachedContinuingPoseSearchResults.FindOrDefault(Database->Schema));
+		SearchContext.SetAssetsToConsider(AssetsToSearchPerDatabasePair.Value);
+		
 		const FSearchResult NewSearchResult = Database->Search(SearchContext);
 
 #if WITH_EDITOR && ENABLE_ANIM_DEBUG && UE_POSE_SEARCH_TRACE_ENABLED
@@ -1289,6 +1330,9 @@ UE::PoseSearch::FSearchResult UPoseSearchLibrary::MotionMatch(
 #endif // WITH_EDITOR && ENABLE_ANIM_DEBUG && UE_POSE_SEARCH_TRACE_ENABLED
 	}
 
+	// clean up the asset to consider array view
+	SearchContext.SetAssetsToConsider(TConstArrayView<const UObject*>());
+
 #if ENABLE_DRAW_DEBUG && ENABLE_ANIM_DEBUG
 	if (SearchResult.IsValid())
 	{
@@ -1312,6 +1356,17 @@ UE::PoseSearch::FSearchResult UPoseSearchLibrary::MotionMatch(
 #endif // ENABLE_DRAW_DEBUG && ENABLE_ANIM_DEBUG
 
 #if UE_POSE_SEARCH_TRACE_ENABLED
+
+	// collecting all the possible continuing pose search (it could be multiple searches, but most likely only one)
+	float DeltaSeconds = FiniteDelta;
+	if (SearchContext.GetContexts()[0]->ObjectParams.Num() > 0)
+	{
+		if (const UAnimInstance* AnimInstance = Cast<UAnimInstance>(SearchContext.GetContexts()[0]->GetFirstObjectParam()))
+		{
+			DeltaSeconds = AnimInstance->GetDeltaSeconds();
+		}
+	}
+
 	// @todo: add and handle an interrupt mode input param MotionMatch
 	TraceMotionMatching(SearchContext, SearchResult, 0.f, DeltaSeconds, true, 1.f, EPoseSearchInterruptMode::DoNotInterrupt);
 #endif // UE_POSE_SEARCH_TRACE_ENABLED
diff --git a/Engine/Plugins/Animation/PoseSearch/Source/Runtime/Private/PoseSearchModule.cpp b/Engine/Plugins/Animation/PoseSearch/Source/Runtime/Private/PoseSearchModule.cpp
--- a/Engine/Plugins/Animation/PoseSearch/Source/Runtime/Private/PoseSearchModule.cpp
+++ b/Engine/Plugins/Animation/PoseSearch/Source/Runtime/Private/PoseSearchModule.cpp
@@ -59,6 +59,7 @@ public:
 		ContinuingProperties.PlayingAssetAccumulatedTime = PlayingAsset.AccumulatedTime;
 		ContinuingProperties.bIsPlayingAssetMirrored = PlayingAsset.bMirrored;
 		ContinuingProperties.PlayingAssetBlendParameters = PlayingAsset.BlendParameters;
+		// @todo: ContinuingProperties.PlayingAssetDatabase = ???;
 
 		FPoseSearchFutureProperties Future;
 		Future.Animation = FutureAsset.Asset;
diff --git a/Engine/Plugins/Animation/PoseSearch/Source/Runtime/Private/PoseSearchResult.cpp b/Engine/Plugins/Animation/PoseSearch/Source/Runtime/Private/PoseSearchResult.cpp
--- a/Engine/Plugins/Animation/PoseSearch/Source/Runtime/Private/PoseSearchResult.cpp
+++ b/Engine/Plugins/Animation/PoseSearch/Source/Runtime/Private/PoseSearchResult.cpp
@@ -10,6 +10,44 @@
 namespace UE::PoseSearch
 {
 
+/////////////////////////////////////////////////////
+// FDatabasePoseIdx
+
+const FSearchIndexAsset* FDatabasePoseIdx::GetSearchIndexAsset(bool bMandatory) const
+{
+	if (bMandatory)
+	{
+		check(IsValid());
+	}
+	else if (!IsValid())
+	{
+		return nullptr;
+	}
+
+	return &Database->GetSearchIndex().GetAssetForPose(PoseIdx);
+}
+
+const UAnimationAsset* FDatabasePoseIdx::GetCurrentResultAnimationAsset() const
+{
+	if (const FSearchIndexAsset* SearchIndexAsset = GetSearchIndexAsset())
+	{
+		return Database->GetDatabaseAnimationAsset<FPoseSearchDatabaseAnimationAssetBase>(*SearchIndexAsset)->GetAnimationAssetForRole(Database->Schema->GetDefaultRole());
+	}
+	return nullptr;
+}
+
+const UAnimationAsset* FDatabasePoseIdx::GetCurrentResultAnimationAsset(const FRole& Role) const
+{
+	if (const FSearchIndexAsset* SearchIndexAsset = GetSearchIndexAsset())
+	{
+		return Database->GetDatabaseAnimationAsset<FPoseSearchDatabaseAnimationAssetBase>(*SearchIndexAsset)->GetAnimationAssetForRole(Role);
+	}
+	return nullptr;
+}
+
+/////////////////////////////////////////////////////
+// FSearchResult
+
 bool FSearchResult::DebugValidate() const
 {
 	bool bIsValidated = true;
@@ -158,38 +196,6 @@ float FSearchResult::CalculateTimeToEvent() const
 	return TimeToEvent;
 }
 
-const FSearchIndexAsset* FSearchResult::GetSearchIndexAsset(bool bMandatory) const
-{
-	if (bMandatory)
-	{
-		check(IsValid());
-	}
-	else if (!IsValid())
-	{
-		return nullptr;
-	}
-
-	return &Database->GetSearchIndex().GetAssetForPose(PoseIdx);
-}
-
-const UAnimationAsset* FSearchResult::GetCurrentResultAnimationAsset() const
-{
-	if (const FSearchIndexAsset* SearchIndexAsset = GetSearchIndexAsset())
-	{
-		return Database->GetDatabaseAnimationAsset<FPoseSearchDatabaseAnimationAssetBase>(*SearchIndexAsset)->GetAnimationAssetForRole(Database->Schema->GetDefaultRole());
-	}
-	return nullptr;
-}
-
-const UAnimationAsset* FSearchResult::GetCurrentResultAnimationAsset(const FRole& Role) const
-{
-	if (const FSearchIndexAsset* SearchIndexAsset = GetSearchIndexAsset())
-	{
-		return Database->GetDatabaseAnimationAsset<FPoseSearchDatabaseAnimationAssetBase>(*SearchIndexAsset)->GetAnimationAssetForRole(Role);
-	}
-	return nullptr;
-}
-
 void FSearchResult::InitFrom(const FPoseSearchBlueprintResult& BlueprintResult)
 {
 	PoseCost = FPoseSearchCost(BlueprintResult.SearchCost, 0.f, 0.f, 0.f);
@@ -257,4 +263,30 @@ UAnimationAsset* FPoseSearchBlueprintResult::GetAnimationAssetForRole()
 	}
 
 	return nullptr;
+}
+
+bool FPoseSearchBlueprintResult::operator==(const FPoseSearchBlueprintResult& Other) const
+{
+	return	
+#if WITH_EDITORONLY_DATA
+		SelectedAnimation_DEPRECATED == Other.SelectedAnimation_DEPRECATED &&
+#endif // WITH_EDITORONLY_DATA
+		SelectedAnim == Other.SelectedAnim &&
+		SelectedTime == Other.SelectedTime &&
+		bIsContinuingPoseSearch == Other.bIsContinuingPoseSearch &&
+		WantedPlayRate == Other.WantedPlayRate &&
+		bLoop == Other.bLoop &&
+		bIsMirrored == Other.bIsMirrored &&
+		BlendParameters == Other.BlendParameters &&
+		SelectedDatabase == Other.SelectedDatabase &&
+		SearchCost == Other.SearchCost &&
+		bIsInteraction == Other.bIsInteraction &&
+		Role == Other.Role &&
+		AnimContexts == Other.AnimContexts &&
+
+		// FTransform(s) don't have operator == so we need to do do some custom work here :/
+		ActorRootTransforms.Num() == Other.ActorRootTransforms.Num() &&
+		ActorRootBoneTransforms.Num() == Other.ActorRootBoneTransforms.Num() &&
+		FMemory::Memcmp(ActorRootTransforms.GetData(), Other.ActorRootTransforms.GetData(), sizeof(FTransform) * ActorRootTransforms.Num()) == 0 &&
+		FMemory::Memcmp(ActorRootBoneTransforms.GetData(), Other.ActorRootBoneTransforms.GetData(), sizeof(FTransform) * ActorRootBoneTransforms.Num()) == 0;
 }
\ No newline at end of file
diff --git a/Engine/Plugins/Animation/PoseSearch/Source/Runtime/Public/PoseSearch/AnimNode_MotionMatching.h b/Engine/Plugins/Animation/PoseSearch/Source/Runtime/Public/PoseSearch/AnimNode_MotionMatching.h
--- a/Engine/Plugins/Animation/PoseSearch/Source/Runtime/Public/PoseSearch/AnimNode_MotionMatching.h
+++ b/Engine/Plugins/Animation/PoseSearch/Source/Runtime/Public/PoseSearch/AnimNode_MotionMatching.h
@@ -140,7 +140,7 @@ private:
 	UPROPERTY(EditAnywhere, Category = "Settings|Experimental", meta = (PinHiddenByDefault))
 	float PlayRateMultiplier = 1.0f;
 	
-	UPROPERTY(EditAnywhere, Category = Settings, Category = Settings)
+	UPROPERTY(EditAnywhere, Category = Settings)
 	bool bUseInertialBlend = false;
 
 	// Reset the motion matching selection state if it has become relevant to the graph after not being updated on previous frames.
diff --git a/Engine/Plugins/Animation/PoseSearch/Source/Runtime/Public/PoseSearch/PoseSearchContext.h b/Engine/Plugins/Animation/PoseSearch/Source/Runtime/Public/PoseSearch/PoseSearchContext.h
--- a/Engine/Plugins/Animation/PoseSearch/Source/Runtime/Public/PoseSearch/PoseSearchContext.h
+++ b/Engine/Plugins/Animation/PoseSearch/Source/Runtime/Public/PoseSearch/PoseSearchContext.h
@@ -150,10 +150,14 @@ struct FSearchContext
 	UE_API FSearchContext(float InDesiredPermutationTimeOffset = 0.f, const FPoseIndicesHistory* InPoseIndicesHistory = nullptr,
 		const FSearchResult& InCurrentResult = FSearchResult(), const FFloatInterval& InPoseJumpThresholdTime = FFloatInterval(0.f, 0.f), bool bInUseCachedChannelData = false);
 
+	UE_DEPRECATED(5.7, "Use other signaure instead")
 	UE_API FSearchContext(float InDesiredPermutationTimeOffset, const FPoseIndicesHistory* InPoseIndicesHistory,
 		const FSearchResult& InCurrentResult, const FFloatInterval& InPoseJumpThresholdTime, const FPoseSearchEvent& InEventToSearch);
 
-	// deleting copies and moves since members could reference other members (e.g.: CurrentResultPoseVector could point to CurrentResultPoseVectorData, so it'll require proper copies and movers implementations)
+	UE_API FSearchContext(float InDesiredPermutationTimeOffset, const FPoseIndicesHistory* InPoseIndicesHistory,
+		const FFloatInterval& InPoseJumpThresholdTime, const FPoseSearchEvent& InEventToSearch);
+
+	// deleting copies and moves since members could reference other members (e.g.: ContinuingPoseValues could point to ContinuingPoseValuesData, so it'll require proper copies and movers implementations)
 	FSearchContext(const FSearchContext& Other) = delete;
 	FSearchContext(FSearchContext&& Other) = delete;
 	FSearchContext& operator=(const FSearchContext& Other) = delete;
@@ -188,13 +192,26 @@ struct FSearchContext
 	UE_API TConstArrayView<float> GetOrBuildQuery(const UPoseSearchSchema* Schema);
 	UE_API TConstArrayView<float> GetCachedQuery(const UPoseSearchSchema* Schema) const;
 
+	UE_DEPRECATED(5.7, "compare Database against GetContinuingPoseSearchResult().Database instead (the internal replacement IsContinuingPoseDatabase is not going to be dll exported)")
 	UE_API bool IsCurrentResultFromDatabase(const UPoseSearchDatabase* Database) const;
+	bool IsContinuingPoseDatabase(const UPoseSearchDatabase* Database) const;
+	
+	UE_DEPRECATED(5.7, "Use CanUseContinuingPoseValues instead")
 	UE_API bool CanUseCurrentResult() const;
+	UE_API bool CanUseContinuingPoseValues() const;
 
-	TConstArrayView<float> GetCurrentResultPoseVector() const { return CurrentResultPoseVector; }
+	UE_DEPRECATED(5.7, "Use GetContinuingPoseValues")
+	TConstArrayView<float> GetCurrentResultPoseVector() const { return ContinuingPoseValues; }
+	TConstArrayView<float> GetContinuingPoseValues() const { return ContinuingPoseValues; }
 
+	UE_DEPRECATED(5.7, "Use UpdateContinuingPoseSearchResult to update the continuing pose search result AND values - previously referenced as 'CurrentResultPoseVector' ")
 	UE_API void UpdateCurrentResultPoseVector();
-	const FSearchResult& GetCurrentResult() const { return CurrentResult; }
+	UE_API void UpdateContinuingPoseSearchResult(const FSearchResult& InContinuingPoseSearchResult, const FDatabasePoseIdx& InContinuingPoseValuesDatabasePoseIdx);
+	
+	UE_DEPRECATED(5.7, "Use GetContinuingPoseSearchResult")
+	const FSearchResult& GetCurrentResult() const { return ContinuingPoseSearchResult; }
+	const FSearchResult& GetContinuingPoseSearchResult() const { return ContinuingPoseSearchResult; }
+
 	const FFloatInterval& GetPoseJumpThresholdTime() const { return PoseJumpThresholdTime; }
 	const FPoseIndicesHistory* GetPoseIndicesHistory() const { return PoseIndicesHistory; }
 	
@@ -265,10 +282,11 @@ struct FSearchContext
 	void SetIsContinuingInteraction(bool bInIsContinuingInteraction) { bIsContinuingInteraction = bInIsContinuingInteraction; }
 
 private:
-	UE_API float GetSampleCurveValueInternal(float SampleTime, const FName& CurveName, const FRole& SampleRole);
-	UE_API FVector GetSamplePositionInternal(float SampleTime, float OriginTime, int8 SchemaSampleBoneIdx, int8 SchemaOriginBoneIdx, const FRole& SampleRole, const FRole& OriginRole, const FVector* SampleBonePositionWorldOverride = nullptr);
-	UE_API FQuat GetSampleRotationInternal(float SampleTime, float OriginTime, int8 SchemaSampleBoneIdx, int8 SchemaOriginBoneIdx, const FRole& SampleRole, const FRole& OriginRole, const FQuat* SampleBoneRotationWorldOverride = nullptr);
-	UE_API FTransform GetWorldRootBoneTransformAtTime(float SampleTime, const FRole& SampleRole) const;
+	float GetSampleCurveValueInternal(float SampleTime, const FName& CurveName, const FRole& SampleRole);
+	FVector GetSamplePositionInternal(float SampleTime, float OriginTime, int8 SchemaSampleBoneIdx, int8 SchemaOriginBoneIdx, const FRole& SampleRole, const FRole& OriginRole, const FVector* SampleBonePositionWorldOverride = nullptr);
+	FQuat GetSampleRotationInternal(float SampleTime, float OriginTime, int8 SchemaSampleBoneIdx, int8 SchemaOriginBoneIdx, const FRole& SampleRole, const FRole& OriginRole, const FQuat* SampleBoneRotationWorldOverride = nullptr);
+	FTransform GetWorldRootBoneTransformAtTime(float SampleTime, const FRole& SampleRole) const;
+	void UpdateContinuingPoseValues();
 	
 	TArray<FChooserEvaluationContext*, TInlineAllocator<PreallocatedRolesNum, TMemStackAllocator<>>> AnimContexts;
 	TArray<const IPoseHistory*, TInlineAllocator<PreallocatedRolesNum, TMemStackAllocator<>>> PoseHistories;
@@ -280,13 +298,28 @@ private:
 
 	const float DesiredPermutationTimeOffset = 0.f;
 	const FPoseIndicesHistory* PoseIndicesHistory = nullptr;
-	const FSearchResult& CurrentResult;
+	
+	// search result used to extract the continuing pose properties, except for gathering the continuing pose values, where we use ContinuingPoseValuesDatabasePoseIdx.
+	// More details in ContinuingPoseValuesDatabasePoseIdx comments
+	FSearchResult ContinuingPoseSearchResult;
+
+	// database - pose index pair used to extract the continuing pose vector properties (values)
+	// the reason it can differ from the ContinuingPoseSearchResult (database - pose index pair) is because we build and cache a query only using a schema, 
+	// NOT a database (for performance reasons), BUT we want to be able to calculate the continuing pose from different databases sharing the same animation assets AND same schema,
+	// and evaluate the continuing pose feature data across both of them. When this happens, the query gets created using the FIRST database pose value that could slightly differs, 
+	// because of quantization of time, numerical errors, not common normalization sets, etc, BUT we still want to keep this logic consistent to the fact we created 
+	// the query and cached it for the common schema, NOT per database, so we need to "lie" to the system by providing a ContinuingPoseSearchResult to gather continuing 
+	// pose properties such as pose index, database etc, and rely on "ContinuingPoseValuesDatabasePoseIdx" to gather the continuing pose feature values vector (ContinuingPoseValues),
+	// used in the channels as cached version of the data that could have been caluclated from the pose history
+	FDatabasePoseIdx ContinuingPoseValuesDatabasePoseIdx;
+
 	const FFloatInterval PoseJumpThresholdTime;
 	bool bUseCachedChannelData = false;
 
-	TConstArrayView<float> CurrentResultPoseVector;
-
-	TStackAlignedArray<float> CurrentResultPoseVectorData;
+	// if ContinuingPoseValuesDatabasePoseIdx is valid, this array view references the continuing pose values (feature vector data) 
+	TConstArrayView<float> ContinuingPoseValues;
+	// buffer optionally used in case the continuing pose needed to be reconstructed. ContinuingPoseValues will reference it
+	TStackAlignedArray<float> ContinuingPoseValuesData;
 
 	// transforms cached in world space
 	TMap<uint32, FTransform, TInlineSetAllocator<64, TMemStackSetAllocator<>>> CachedTransforms;
@@ -324,12 +357,13 @@ public:
 
 	void Track(const UPoseSearchDatabase* Database, int32 PoseIdx = INDEX_NONE, EPoseCandidateFlags PoseCandidateFlags = EPoseCandidateFlags::None, const FPoseSearchCost& Cost = FPoseSearchCost())
 	{
-		check(Database);
-
-		FBestPoseCandidates& BestPoseCandidates = BestPoseCandidatesMap.FindOrAdd(Database);
-		if (PoseIdx != INDEX_NONE)
+		if (Database)
 		{
-			BestPoseCandidates.Add(PoseIdx, PoseCandidateFlags, Cost);
+			FBestPoseCandidates& BestPoseCandidates = BestPoseCandidatesMap.FindOrAdd(Database);
+			if (PoseIdx != INDEX_NONE)
+			{
+				BestPoseCandidates.Add(PoseIdx, PoseCandidateFlags, Cost);
+			}
 		}
 	}
 
diff --git a/Engine/Plugins/Animation/PoseSearch/Source/Runtime/Public/PoseSearch/PoseSearchHistory.h b/Engine/Plugins/Animation/PoseSearch/Source/Runtime/Public/PoseSearch/PoseSearchHistory.h
--- a/Engine/Plugins/Animation/PoseSearch/Source/Runtime/Public/PoseSearch/PoseSearchHistory.h
+++ b/Engine/Plugins/Animation/PoseSearch/Source/Runtime/Public/PoseSearch/PoseSearchHistory.h
@@ -322,6 +322,8 @@ struct FPoseIndicesHistory
 {
 	void Update(const FSearchResult& SearchResult, float DeltaTime, float MaxTime);
 	void Reset() { IndexToTime.Reset(); }
+	bool operator==(const FPoseIndicesHistory& Other) const;
+		
 	TMap<FHistoricalPoseIndex, float> IndexToTime;
 };
 
diff --git a/Engine/Plugins/Animation/PoseSearch/Source/Runtime/Public/PoseSearch/PoseSearchLibrary.h b/Engine/Plugins/Animation/PoseSearch/Source/Runtime/Public/PoseSearch/PoseSearchLibrary.h
--- a/Engine/Plugins/Animation/PoseSearch/Source/Runtime/Public/PoseSearch/PoseSearchLibrary.h
+++ b/Engine/Plugins/Animation/PoseSearch/Source/Runtime/Public/PoseSearch/PoseSearchLibrary.h
@@ -17,13 +17,20 @@
 
 #define UE_API POSESEARCH_API
 
+struct FAnimationUpdateContext;
+struct FAnimNode_PoseSearchHistoryCollector_Base;
+
 namespace UE::PoseSearch
 {
 	struct FSearchContext;
+
+	// Experimental, this feature might be removed without warning, not for production use
+	FRole GetCommonDefaultRole(const TConstArrayView<TObjectPtr<const UPoseSearchDatabase>> Databases);
+	// Experimental, this feature might be removed without warning, not for production use
+	float CalculateWantedPlayRate(const FSearchResult& SearchResult, const FSearchContext& SearchContext, const FFloatInterval& PlayRate, float TrajectorySpeedMultiplier, const FPoseSearchEvent& EventToSearch);
+
 } // namespace UE::PoseSearch
 
-struct FAnimationUpdateContext;
-struct FAnimNode_PoseSearchHistoryCollector_Base;
 
 UENUM(BlueprintType)
 enum class EPoseSearchInterruptMode : uint8
@@ -81,7 +88,7 @@ PRAGMA_ENABLE_DEPRECATION_WARNINGS
 	UPROPERTY(Transient)
 	FPoseSearchBlueprintResult SearchResult;
 
-	// Time since the last pose jump
+	// Time since the last pose jump (real time in seconds, not a normalized time)
 	float ElapsedPoseSearchTime = 0.f;
 
 	UE_DEPRECATED(5.6, "Use SearchResult.WantedPlayRate instead")
@@ -119,25 +126,33 @@ struct FPoseSearchContinuingProperties
 	GENERATED_BODY()
 
 	// Currently playing animation
-	UPROPERTY(EditAnywhere, BlueprintReadWrite, Category=State)
+	UPROPERTY(Transient, EditAnywhere, BlueprintReadWrite, Category=State)
 	TObjectPtr<const UObject> PlayingAsset = nullptr;
 
 	// Currently playing animation accumulated time
-	UPROPERTY(EditAnywhere, BlueprintReadWrite, Category=State)
+	UPROPERTY(Transient, EditAnywhere, BlueprintReadWrite, Category=State)
 	float PlayingAssetAccumulatedTime = 0.f;
 
-	UPROPERTY(EditAnywhere, BlueprintReadWrite, Category=State)
+	UPROPERTY(Transient, EditAnywhere, BlueprintReadWrite, Category=State)
 	bool bIsPlayingAssetMirrored = false;
 
 	// PlayingAsset associated BlendParameters (if PlayingAsset is a UBlendSpace)
-	UPROPERTY(VisibleAnywhere, BlueprintReadOnly, Category=State)
+	UPROPERTY(Transient, VisibleAnywhere, BlueprintReadOnly, Category=State)
 	FVector PlayingAssetBlendParameters = FVector::ZeroVector;
 
-	UPROPERTY(EditAnywhere, BlueprintReadWrite, Category=State)
+	UPROPERTY(Transient, EditAnywhere, BlueprintReadWrite, Category=State)
 	EPoseSearchInterruptMode InterruptMode = EPoseSearchInterruptMode::DoNotInterrupt;
 
-	UPROPERTY(EditAnywhere, BlueprintReadWrite, Category=State)
+	UPROPERTY(Transient, EditAnywhere, BlueprintReadWrite, Category=State)
 	bool bIsContinuingInteraction = false;
+
+	// database where the PlayingAsset was originally from
+	// (optional property, used to provide an additional database to search for the continuing pose associated to PlayingAsset)
+	UPROPERTY(Transient, EditAnywhere, BlueprintReadWrite, Category=State)
+	TObjectPtr<const UPoseSearchDatabase> PlayingAssetDatabase = nullptr;
+
+	// Experimental, this feature might be removed without warning, not for production use
+	void InitFrom(const FPoseSearchBlueprintResult& SearchResult, EPoseSearchInterruptMode InInterruptMode);
 };
 
 UCLASS(MinimalAPI)
@@ -318,6 +333,12 @@ public:
 		const float DesiredPermutationTimeOffset,
 		const FPoseSearchEvent& EventToSearch);
 
+	// Experimental, this feature might be removed without warning, not for production use
+	static UE_API UE::PoseSearch::FSearchResult MotionMatch(
+		UE::PoseSearch::FSearchContext& SearchContext,
+		const TArrayView<const UObject*> AssetsToSearch,
+		const FPoseSearchContinuingProperties& ContinuingProperties);
+
 	static UE_API const FAnimNode_PoseSearchHistoryCollector_Base* FindPoseHistoryNode(
 		const FName PoseHistoryName,
 		const UAnimInstance* AnimInstance);
diff --git a/Engine/Plugins/Animation/PoseSearch/Source/Runtime/Public/PoseSearch/PoseSearchResult.h b/Engine/Plugins/Animation/PoseSearch/Source/Runtime/Public/PoseSearch/PoseSearchResult.h
--- a/Engine/Plugins/Animation/PoseSearch/Source/Runtime/Public/PoseSearch/PoseSearchResult.h
+++ b/Engine/Plugins/Animation/PoseSearch/Source/Runtime/Public/PoseSearch/PoseSearchResult.h
@@ -19,22 +19,34 @@ namespace UE::PoseSearch
 {
 struct FSearchIndexAsset;
 
-struct FSearchResult
+struct FDatabasePoseIdx
 {
-	// best cost of the currently selected PoseIdx (it could be equal to ContinuingPoseCost)
-	FPoseSearchCost PoseCost;
-	
 	// pose index from Database associated to this search result
 	int32 PoseIdx = INDEX_NONE;
 
+	// database associated to this search result. Poses are derived data stored in Database->GetSearchIndex()
+	TWeakObjectPtr<const UPoseSearchDatabase> Database;
+
+	bool IsValid() const { return PoseIdx != INDEX_NONE && Database != nullptr; }
+	bool operator==(const FDatabasePoseIdx& Other) const { return PoseIdx == Other.PoseIdx && Database == Other.Database; }
+
+	POSESEARCH_API const FSearchIndexAsset* GetSearchIndexAsset(bool bMandatory = false) const;
+	// Experimental, this feature might be removed without warning, not for production use
+	POSESEARCH_API const UAnimationAsset* GetCurrentResultAnimationAsset() const;
+	// Experimental, this feature might be removed without warning, not for production use
+	POSESEARCH_API const UAnimationAsset* GetCurrentResultAnimationAsset(const FRole& Role) const;
+};
+
+struct FSearchResult : FDatabasePoseIdx
+{
+	// best cost of the currently selected PoseIdx (it could be equal to ContinuingPoseCost)
+	FPoseSearchCost PoseCost;
+	
 	// Experimental, this feature might be removed without warning, not for production use
 	// if this search result is from an event search, EventPoseIdx is the pose index at event time from Database associated to this search result
 	// if it's not an event search result EventPoseIdx is INDEX_NONE
 	int32 EventPoseIdx = INDEX_NONE;
 
-	// database associated to this search result. Poses are derived data stored in Database->GetSearchIndex()
-	TWeakObjectPtr<const UPoseSearchDatabase> Database;
-
 	// time of the associated GetSearchIndexAsset from this search. During continuing pose searches,
 	// AssetTime is resynchronized via UpdateWithNormalizedTime or UpdateWithRealTime, so it's not precisely the quantized time associated to PoseIdx
 	float AssetTime = 0.f;
@@ -60,13 +72,6 @@ struct FSearchResult
 	// Experimental, this feature might be removed without warning, not for production use
 	bool DebugValidate() const;
 
-	bool IsValid() const
-	{
-		return
-			PoseIdx != INDEX_NONE &&
-			Database != nullptr;
-	}
-
 	// Experimental, this feature might be removed without warning, not for production use
 	bool IsEventSearchResult() const
 	{
@@ -80,35 +85,14 @@ struct FSearchResult
 	// Experimental, this feature might be removed without warning, not for production use
 	float CalculateTimeToEvent() const;
 
-	void Reset()
-	{
-		PoseCost = FPoseSearchCost();
-		PoseIdx = INDEX_NONE;
-		EventPoseIdx = INDEX_NONE;
-		Database = nullptr;
-		AssetTime = 0.0f;
-		bIsContinuingPoseSearch = false;
-
-#if WITH_EDITOR && ENABLE_ANIM_DEBUG && UE_POSE_SEARCH_TRACE_ENABLED
-		BruteForcePoseCost = FPoseSearchCost();
-		BestPosePos = 0;
-#endif // WITH_EDITOR && ENABLE_ANIM_DEBUG && UE_POSE_SEARCH_TRACE_ENABLED
-	}
+	void Reset() { *this = FSearchResult(); }
 
-	POSESEARCH_API const FSearchIndexAsset* GetSearchIndexAsset(bool bMandatory = false) const;
-	// Experimental, this feature might be removed without warning, not for production use
-	POSESEARCH_API const UAnimationAsset* GetCurrentResultAnimationAsset() const;
-	// Experimental, this feature might be removed without warning, not for production use
-	POSESEARCH_API const UAnimationAsset* GetCurrentResultAnimationAsset(const FRole& Role) const;
-	
 	bool operator==(const FSearchResult& Other) const
 	{
 		// best cost of the currently selected PoseIdx (it could be equal to ContinuingPoseCost)
-		return 
+		return static_cast<const FDatabasePoseIdx&>(*this) == static_cast<const FDatabasePoseIdx&>(Other) &&
 			PoseCost == Other.PoseCost &&
-			PoseIdx == Other.PoseIdx &&
 			EventPoseIdx == Other.EventPoseIdx &&
-			Database == Other.Database &&
 			AssetTime == Other.AssetTime &&
 			bIsContinuingPoseSearch == Other.bIsContinuingPoseSearch;
 	}
@@ -193,6 +177,15 @@ struct FPoseSearchBlueprintResult
 
 	// Experimental, this feature might be removed without warning, not for production use
 	UE_API UAnimationAsset* GetAnimationAssetForRole();
+
+	// Experimental, this feature might be removed without warning, not for production use
+	bool operator==(const FPoseSearchBlueprintResult& Other) const;
+};
+
+// Experimental, this feature might be removed without warning, not for production use
+template<> struct TStructOpsTypeTraits<FPoseSearchBlueprintResult> : public TStructOpsTypeTraitsBase2<FPoseSearchBlueprintResult>
+{
+	enum { WithIdenticalViaEquality = true };
 };
 
 #undef UE_API
