Fixing compilation error when building one of the internal projects

diff --git a/Engine/Plugins/Animation/PoseSearch/Source/Runtime/Private/PoseSearchLibrary.cpp b/Engine/Plugins/Animation/PoseSearch/Source/Runtime/Private/PoseSearchLibrary.cpp
--- a/Engine/Plugins/Animation/PoseSearch/Source/Runtime/Private/PoseSearchLibrary.cpp
+++ b/Engine/Plugins/Animation/PoseSearch/Source/Runtime/Private/PoseSearchLibrary.cpp
@@ -92,7 +92,7 @@ namespace UE::PoseSearch
 				return *FoundSearchResult;
 			}
 
-			static const FSearchResult DefaultSearchResult;
+			static FSearchResult DefaultSearchResult;
 			return DefaultSearchResult;
 		}
 	};
