PoseSearch - fix for improper use of Reset, when a SetNum is required

diff --git a/Engine/Plugins/Animation/PoseSearch/Source/Runtime/Private/PoseSearchIndex.cpp b/Engine/Plugins/Animation/PoseSearch/Source/Runtime/Private/PoseSearchIndex.cpp
--- a/Engine/Plugins/Animation/PoseSearch/Source/Runtime/Private/PoseSearchIndex.cpp
+++ b/Engine/Plugins/Animation/PoseSearch/Source/Runtime/Private/PoseSearchIndex.cpp
@@ -897,7 +897,7 @@ TConstArrayView<float> FSearchIndex::GetPoseValuesSafe(int32 PoseIdx, TArray<flo
 	if (IsValuesEmpty())
 	{
 		const int32 NumDimensions = GetNumDimensions();
-		BufferUsedForReconstruction.Reset(NumDimensions);
+		BufferUsedForReconstruction.SetNum(NumDimensions, EAllowShrinking::No);
 		return GetReconstructedPoseValues(PoseIdx, BufferUsedForReconstruction);
 	}
 		
diff --git a/Engine/Plugins/Animation/PoseSearch/Source/Runtime/Private/PoseSearchTrajectoryTypes.cpp b/Engine/Plugins/Animation/PoseSearch/Source/Runtime/Private/PoseSearchTrajectoryTypes.cpp
--- a/Engine/Plugins/Animation/PoseSearch/Source/Runtime/Private/PoseSearchTrajectoryTypes.cpp
+++ b/Engine/Plugins/Animation/PoseSearch/Source/Runtime/Private/PoseSearchTrajectoryTypes.cpp
@@ -218,7 +218,7 @@ FArchive& operator<<(FArchive& Ar, FPoseSearchQueryTrajectorySample& TrajectoryS
 
 FPoseSearchQueryTrajectory::FPoseSearchQueryTrajectory(const FTransformTrajectory& InTrajectory)
 {
-	Samples.Reset(InTrajectory.Samples.Num());
+	Samples.SetNum(InTrajectory.Samples.Num(), EAllowShrinking::No);
 		
 	for (int i = 0; i < InTrajectory.Samples.Num(); ++i)
 	{
