PoseSearch
- improvements for GVarAnimNodeMotionMatchingCompareUpdateMotionMatchingStateAgainstMotionMatch
- using the new UPoseSearchLibrary::MotionMatch API in FInteractionIsland::DoSearch_AnyThread
- adding user-provided default constructors to FDatabasePoseIdx and FSearchResult

diff --git a/Engine/Plugins/Animation/PoseSearch/Source/Runtime/Private/AnimNode_MotionMatching.cpp b/Engine/Plugins/Animation/PoseSearch/Source/Runtime/Private/AnimNode_MotionMatching.cpp
--- a/Engine/Plugins/Animation/PoseSearch/Source/Runtime/Private/AnimNode_MotionMatching.cpp
+++ b/Engine/Plugins/Animation/PoseSearch/Source/Runtime/Private/AnimNode_MotionMatching.cpp
@@ -47,7 +47,7 @@ static float GVarAnimNodeMotionMatchingDrawInfoHeight = 50.f;
 static FAutoConsoleVariableRef CVarAnimNodeMotionMatchingDrawInfoHeight(TEXT("a.AnimNode.MotionMatching.DebugDrawInfoHeight"), GVarAnimNodeMotionMatchingDrawInfoHeight, TEXT("Vertical offset for DebugDrawInfo"));
 
 static int32 GVarAnimNodeMotionMatchingCompareUpdateMotionMatchingStateAgainstMotionMatch = 0;
-static FAutoConsoleVariableRef CVarAnimNodeMotionMatchingCompareUpdateMotionMatchingStateAgainstMotionMatch(TEXT("a.AnimNode.MotionMatching.CompareUpdateMotionMatchingStateAgainstMotionMatch"), GVarAnimNodeMotionMatchingCompareUpdateMotionMatchingStateAgainstMotionMatch, TEXT("Compare UpdateMotionMatchingState Against MotionMatch. 0:disabled 1:warn 2:debug restore MM state"));
+static FAutoConsoleVariableRef CVarAnimNodeMotionMatchingCompareUpdateMotionMatchingStateAgainstMotionMatch(TEXT("a.AnimNode.MotionMatching.CompareUpdateMotionMatchingStateAgainstMotionMatch"), GVarAnimNodeMotionMatchingCompareUpdateMotionMatchingStateAgainstMotionMatch, TEXT("Compare UpdateMotionMatchingState Against MotionMatch. 0:disabled 1:enabled 2:compare 3:debug restore MM state"));
 
 static int32 GVarAnimNodeMotionMatchingPlayRateEnabled = UE::Private::EPlayRateState::Enabled;
 static FAutoConsoleVariableRef CVarAnimNodeMotionMatchingPlayRateEnabled(
@@ -167,7 +167,9 @@ void FAnimNode_MotionMatching::UpdateAssetPlayer(const FAnimationUpdateContext&
 	// and expose FSearchContext to be able to perform multiple searches with the same cached queries (to speed up MM interactions)
 	FMotionMatchingState MotionMatchingStateBackup;
 	FMotionMatchingState MotionMatchingStateCopy;
-	if (GVarAnimNodeMotionMatchingCompareUpdateMotionMatchingStateAgainstMotionMatch)
+
+	// Compare UpdateMotionMatchingState Against MotionMatch. 0:disabled 1:enabled 2:compare 3:debug restore MM state
+	if (GVarAnimNodeMotionMatchingCompareUpdateMotionMatchingStateAgainstMotionMatch >= 1)
 	{
 		MotionMatchingStateBackup = MotionMatchingState;
 		MotionMatchingStateCopy = MotionMatchingState;
@@ -223,33 +225,42 @@ void FAnimNode_MotionMatching::UpdateAssetPlayer(const FAnimationUpdateContext&
 		MotionMatchingStateCopy.WantedPlayRate = WantedPlayRate;
 		PRAGMA_ENABLE_DEPRECATION_WARNINGS
 	}
+	
+	// Compare UpdateMotionMatchingState Against MotionMatch. 0:disabled 1:enabled 2:compare 3:debug restore MM state
+	if (GVarAnimNodeMotionMatchingCompareUpdateMotionMatchingStateAgainstMotionMatch == 1)
+	{
+		MotionMatchingState = MotionMatchingStateCopy;
+	}
+	else
 #endif // ENABLE_ANIM_DEBUG
-
-	FChooserEvaluationContext EvaluationContext(Context.AnimInstanceProxy->GetAnimInstanceObject());
-	UPoseSearchLibrary::UpdateMotionMatchingState(
-		&EvaluationContext,
-		PoseHistory,
-		DatabasesToSearch,
-		DeltaTime,
-		PoseJumpThresholdTime,
-		PoseReselectHistory,
-		bShouldSearch ? SearchThrottleTime : UE_BIG_NUMBER,
-		PoseSearchPlayRate,
-		MotionMatchingState,
-		NextUpdateInterruptMode,
-		bShouldUseCachedChannelData
+	{
+		FChooserEvaluationContext EvaluationContext(Context.AnimInstanceProxy->GetAnimInstanceObject());
+		UPoseSearchLibrary::UpdateMotionMatchingState(
+			&EvaluationContext,
+			PoseHistory,
+			DatabasesToSearch,
+			DeltaTime,
+			PoseJumpThresholdTime,
+			PoseReselectHistory,
+			bShouldSearch ? SearchThrottleTime : UE_BIG_NUMBER,
+			PoseSearchPlayRate,
+			MotionMatchingState,
+			NextUpdateInterruptMode,
+			bShouldUseCachedChannelData
 #if ENABLE_ANIM_DEBUG
-		, GVarAnimNodeMotionMatchingDrawQuery
-		, GVarAnimNodeMotionMatchingDrawCurResult
+			, GVarAnimNodeMotionMatchingDrawQuery
+			, GVarAnimNodeMotionMatchingDrawCurResult
 #else // ENABLE_ANIM_DEBUG
-		, false
-		, false
+			, false
+			, false
 #endif // ENABLE_ANIM_DEBUG
-		, EventToSearch
-	);
+			, EventToSearch
+		);
+	}
 
 #if ENABLE_ANIM_DEBUG
-	if (GVarAnimNodeMotionMatchingCompareUpdateMotionMatchingStateAgainstMotionMatch)
+	// Compare UpdateMotionMatchingState Against MotionMatch. 0:disabled 1:enabled 2:compare 3:debug restore MM state
+	if (GVarAnimNodeMotionMatchingCompareUpdateMotionMatchingStateAgainstMotionMatch >= 2)
 	{
 		bool bUpdateMotionMatchingDiffersFromMotionMatch = false;
 		PRAGMA_DISABLE_DEPRECATION_WARNINGS
@@ -284,7 +295,7 @@ void FAnimNode_MotionMatching::UpdateAssetPlayer(const FAnimationUpdateContext&
 		}
 		PRAGMA_ENABLE_DEPRECATION_WARNINGS
 
-		if (GVarAnimNodeMotionMatchingCompareUpdateMotionMatchingStateAgainstMotionMatch > 1 && bUpdateMotionMatchingDiffersFromMotionMatch)
+		if (GVarAnimNodeMotionMatchingCompareUpdateMotionMatchingStateAgainstMotionMatch == 3 && bUpdateMotionMatchingDiffersFromMotionMatch)
 		{
 			MotionMatchingStateCopy = MotionMatchingStateBackup;
 			MotionMatchingState = MotionMatchingStateBackup;
diff --git a/Engine/Plugins/Animation/PoseSearch/Source/Runtime/Private/PoseSearchInteractionIsland.cpp b/Engine/Plugins/Animation/PoseSearch/Source/Runtime/Private/PoseSearchInteractionIsland.cpp
--- a/Engine/Plugins/Animation/PoseSearch/Source/Runtime/Private/PoseSearchInteractionIsland.cpp
+++ b/Engine/Plugins/Animation/PoseSearch/Source/Runtime/Private/PoseSearchInteractionIsland.cpp
@@ -997,17 +997,16 @@ bool FInteractionIsland::DoSearch_AnyThread(const UObject* AnimContext, const TC
 
 	FMemMark Mark(FMemStack::Get());
 
-	TArray<const UObject*, TInlineAllocator<PreallocatedRolesNum, TMemStackAllocator<>>> AnimContexts;
-	TArray<const UE::PoseSearch::IPoseHistory*, TInlineAllocator<PreallocatedRolesNum, TMemStackAllocator<>>> PoseHistories;
 	TArray<FSearchResult, TMemStackAllocator<>> PoseSearchResults;
+	TArray<FChooserEvaluationContext, TInlineAllocator<UE::PoseSearch::PreallocatedRolesNum>> Contexts;
 
 	// SearchContexts are modified only by UPoseSearchInteractionSubsystem::Tick and constant otherwise, so it's safe to access them in a threaded enviroment without locks
 	PoseSearchResults.SetNum(SearchContexts.Num());
 
 	for (int32 SearchIndex = 0; SearchIndex < SearchContexts.Num(); ++SearchIndex)
 	{
-		FInteractionSearchContext& SearchContext = SearchContexts[SearchIndex];
-		const UPoseSearchDatabase* Database = SearchContext.Database.Get();
+		FInteractionSearchContext& InteractionSearchContext = SearchContexts[SearchIndex];
+		const UPoseSearchDatabase* Database = InteractionSearchContext.Database.Get();
 		if (!Database)
 		{
 			UE_LOG(LogPoseSearch, Error, TEXT("FInteractionIsland::DoSearch_AnyThread invalid context database"));
@@ -1020,38 +1019,34 @@ bool FInteractionIsland::DoSearch_AnyThread(const UObject* AnimContext, const TC
 			return false;
 		}
 
-		const int32 NumRoles = SearchContext.Num();
-		
-		AnimContexts.Reset();
-		PoseHistories.Reset();
+		const int32 NumRoles = InteractionSearchContext.Num();
+		Contexts.SetNum(NumRoles);
 
+		FSearchContext SearchContext(0.f, nullptr, FFloatInterval(0.f, 0.f), FPoseSearchEvent());
 		for (int32 RoleIndex = 0; RoleIndex < NumRoles; ++RoleIndex)
 		{
-			const UObject* SearchContextAnimContext = SearchContext.GetAnimContext(RoleIndex);
+			const UObject* SearchContextAnimContext = InteractionSearchContext.GetAnimContext(RoleIndex);
 			if (!SearchContextAnimContext)
 			{
 				UE_LOG(LogPoseSearch, Error, TEXT("FInteractionIsland::DoSearch_AnyThread null anim context"));
 				return false;
 			}
 
-			AnimContexts.Add(SearchContextAnimContext);
+			Contexts[RoleIndex].AddObjectParam(const_cast<UObject*>(SearchContextAnimContext));
 
-			const UE::PoseSearch::IPoseHistory* PoseHistory = SearchContext.GetPoseHistory(RoleIndex);
+			const UE::PoseSearch::IPoseHistory* PoseHistory = InteractionSearchContext.GetPoseHistory(RoleIndex);
 			if (!PoseHistory)
 			{
 				UE_LOG(LogPoseSearch, Error, TEXT("FInteractionIsland::DoSearch_AnyThread null PoseHistory"));
 				return false;
 			}
-
-			PoseHistories.Add(PoseHistory);
+			
+			SearchContext.AddRole(InteractionSearchContext.GetRoles()[RoleIndex], &Contexts[RoleIndex], PoseHistory);
 		}
 
 		const UObject* AssetsToSearch[] = { Database };
-		FPoseSearchFutureProperties PoseSearchFutureProperties;
-
 		// @todo: we could perform multiple UPoseSearchLibrary::MotionMatch in parallel!
-		PoseSearchResults[SearchIndex] = UPoseSearchLibrary::MotionMatch(AnimContexts, SearchContext.GetRoles(), PoseHistories,
-			AssetsToSearch, SearchContext.GetContinuingProperties(), PoseSearchFutureProperties, FPoseSearchEvent());
+		PoseSearchResults[SearchIndex] = UPoseSearchLibrary::MotionMatch(SearchContext, AssetsToSearch, InteractionSearchContext.GetContinuingProperties());
 	}
 		
 	InitSearchResults(SearchResults, PoseSearchResults, SearchContexts);
diff --git a/Engine/Plugins/Animation/PoseSearch/Source/Runtime/Public/PoseSearch/PoseSearchResult.h b/Engine/Plugins/Animation/PoseSearch/Source/Runtime/Public/PoseSearch/PoseSearchResult.h
--- a/Engine/Plugins/Animation/PoseSearch/Source/Runtime/Public/PoseSearch/PoseSearchResult.h
+++ b/Engine/Plugins/Animation/PoseSearch/Source/Runtime/Public/PoseSearch/PoseSearchResult.h
@@ -27,6 +27,7 @@ struct FDatabasePoseIdx
 	// database associated to this search result. Poses are derived data stored in Database->GetSearchIndex()
 	TWeakObjectPtr<const UPoseSearchDatabase> Database;
 
+	FDatabasePoseIdx() = default;
 	bool IsValid() const { return PoseIdx != INDEX_NONE && Database != nullptr; }
 	bool operator==(const FDatabasePoseIdx& Other) const { return PoseIdx == Other.PoseIdx && Database == Other.Database; }
 
@@ -60,6 +61,8 @@ struct FSearchResult : FDatabasePoseIdx
 	int32 BestPosePos = 0;
 #endif // WITH_EDITOR && ENABLE_ANIM_DEBUG && UE_POSE_SEARCH_TRACE_ENABLED
 
+	FSearchResult() = default;
+
 	// Experimental, this feature might be removed without warning, not for production use
 	// Attempts to set the internal state to match the provided asset time including updating the internal DbPoseIdx. 
 	// If the provided asset time is out of bounds for the currently playing asset then this function will reset the 
