PoseSearch - fix for FSearchContext::GetSampleRotationInternal, FAssetIndexer::GetSampleRotation not properly calculating the rotation relative to the reference bone

diff --git a/Engine/Plugins/Animation/PoseSearch/Source/Runtime/Private/PoseSearchAssetIndexer.cpp b/Engine/Plugins/Animation/PoseSearch/Source/Runtime/Private/PoseSearchAssetIndexer.cpp
--- a/Engine/Plugins/Animation/PoseSearch/Source/Runtime/Private/PoseSearchAssetIndexer.cpp
+++ b/Engine/Plugins/Animation/PoseSearch/Source/Runtime/Private/PoseSearchAssetIndexer.cpp
@@ -634,9 +634,10 @@ bool FAssetIndexer::GetSampleRotation(FQuat& OutSampleRotation, float SampleTime
 	}
 
 	bool bUnused;
-	const FTransform RootBoneTransform = GetTransform(OriginTime, OriginRole, bUnused, RootSchemaBoneIdx);
+	const FTransform OriginBoneTransform = GetTransform(OriginTime, OriginRole, bUnused, SchemaOriginBoneIdx);
 	const FTransform SampleBoneTransform = GetTransform(SampleTime, SampleRole, bUnused, SchemaSampleBoneIdx);
-	OutSampleRotation = RootBoneTransform.InverseTransformRotation(SampleBoneTransform.GetRotation());
+	OutSampleRotation = OriginBoneTransform.InverseTransformRotation(SampleBoneTransform.GetRotation());
+
 	return true;
 }
 
diff --git a/Engine/Plugins/Animation/PoseSearch/Source/Runtime/Private/PoseSearchContext.cpp b/Engine/Plugins/Animation/PoseSearch/Source/Runtime/Private/PoseSearchContext.cpp
--- a/Engine/Plugins/Animation/PoseSearch/Source/Runtime/Private/PoseSearchContext.cpp
+++ b/Engine/Plugins/Animation/PoseSearch/Source/Runtime/Private/PoseSearchContext.cpp
@@ -1138,9 +1138,9 @@ FQuat FSearchContext::GetSampleRotationInternal(float SampleTime, float OriginTi
 		return RootBoneTransform.InverseTransformRotation(DeltaBoneRotation);
 	}
 
-	const FTransform RootBoneTransform = GetWorldBoneTransformAtTime(OriginTime, OriginRole, RootSchemaBoneIdx);
+	const FTransform OriginBoneTransform = GetWorldBoneTransformAtTime(OriginTime, OriginRole, SchemaOriginBoneIdx);
 	const FTransform SampleBoneTransform = GetWorldBoneTransformAtTime(SampleTime, SampleRole, SchemaSampleBoneIdx);
-	return RootBoneTransform.InverseTransformRotation(SampleBoneTransform.GetRotation());
+	return OriginBoneTransform.InverseTransformRotation(SampleBoneTransform.GetRotation());
 }
 
 TArrayView<float> FSearchContext::EditFeatureVector()
diff --git a/Engine/Plugins/Animation/PoseSearch/Source/Runtime/Public/PoseSearch/PoseSearchAssetIndexer.h b/Engine/Plugins/Animation/PoseSearch/Source/Runtime/Public/PoseSearch/PoseSearchAssetIndexer.h
--- a/Engine/Plugins/Animation/PoseSearch/Source/Runtime/Public/PoseSearch/PoseSearchAssetIndexer.h
+++ b/Engine/Plugins/Animation/PoseSearch/Source/Runtime/Public/PoseSearch/PoseSearchAssetIndexer.h
@@ -126,8 +126,8 @@ private:
 	// bClamped will be true if SampleTime is outside the animation duration boundaries
 	POSESEARCH_API FTransform GetComponentSpaceTransform(float SampleTime, const FRole& Role, bool& bClamped, int8 SchemaBoneIdx = RootSchemaBoneIdx);
 
-	// Returns OutSamplePosition as the position of the bone Schema.BoneReferences[SchemaSampleBoneIdx] at time SampleTime relative to the
-	// transform of the bone Schema.BoneReferences[SchemaOriginBoneIdx] at time OriginTime.
+	// Returns OutSamplePosition as the displacement of the bone Schema.BoneReferences[SchemaSampleBoneIdx] position at time SampleTime relative to the
+	// position of the bone Schema.BoneReferences[SchemaOriginBoneIdx] at time OriginTime in root bone space
 	// bClamped will be true if SampleTime or OriginTime are outside the animation duration boundaries
 	POSESEARCH_API bool GetSamplePositionInternal(FVector& OutSamplePosition, float SampleTime, float OriginTime, bool& bClamped, int8 SchemaSampleBoneIdx, int8 SchemaOriginBoneIdx, const FRole& SampleRole, const FRole& OriginRole, int32 SamplingAttributeId);
 
diff --git a/Engine/Plugins/Animation/PoseSearch/Source/Runtime/Public/PoseSearch/PoseSearchContext.h b/Engine/Plugins/Animation/PoseSearch/Source/Runtime/Public/PoseSearch/PoseSearchContext.h
--- a/Engine/Plugins/Animation/PoseSearch/Source/Runtime/Public/PoseSearch/PoseSearchContext.h
+++ b/Engine/Plugins/Animation/PoseSearch/Source/Runtime/Public/PoseSearch/PoseSearchContext.h
@@ -169,18 +169,18 @@ struct FSearchContext
 	// If the curve is not found, assume value of 0, which is consistent with curve behavior in the animation update.
 	UE_API float GetSampleCurveValue(float SampleTimeOffset, const FName& CurveName, const FRole& SampleRole);
 
-	// Returns the rotation of the bone Schema.BoneReferences[SchemaSampleBoneIdx] at an offset time of SampleTimeOffset relative to the
-	// transform of the bone Schema.BoneReferences[SchemaOriginBoneIdx] at an offset time of time OriginTimeOffset 
+	// Returns the rotation of the bone with index Schema.BoneReferences[SchemaSampleBoneIdx] at an offset time of SampleTimeOffset relative to the
+	// transform of the bone with index Schema.BoneReferences[SchemaOriginBoneIdx] at an offset time of time OriginTimeOffset 
 	// Times will be processed by GetPermutationTimeOffsets(PermutationTimeType, ...)
 	UE_API FQuat GetSampleRotation(float SampleTimeOffset, float OriginTimeOffset, int8 SchemaSampleBoneIdx, int8 SchemaOriginBoneIdx, const FRole& SampleRole, const FRole& OriginRole, EPermutationTimeType PermutationTimeType = EPermutationTimeType::UseSampleTime, const FQuat* SampleBoneRotationWorldOverride = nullptr);
 	
-	// Returns the position of the bone Schema.BoneReferences[SchemaSampleBoneIdx] at an offset time of SampleTimeOffset relative to the
-	// transform of the bone Schema.BoneReferences[SchemaOriginBoneIdx] at an offset time of time OriginTimeOffset 
+	// Returns the displacement of the bone with index Schema.BoneReferences[SchemaSampleBoneIdx] at an offset time of SampleTimeOffset relative to the
+	// position of the bone with index Schema.BoneReferences[SchemaOriginBoneIdx] at an offset time of OriginTimeOffset in root bone space
 	// Times will be processed by GetPermutationTimeOffsets(PermutationTimeType, ...)
 	UE_API FVector GetSamplePosition(float SampleTimeOffset, float OriginTimeOffset, int8 SchemaSampleBoneIdx, int8 SchemaOriginBoneIdx, const FRole& SampleRole, const FRole& OriginRole, EPermutationTimeType PermutationTimeType = EPermutationTimeType::UseSampleTime, const FVector* SampleBonePositionWorldOverride = nullptr);
 	
-	// Returns the delta velocity of the velocity of the bone Schema.BoneReferences[SchemaSampleBoneIdx] at an offset time of SampleTimeOffset minus
-	// the velocity of the bone Schema.BoneReferences[SchemaOriginBoneIdx] at an offset time of time OriginTimeOffset 
+	// Returns the delta velocity of the velocity of the bone with index Schema.BoneReferences[SchemaSampleBoneIdx] at an offset time of SampleTimeOffset minus
+	// the velocity of the bone with index Schema.BoneReferences[SchemaOriginBoneIdx] at an offset time of OriginTimeOffset 
 	// Times will be processed by GetPermutationTimeOffsets(PermutationTimeType, ...)
 	// if bUseCharacterSpaceVelocities is true, velocities will be computed in root bone space, rather than world space
 	UE_API FVector GetSampleVelocity(float SampleTimeOffset, float OriginTimeOffset, int8 SchemaSampleBoneIdx, int8 SchemaOriginBoneIdx, const FRole& SampleRole, const FRole& OriginRole, bool bUseCharacterSpaceVelocities = true, EPermutationTimeType PermutationTimeType = EPermutationTimeType::UseSampleTime, const FVector* SampleBoneVelocityWorldOverride = nullptr);
diff --git a/Engine/Plugins/Animation/PoseSearch/Source/Runtime/Public/PoseSearch/PoseSearchDefines.h b/Engine/Plugins/Animation/PoseSearch/Source/Runtime/Public/PoseSearch/PoseSearchDefines.h
--- a/Engine/Plugins/Animation/PoseSearch/Source/Runtime/Public/PoseSearch/PoseSearchDefines.h
+++ b/Engine/Plugins/Animation/PoseSearch/Source/Runtime/Public/PoseSearch/PoseSearchDefines.h
@@ -43,7 +43,7 @@ static constexpr int32 PreallocatedSearchesNum = 4;
 static constexpr float FiniteDelta = 1 / 60.0f; // Time delta used for computing pose derivatives
 static constexpr int32 MaxNumberOfCollectedPoseCandidatesPerDatabase = 256;
 static constexpr int32 InvalidSearchId = INDEX_NONE; // @todo: it should be 0, since the search id is an hash!
-static constexpr int32 DatabaseIndexDerivedDataCacheKeyVersion = 2708; // bump this value to force reindex ALL the UPoseSearchDatabases, without having to change the global FDevSystemGuids::POSESEARCHDB_DERIVEDDATA_VER
+static constexpr int32 DatabaseIndexDerivedDataCacheKeyVersion = 2713; // bump this value to force reindex ALL the UPoseSearchDatabases, without having to change the global FDevSystemGuids::POSESEARCHDB_DERIVEDDATA_VER
 static constexpr int32 DefaultTickPriority = 100;
 
 template<uint32 Alignment = DEFAULT_ALIGNMENT> class TMemStackSetAllocator : public TSetAllocator<TSparseArrayAllocator<TMemStackAllocator<Alignment>, TMemStackAllocator<Alignment>>, TMemStackAllocator<Alignment>> {};
