PoseSearch - UpdateMotionMatchingState deprecation

diff --git a/Engine/Plugins/Animation/PoseSearch/Source/Runtime/Private/AnimNode_MotionMatching.cpp b/Engine/Plugins/Animation/PoseSearch/Source/Runtime/Private/AnimNode_MotionMatching.cpp
--- a/Engine/Plugins/Animation/PoseSearch/Source/Runtime/Private/AnimNode_MotionMatching.cpp
+++ b/Engine/Plugins/Animation/PoseSearch/Source/Runtime/Private/AnimNode_MotionMatching.cpp
@@ -46,9 +46,6 @@ static FAutoConsoleVariableRef CVarAnimNodeMotionMatchingDrawInfoVerbose(TEXT("a
 static float GVarAnimNodeMotionMatchingDrawInfoHeight = 50.f;
 static FAutoConsoleVariableRef CVarAnimNodeMotionMatchingDrawInfoHeight(TEXT("a.AnimNode.MotionMatching.DebugDrawInfoHeight"), GVarAnimNodeMotionMatchingDrawInfoHeight, TEXT("Vertical offset for DebugDrawInfo"));
 
-static int32 GVarAnimNodeMotionMatchingCompareUpdateMotionMatchingStateAgainstMotionMatch = 0;
-static FAutoConsoleVariableRef CVarAnimNodeMotionMatchingCompareUpdateMotionMatchingStateAgainstMotionMatch(TEXT("a.AnimNode.MotionMatching.CompareUpdateMotionMatchingStateAgainstMotionMatch"), GVarAnimNodeMotionMatchingCompareUpdateMotionMatchingStateAgainstMotionMatch, TEXT("Compare UpdateMotionMatchingState Against MotionMatch. 0:disabled 1:enabled 2:compare 3:debug restore MM state"));
-
 static int32 GVarAnimNodeMotionMatchingPlayRateEnabled = UE::Private::EPlayRateState::Enabled;
 static FAutoConsoleVariableRef CVarAnimNodeMotionMatchingPlayRateEnabled(
 	TEXT("a.AnimNode.MotionMatching.DebugPlayRateEnabled"), 
@@ -162,149 +159,56 @@ void FAnimNode_MotionMatching::UpdateAssetPlayer(const FAnimationUpdateContext&
 
 	const float DeltaTime = Context.GetDeltaTime();
 
-#if ENABLE_ANIM_DEBUG
-	// debug code to validate UPoseSearchLibrary::MotionMatch against UPoseSearchLibrary::UpdateMotionMatchingState to deprecate the latter in the near future,
-	// and expose FSearchContext to be able to perform multiple searches with the same cached queries (to speed up MM interactions)
-	FMotionMatchingState MotionMatchingStateBackup;
-	FMotionMatchingState MotionMatchingStateCopy;
+	FSearchContext SearchContext(0.f, &MotionMatchingState.PoseIndicesHistory, PoseJumpThresholdTime, EventToSearch.GetPlayRateOverriddenEvent(PoseSearchPlayRate));
+	FChooserEvaluationContext ChooserEvaluationContext(const_cast<UObject*>(Context.AnimInstanceProxy->GetAnimInstanceObject()));
+	SearchContext.AddRole(GetCommonDefaultRole(DatabasesToSearch), &ChooserEvaluationContext, PoseHistory);
 
-	// Compare UpdateMotionMatchingState Against MotionMatch. 0:disabled 1:enabled 2:compare 3:debug restore MM state
-	if (GVarAnimNodeMotionMatchingCompareUpdateMotionMatchingStateAgainstMotionMatch >= 1)
+	// If we can't advance or enough time has elapsed since the last pose jump then search
+	const bool bCanAdvance = SearchContext.GetContinuingPoseSearchResult().PoseIdx != INDEX_NONE;
+	const bool bSearch = !bCanAdvance || (MotionMatchingState.ElapsedPoseSearchTime >= SearchThrottleTime);
+	FSearchResult SearchResult;
+	if (bSearch)
 	{
-		MotionMatchingStateBackup = MotionMatchingState;
-		MotionMatchingStateCopy = MotionMatchingState;
-
-		const FPoseSearchEvent PlayRateOverriddenEvent = EventToSearch.GetPlayRateOverriddenEvent(PoseSearchPlayRate);
-		FSearchContext SearchContext(0.f, &MotionMatchingStateCopy.PoseIndicesHistory, PoseJumpThresholdTime, PlayRateOverriddenEvent);
-		FChooserEvaluationContext ChooserEvaluationContext(const_cast<UObject*>(Context.AnimInstanceProxy->GetAnimInstanceObject()));
-		SearchContext.AddRole(GetCommonDefaultRole(DatabasesToSearch), &ChooserEvaluationContext, PoseHistory);
-
-		// If we can't advance or enough time has elapsed since the last pose jump then search
-		const bool bCanAdvance = SearchContext.GetContinuingPoseSearchResult().PoseIdx != INDEX_NONE;
-		const bool bSearch = !bCanAdvance || (MotionMatchingStateCopy.ElapsedPoseSearchTime >= SearchThrottleTime);
-		FSearchResult SearchResult;
-		if (bSearch)
-		{
-			MotionMatchingStateCopy.ElapsedPoseSearchTime = 0.f;
+		MotionMatchingState.ElapsedPoseSearchTime = 0.f;
 			
-			TArrayView<const UObject*> AssetsToSearch((const UObject**)FMemory_Alloca(DatabasesToSearch.Num() * sizeof(UObject*)), (int32)DatabasesToSearch.Num());
-			for (int32 DatabaseToSearchIndex = 0; DatabaseToSearchIndex < DatabasesToSearch.Num(); ++DatabaseToSearchIndex)
-			{
-				AssetsToSearch[DatabaseToSearchIndex] = DatabasesToSearch[DatabaseToSearchIndex].Get();
-			}
-
-			FPoseSearchContinuingProperties ContinuingProperties;
-			ContinuingProperties.InitFrom(MotionMatchingStateCopy.SearchResult, NextUpdateInterruptMode);
-			SearchResult = UPoseSearchLibrary::MotionMatch(SearchContext, AssetsToSearch, ContinuingProperties);
-		}
-		else
+		TArrayView<const UObject*> AssetsToSearch((const UObject**)FMemory_Alloca(DatabasesToSearch.Num() * sizeof(UObject*)), (int32)DatabasesToSearch.Num());
+		for (int32 DatabaseToSearchIndex = 0; DatabaseToSearchIndex < DatabasesToSearch.Num(); ++DatabaseToSearchIndex)
 		{
-			MotionMatchingStateCopy.ElapsedPoseSearchTime += DeltaTime;
-
-			// if we didn't perform a search we carry on with the continuing pose search result
-			SearchResult = SearchContext.GetContinuingPoseSearchResult();
-			SearchResult.bIsContinuingPoseSearch = true;
-
-	#if UE_POSE_SEARCH_TRACE_ENABLED
-			// in case we skipped the search, we still have to track we would have requested to evaluate DatabasesToSearch and SearchResult.Database
-			for (const TObjectPtr<const UPoseSearchDatabase>& DatabaseToSearch : DatabasesToSearch)
-			{
-				SearchContext.Track(DatabaseToSearch.Get());
-			}
-
-			SearchContext.Track(SearchResult.Database.Get());
-	#endif // UE_POSE_SEARCH_TRACE_ENABLED
+			AssetsToSearch[DatabaseToSearchIndex] = DatabasesToSearch[DatabaseToSearchIndex].Get();
 		}
 
-		const float WantedPlayRate = CalculateWantedPlayRate(SearchResult, SearchContext, PoseSearchPlayRate, PoseHistory ? PoseHistory->GetTrajectorySpeedMultiplier() : 1.f, EventToSearch);
-		MotionMatchingStateCopy.SearchResult.InitFrom(SearchResult, WantedPlayRate);
-		MotionMatchingStateCopy.PoseIndicesHistory.Update(SearchResult, DeltaTime, PoseReselectHistory);
-
-		PRAGMA_DISABLE_DEPRECATION_WARNINGS
-		MotionMatchingStateCopy.CurrentSearchResult = SearchResult;
-		MotionMatchingStateCopy.WantedPlayRate = WantedPlayRate;
-		PRAGMA_ENABLE_DEPRECATION_WARNINGS
-	}
-	
-	// Compare UpdateMotionMatchingState Against MotionMatch. 0:disabled 1:enabled 2:compare 3:debug restore MM state
-	if (GVarAnimNodeMotionMatchingCompareUpdateMotionMatchingStateAgainstMotionMatch == 1)
-	{
-		MotionMatchingState = MotionMatchingStateCopy;
+		FPoseSearchContinuingProperties ContinuingProperties;
+		ContinuingProperties.InitFrom(MotionMatchingState.SearchResult, NextUpdateInterruptMode);
+		SearchResult = UPoseSearchLibrary::MotionMatch(SearchContext, AssetsToSearch, ContinuingProperties);
 	}
 	else
-#endif // ENABLE_ANIM_DEBUG
 	{
-		FChooserEvaluationContext EvaluationContext(Context.AnimInstanceProxy->GetAnimInstanceObject());
-		UPoseSearchLibrary::UpdateMotionMatchingState(
-			&EvaluationContext,
-			PoseHistory,
-			DatabasesToSearch,
-			DeltaTime,
-			PoseJumpThresholdTime,
-			PoseReselectHistory,
-			bShouldSearch ? SearchThrottleTime : UE_BIG_NUMBER,
-			PoseSearchPlayRate,
-			MotionMatchingState,
-			NextUpdateInterruptMode,
-			bShouldUseCachedChannelData
-#if ENABLE_ANIM_DEBUG
-			, GVarAnimNodeMotionMatchingDrawQuery
-			, GVarAnimNodeMotionMatchingDrawCurResult
-#else // ENABLE_ANIM_DEBUG
-			, false
-			, false
-#endif // ENABLE_ANIM_DEBUG
-			, EventToSearch
-		);
-	}
+		MotionMatchingState.ElapsedPoseSearchTime += DeltaTime;
 
-#if ENABLE_ANIM_DEBUG
-	// Compare UpdateMotionMatchingState Against MotionMatch. 0:disabled 1:enabled 2:compare 3:debug restore MM state
-	if (GVarAnimNodeMotionMatchingCompareUpdateMotionMatchingStateAgainstMotionMatch >= 2)
-	{
-		bool bUpdateMotionMatchingDiffersFromMotionMatch = false;
-		PRAGMA_DISABLE_DEPRECATION_WARNINGS
-		if (MotionMatchingState.CurrentSearchResult != MotionMatchingStateCopy.CurrentSearchResult)
-		{
-			UE_LOG(LogPoseSearch, Warning, TEXT("MotionMatchingState.CurrentSearchResult != MotionMatchingStateCopy.CurrentSearchResult"));
-			bUpdateMotionMatchingDiffersFromMotionMatch = true;
-		}
-
-		if (MotionMatchingState.SearchResult != MotionMatchingStateCopy.SearchResult)
-		{
-			UE_LOG(LogPoseSearch, Warning, TEXT("MotionMatchingState.SearchResult != MotionMatchingStateCopy.SearchResult"));
-			bUpdateMotionMatchingDiffersFromMotionMatch = true;
-		}
-
-		if (MotionMatchingState.ElapsedPoseSearchTime != MotionMatchingStateCopy.ElapsedPoseSearchTime)
-		{
-			UE_LOG(LogPoseSearch, Warning, TEXT("MotionMatchingState.ElapsedPoseSearchTime != MotionMatchingStateCopy.ElapsedPoseSearchTime"));
-			bUpdateMotionMatchingDiffersFromMotionMatch = true;
-		}
-
-		if (MotionMatchingState.WantedPlayRate != MotionMatchingStateCopy.WantedPlayRate)
-		{
-			UE_LOG(LogPoseSearch, Warning, TEXT("MotionMatchingState.CurrentSearchResult != MotionMatchingStateCopy.CurrentSearchResult"));
-			bUpdateMotionMatchingDiffersFromMotionMatch = true;
-		}
+		// if we didn't perform a search we carry on with the continuing pose search result
+		SearchResult = SearchContext.GetContinuingPoseSearchResult();
+		SearchResult.bIsContinuingPoseSearch = true;
 
-		if (MotionMatchingState.PoseIndicesHistory != MotionMatchingStateCopy.PoseIndicesHistory)
+#if UE_POSE_SEARCH_TRACE_ENABLED
+		// in case we skipped the search, we still have to track we would have requested to evaluate DatabasesToSearch and SearchResult.Database
+		for (const TObjectPtr<const UPoseSearchDatabase>& DatabaseToSearch : DatabasesToSearch)
 		{
-			UE_LOG(LogPoseSearch, Warning, TEXT("MotionMatchingState.PoseIndicesHistory != MotionMatchingStateCopy.PoseIndicesHistory"));
-			bUpdateMotionMatchingDiffersFromMotionMatch = true;
+			SearchContext.Track(DatabaseToSearch.Get());
 		}
-		PRAGMA_ENABLE_DEPRECATION_WARNINGS
 
-		if (GVarAnimNodeMotionMatchingCompareUpdateMotionMatchingStateAgainstMotionMatch == 3 && bUpdateMotionMatchingDiffersFromMotionMatch)
-		{
-			MotionMatchingStateCopy = MotionMatchingStateBackup;
-			MotionMatchingState = MotionMatchingStateBackup;
-
-			// the idea is that now you can run MM searches again and analyze why they differs
-		}
+		SearchContext.Track(SearchResult.Database.Get());
+#endif // UE_POSE_SEARCH_TRACE_ENABLED
 	}
-#endif // ENABLE_ANIM_DEBUG
 
+	const float WantedPlayRate = CalculateWantedPlayRate(SearchResult, SearchContext, PoseSearchPlayRate, PoseHistory ? PoseHistory->GetTrajectorySpeedMultiplier() : 1.f, EventToSearch);
+	MotionMatchingState.SearchResult.InitFrom(SearchResult, WantedPlayRate);
+	MotionMatchingState.PoseIndicesHistory.Update(SearchResult, DeltaTime, PoseReselectHistory);
+
+	PRAGMA_DISABLE_DEPRECATION_WARNINGS
+	MotionMatchingState.CurrentSearchResult = SearchResult;
+	MotionMatchingState.WantedPlayRate = WantedPlayRate;
+	PRAGMA_ENABLE_DEPRECATION_WARNINGS
+	
 	const bool bJumpToPose = MotionMatchingState.SearchResult.SelectedAnim && !MotionMatchingState.SearchResult.bIsContinuingPoseSearch;
 
 	PRAGMA_DISABLE_DEPRECATION_WARNINGS
diff --git a/Engine/Plugins/Animation/PoseSearch/Source/Runtime/Private/PoseSearchLibrary.cpp b/Engine/Plugins/Animation/PoseSearch/Source/Runtime/Private/PoseSearchLibrary.cpp
--- a/Engine/Plugins/Animation/PoseSearch/Source/Runtime/Private/PoseSearchLibrary.cpp
+++ b/Engine/Plugins/Animation/PoseSearch/Source/Runtime/Private/PoseSearchLibrary.cpp
@@ -755,9 +755,11 @@ void UPoseSearchLibrary::UpdateMotionMatchingState(
 
 	check(Context.AnimInstanceProxy);
 	FChooserEvaluationContext AnimContext(Context.AnimInstanceProxy->GetAnimInstanceObject());
+	PRAGMA_DISABLE_DEPRECATION_WARNINGS
 	UpdateMotionMatchingState(&AnimContext, PoseHistory, Databases, Context.GetDeltaTime(),
 		PoseJumpThresholdTime, PoseReselectHistory, bShouldSearch ? SearchThrottleTime : UE_BIG_NUMBER, PlayRate, InOutMotionMatchingState,
 		InterruptMode, bShouldUseCachedChannelData, bDebugDrawQuery, bDebugDrawCurResult);
+	PRAGMA_ENABLE_DEPRECATION_WARNINGS
 }
 
 void UPoseSearchLibrary::UpdateMotionMatchingState(
@@ -777,6 +779,7 @@ void UPoseSearchLibrary::UpdateMotionMatchingState(
 	const FPoseSearchEvent& EventToSearch)
 {
 	FChooserEvaluationContext Context(const_cast<UObject*>(AnimContext));
+	PRAGMA_DISABLE_DEPRECATION_WARNINGS
 	UpdateMotionMatchingState(&Context,
 		PoseHistory,
 		Databases,
@@ -789,6 +792,7 @@ void UPoseSearchLibrary::UpdateMotionMatchingState(
 		InterruptMode,
 		bShouldUseCachedChannelData,
 		bDebugDrawQuery, bDebugDrawCurResult);
+	PRAGMA_ENABLE_DEPRECATION_WARNINGS
 }
 
 void UPoseSearchLibrary::UpdateMotionMatchingState(
diff --git a/Engine/Plugins/Animation/PoseSearch/Source/Runtime/Public/PoseSearch/PoseSearchEvent.h b/Engine/Plugins/Animation/PoseSearch/Source/Runtime/Public/PoseSearch/PoseSearchEvent.h
--- a/Engine/Plugins/Animation/PoseSearch/Source/Runtime/Public/PoseSearch/PoseSearchEvent.h
+++ b/Engine/Plugins/Animation/PoseSearch/Source/Runtime/Public/PoseSearch/PoseSearchEvent.h
@@ -6,6 +6,8 @@
 #include "Kismet/BlueprintFunctionLibrary.h"
 #include "PoseSearchEvent.generated.h"
 
+#define UE_API POSESEARCH_API
+
 USTRUCT(Experimental, BlueprintType, Category="Animation|Pose Search|Experimental")
 struct FPoseSearchEvent
 {
@@ -13,7 +15,7 @@ struct FPoseSearchEvent
 
 	bool IsValid() const;
 	void Reset();
-	FPoseSearchEvent GetPlayRateOverriddenEvent(const FFloatInterval& PlayRateRangeBase) const;
+	UE_INTERNAL UE_API FPoseSearchEvent GetPlayRateOverriddenEvent(const FFloatInterval& PlayRateRangeBase) const;
 
 	// Tag identifying this event
 	UPROPERTY(EditAnywhere, BlueprintReadWrite, Category=State)
@@ -49,3 +51,5 @@ class UPoseSearchEventLibrary : public UBlueprintFunctionLibrary
 	UFUNCTION(BlueprintCallable, meta = (BlueprintThreadSafe), Category="Animation|Pose Search|Experimental")
 	static void UpdatePoseSearchEvent(const FPoseSearchEvent& InNewEvent, bool bIsNewEventValid, float DeltaSeconds, UPARAM(ref) FPoseSearchEvent& InOutCurrentEvent);
 };
+
+#undef UE_API
diff --git a/Engine/Plugins/Animation/PoseSearch/Source/Runtime/Public/PoseSearch/PoseSearchHistory.h b/Engine/Plugins/Animation/PoseSearch/Source/Runtime/Public/PoseSearch/PoseSearchHistory.h
--- a/Engine/Plugins/Animation/PoseSearch/Source/Runtime/Public/PoseSearch/PoseSearchHistory.h
+++ b/Engine/Plugins/Animation/PoseSearch/Source/Runtime/Public/PoseSearch/PoseSearchHistory.h
@@ -320,7 +320,7 @@ struct FHistoricalPoseIndex
 
 struct FPoseIndicesHistory
 {
-	void Update(const FSearchResult& SearchResult, float DeltaTime, float MaxTime);
+	UE_INTERNAL UE_API void Update(const FSearchResult& SearchResult, float DeltaTime, float MaxTime);
 	void Reset() { IndexToTime.Reset(); }
 	bool operator==(const FPoseIndicesHistory& Other) const;
 		
diff --git a/Engine/Plugins/Animation/PoseSearch/Source/Runtime/Public/PoseSearch/PoseSearchLibrary.h b/Engine/Plugins/Animation/PoseSearch/Source/Runtime/Public/PoseSearch/PoseSearchLibrary.h
--- a/Engine/Plugins/Animation/PoseSearch/Source/Runtime/Public/PoseSearch/PoseSearchLibrary.h
+++ b/Engine/Plugins/Animation/PoseSearch/Source/Runtime/Public/PoseSearch/PoseSearchLibrary.h
@@ -24,11 +24,8 @@ namespace UE::PoseSearch
 {
 	struct FSearchContext;
 
-	// Experimental, this feature might be removed without warning, not for production use
-	FRole GetCommonDefaultRole(const TConstArrayView<TObjectPtr<const UPoseSearchDatabase>> Databases);
-	// Experimental, this feature might be removed without warning, not for production use
-	float CalculateWantedPlayRate(const FSearchResult& SearchResult, const FSearchContext& SearchContext, const FFloatInterval& PlayRate, float TrajectorySpeedMultiplier, const FPoseSearchEvent& EventToSearch);
-
+	UE_INTERNAL UE_API FRole GetCommonDefaultRole(const TConstArrayView<TObjectPtr<const UPoseSearchDatabase>> Databases);
+	UE_INTERNAL UE_API float CalculateWantedPlayRate(const FSearchResult& SearchResult, const FSearchContext& SearchContext, const FFloatInterval& PlayRate, float TrajectorySpeedMultiplier, const FPoseSearchEvent& EventToSearch);
 } // namespace UE::PoseSearch
 
 
@@ -151,8 +148,7 @@ struct FPoseSearchContinuingProperties
 	UPROPERTY(Transient, EditAnywhere, BlueprintReadWrite, Category=State)
 	TObjectPtr<const UPoseSearchDatabase> PlayingAssetDatabase = nullptr;
 
-	// Experimental, this feature might be removed without warning, not for production use
-	void InitFrom(const FPoseSearchBlueprintResult& SearchResult, EPoseSearchInterruptMode InInterruptMode);
+	UE_INTERNAL UE_API void InitFrom(const FPoseSearchBlueprintResult& SearchResult, EPoseSearchInterruptMode InInterruptMode);
 };
 
 UCLASS(MinimalAPI)
@@ -215,6 +211,7 @@ public:
 	* @param bDebugDrawQuery				Input draw the composed query if valid
 	* @param bDebugDrawCurResult			Input draw the current result if valid
 	*/
+	UE_DEPRECATED(5.7, "Use MotionMatch API instead")
 	static UE_API void UpdateMotionMatchingState(
 		FChooserEvaluationContext* AnimContext,
 		const UE::PoseSearch::IPoseHistory* PoseHistory,
@@ -231,7 +228,7 @@ public:
 		bool bDebugDrawCurResult = false,
 		const FPoseSearchEvent& EventToSearch = FPoseSearchEvent());
 	
-	UE_DEPRECATED(5.6, "Use other UpdateMotionMatchingState signatures instead")
+	UE_DEPRECATED(5.6, "Use MotionMatch API instead")
 	static UE_API void UpdateMotionMatchingState(
 		const UObject* AnimContext,
 		const UE::PoseSearch::IPoseHistory* PoseHistory,
@@ -248,7 +245,7 @@ public:
 		bool bDebugDrawCurResult = false,
 		const FPoseSearchEvent& EventToSearch = FPoseSearchEvent());
 	
-	UE_DEPRECATED(5.6, "Use other UpdateMotionMatchingState signatures instead")
+	UE_DEPRECATED(5.6, "Use MotionMatch API instead")
 	static UE_API void UpdateMotionMatchingState(
 		const FAnimationUpdateContext& Context,
 		const TArray<TObjectPtr<const UPoseSearchDatabase>>& Databases,
@@ -271,6 +268,7 @@ public:
 	* @param AnimInstance					Input animation instance
 	* @param AssetsToSearch					Input assets to search (UPoseSearchDatabase or any animation asset containing UAnimNotifyState_PoseSearchBranchIn)
 	* @param PoseHistoryName				Input tag of the associated PoseSearchHistoryCollector node in the anim graph
+	* @param ContinuingProperties			Input struct specyfying the currently playing animation to be able to bias the search in selecting the continuing pose
 	* @param Future							Input future properties to match (animation / start time / time offset)
 	* @param SelectedAnimation				Output selected animation from the Database asset
 	* @param Result							Output FPoseSearchBlueprintResult with the search result
@@ -303,8 +301,7 @@ public:
 		const FPoseSearchFutureProperties& Future,
 		const FPoseSearchEvent& EventToSearch);
 	
-	// Experimental, this feature might be removed without warning, not for production use
-	static UE_API UE::PoseSearch::FSearchResult MotionMatch(
+	static UE_INTERNAL UE_API UE::PoseSearch::FSearchResult MotionMatch(
 		const TArrayView<FChooserEvaluationContext> Contexts,
 		const TArrayView<const UE::PoseSearch::FRole> Roles,
 		const TArrayView<const UE::PoseSearch::IPoseHistory*> PoseHistories, 
@@ -313,8 +310,7 @@ public:
 		const FPoseSearchFutureProperties& Future,
 		const FPoseSearchEvent& EventToSearch);
 			
-	// Experimental, this feature might be removed without warning, not for production use
-	static UE_API UE::PoseSearch::FSearchResult MotionMatch(
+	static UE_INTERNAL UE_API UE::PoseSearch::FSearchResult MotionMatch(
 		const TArrayView<FChooserEvaluationContext> Contexts,
 		const TArrayView<const UE::PoseSearch::FRole> Roles,
 		const TArrayView<const UE::PoseSearch::IPoseHistory*> PoseHistories,
@@ -323,8 +319,7 @@ public:
 		const float DesiredPermutationTimeOffset,
 		const FPoseSearchEvent& EventToSearch);
 
-	// Experimental, this feature might be removed without warning, not for production use
-	static UE_API UE::PoseSearch::FSearchResult MotionMatch(
+	static UE_INTERNAL UE_API UE::PoseSearch::FSearchResult MotionMatch(
 		const TArrayView<const UObject*> AnimContexts,
 		const TArrayView<const UE::PoseSearch::FRole> Roles,
 		const TArrayView<const UE::PoseSearch::IPoseHistory*> PoseHistories,
@@ -333,8 +328,7 @@ public:
 		const float DesiredPermutationTimeOffset,
 		const FPoseSearchEvent& EventToSearch);
 
-	// Experimental, this feature might be removed without warning, not for production use
-	static UE_API UE::PoseSearch::FSearchResult MotionMatch(
+	static UE_INTERNAL UE_API UE::PoseSearch::FSearchResult MotionMatch(
 		UE::PoseSearch::FSearchContext& SearchContext,
 		const TArrayView<const UObject*> AssetsToSearch,
 		const FPoseSearchContinuingProperties& ContinuingProperties);
diff --git a/Engine/Plugins/Experimental/UAF/UAFPoseSearch/Source/UAFPoseSearch/Internal/MotionMatchingTraitData.h b/Engine/Plugins/Experimental/UAF/UAFPoseSearch/Source/UAFPoseSearch/Internal/MotionMatchingTraitData.h
--- a/Engine/Plugins/Experimental/UAF/UAFPoseSearch/Source/UAFPoseSearch/Internal/MotionMatchingTraitData.h
+++ b/Engine/Plugins/Experimental/UAF/UAFPoseSearch/Source/UAFPoseSearch/Internal/MotionMatchingTraitData.h
@@ -1,4 +1,4 @@
-﻿// Copyright Epic Games, Inc. All Rights Reserved.
+// Copyright Epic Games, Inc. All Rights Reserved.
 
 #pragma once
 
@@ -57,6 +57,9 @@ struct FMotionMatchingTraitSharedData : public FAnimNextTraitSharedData
 	UPROPERTY(EditAnywhere, Category = Settings, meta = (PinHiddenByDefault))
 	bool bShouldUseCachedChannelData = false;
 
+	UPROPERTY(EditAnywhere, Category = Settings, meta = (PinHiddenByDefault))
+	FPoseSearchEvent EventToSearch;
+
 	// If set to true, blend space inputs are continuously set based on input pins.
 	UPROPERTY(EditAnywhere, Category = Settings)
 	bool bUpdateBlendSpaceInputs = false;
@@ -132,6 +135,7 @@ struct FMotionMatchingTraitSharedData : public FAnimNextTraitSharedData
 		GeneratorMacro(SearchThrottleTime) \
 		GeneratorMacro(PlayRate) \
 		GeneratorMacro(bShouldUseCachedChannelData) \
+		GeneratorMacro(EventToSearch) \
 		GeneratorMacro(bUpdateBlendSpaceInputs) \
 		GeneratorMacro(bTrySkipBlendsForBlendSpaces) \
 		GeneratorMacro(MaxDeltaAssetTimeToTrySkipBlendsForBlendSpaces) \
diff --git a/Engine/Plugins/Experimental/UAF/UAFPoseSearch/Source/UAFPoseSearch/Private/MotionMatchingTrait.cpp b/Engine/Plugins/Experimental/UAF/UAFPoseSearch/Source/UAFPoseSearch/Private/MotionMatchingTrait.cpp
--- a/Engine/Plugins/Experimental/UAF/UAFPoseSearch/Source/UAFPoseSearch/Private/MotionMatchingTrait.cpp
+++ b/Engine/Plugins/Experimental/UAF/UAFPoseSearch/Source/UAFPoseSearch/Private/MotionMatchingTrait.cpp
@@ -7,24 +7,25 @@
 #include "AnimNextAnimGraphSettings.h"
 #include "Component/AnimNextComponent.h"
 #include "EvaluationVM/EvaluationVM.h"
+#include "Factory/AnimGraphFactory.h"
 #include "Graph/AnimNextGraphInstance.h"
-#include "Traits/SequencePlayerTraitData.h"
-#include "Traits/BlendSpacePlayerTraitData.h"
-#include "Traits/SynchronizeUsingGroupsTraitData.h"
 #include "IPoseHistory.h"
-#include "Factory/AnimGraphFactory.h"
 #include "Module/AnimNextModuleInstance.h"
 #include "Module/RigUnit_AnimNextModuleEvents.h"
 #include "PoseSearch/MultiAnimAsset.h"
+#include "PoseSearch/PoseSearchContext.h"
 #include "PoseSearch/PoseSearchDatabase.h"
 #include "PoseSearch/PoseSearchDerivedData.h"
 #include "PoseSearch/PoseSearchInteractionUtils.h"
 #include "PoseSearch/PoseSearchSchema.h"
+#include "TraitCore/NodeInstance.h"
 #include "TraitInterfaces/IBlendStack.h"
 #include "TraitInterfaces/IGraphFactory.h"
 #include "TraitInterfaces/ITimeline.h"
+#include "Traits/BlendSpacePlayerTraitData.h"
+#include "Traits/SequencePlayerTraitData.h"
+#include "Traits/SynchronizeUsingGroupsTraitData.h"
 #include "VisualLogger/VisualLogger.h"
-#include "TraitCore/NodeInstance.h"
 
 #if ENABLE_ANIM_DEBUG
 namespace UE::Private
@@ -270,28 +271,67 @@ namespace UE::UAF
 				PoseSearchPlayRate = bPoseSearchPlayRateEnabled ? PoseSearchPlayRate : FFloatInterval(1.f, 1.f);
 #endif // ENABLE_ANIM_DEBUG
 
-				FChooserEvaluationContext EvaluationContext(const_cast<UObject*>(AnimContext));
+				const float SearchThrottleTime = SharedData->GetbShouldSearch(Binding) ? SharedData->GetSearchThrottleTime(Binding) : UE_BIG_NUMBER;
+
+				FPoseSearchEvent EventToSearch = SharedData->GetEventToSearch(Binding);
+				FSearchContext SearchContext(0.f, &MotionMatchingState.PoseIndicesHistory, SharedData->GetPoseJumpThresholdTime(Binding), EventToSearch.GetPlayRateOverriddenEvent(PoseSearchPlayRate));
+
+				FChooserEvaluationContext ChooserEvaluationContext(const_cast<UObject*>(AnimContext));
 				const FUAFAssetInstance& Instance = Binding.GetTraitPtr().GetNodeInstance()->GetOwner();
-                EvaluationContext.AddStructViewParam(	FStructView::Make(const_cast<FUAFAssetInstance&>(Instance)));
-                
-				UPoseSearchLibrary::UpdateMotionMatchingState(
-					&EvaluationContext,
-					PoseHistory,
-					DatabasesToSearch,
-					DeltaTime,
-					SharedData->GetPoseJumpThresholdTime(Binding),
-					SharedData->GetPoseReselectHistory(Binding),
-					SharedData->GetbShouldSearch(Binding) ? SharedData->GetSearchThrottleTime(Binding) : UE_BIG_NUMBER,
-					PoseSearchPlayRate,
-					MotionMatchingState,
-					SharedData->GetNextUpdateInterruptMode(Binding),
-					SharedData->GetbShouldUseCachedChannelData(Binding)
-					// @todo: add debug drawing suport
-	//#if ENABLE_ANIM_DEBUG
-	//				, CVarAnimNodeMotionMatchingDrawQuery.GetValueOnAnyThread()
-	//				, CVarAnimNodeMotionMatchingDrawCurResult.GetValueOnAnyThread() 
-	//#endif // ENABLE_ANIM_DEBUG
-				);
+                ChooserEvaluationContext.AddStructViewParam(FStructView::Make(const_cast<FUAFAssetInstance&>(Instance)));
+				SearchContext.AddRole(GetCommonDefaultRole(DatabasesToSearch), &ChooserEvaluationContext, PoseHistory);
+
+				// If we can't advance or enough time has elapsed since the last pose jump then search
+				const bool bCanAdvance = SearchContext.GetContinuingPoseSearchResult().PoseIdx != INDEX_NONE;
+				const bool bSearch = !bCanAdvance || (MotionMatchingState.ElapsedPoseSearchTime >= SearchThrottleTime);
+				FSearchResult SearchResult;
+				if (bSearch)
+				{
+					MotionMatchingState.ElapsedPoseSearchTime = 0.f;
+			
+					TArrayView<const UObject*> AssetsToSearch((const UObject**)FMemory_Alloca(DatabasesToSearch.Num() * sizeof(UObject*)), (int32)DatabasesToSearch.Num());
+					for (int32 DatabaseToSearchIndex = 0; DatabaseToSearchIndex < DatabasesToSearch.Num(); ++DatabaseToSearchIndex)
+					{
+						AssetsToSearch[DatabaseToSearchIndex] = DatabasesToSearch[DatabaseToSearchIndex].Get();
+					}
+
+					FPoseSearchContinuingProperties ContinuingProperties;
+					ContinuingProperties.InitFrom(MotionMatchingState.SearchResult, SharedData->GetNextUpdateInterruptMode(Binding));
+					SearchResult = UPoseSearchLibrary::MotionMatch(SearchContext, AssetsToSearch, ContinuingProperties);
+				}
+				else
+				{
+					MotionMatchingState.ElapsedPoseSearchTime += DeltaTime;
+
+					// if we didn't perform a search we carry on with the continuing pose search result
+					SearchResult = SearchContext.GetContinuingPoseSearchResult();
+					SearchResult.bIsContinuingPoseSearch = true;
+
+			#if UE_POSE_SEARCH_TRACE_ENABLED
+					// in case we skipped the search, we still have to track we would have requested to evaluate DatabasesToSearch and SearchResult.Database
+					for (const TObjectPtr<const UPoseSearchDatabase>& DatabaseToSearch : DatabasesToSearch)
+					{
+						SearchContext.Track(DatabaseToSearch.Get());
+					}
+
+					SearchContext.Track(SearchResult.Database.Get());
+			#endif // UE_POSE_SEARCH_TRACE_ENABLED
+				}
+
+				const float WantedPlayRate = CalculateWantedPlayRate(SearchResult, SearchContext, PoseSearchPlayRate, PoseHistory ? PoseHistory->GetTrajectorySpeedMultiplier() : 1.f, EventToSearch);
+				MotionMatchingState.SearchResult.InitFrom(SearchResult, WantedPlayRate);
+				MotionMatchingState.PoseIndicesHistory.Update(SearchResult, DeltaTime, SharedData->GetPoseReselectHistory(Binding));
+
+				PRAGMA_DISABLE_DEPRECATION_WARNINGS
+				MotionMatchingState.CurrentSearchResult = SearchResult;
+				MotionMatchingState.WantedPlayRate = WantedPlayRate;
+				PRAGMA_ENABLE_DEPRECATION_WARNINGS
+	
+				const bool bJumpToPose = MotionMatchingState.SearchResult.SelectedAnim && !MotionMatchingState.SearchResult.bIsContinuingPoseSearch;
+
+				PRAGMA_DISABLE_DEPRECATION_WARNINGS
+				MotionMatchingState.bJumpedToPose = bJumpToPose;
+				PRAGMA_ENABLE_DEPRECATION_WARNINGS
 			}
 		}
 
