PoseSearch - backing out the usage of UE_INTERNAL in PoseSearchLibrary APIs introduced in CL 44177024, creating issues with Linux server build

diff --git a/Engine/Plugins/Animation/PoseSearch/Source/Runtime/Public/PoseSearch/PoseSearchLibrary.h b/Engine/Plugins/Animation/PoseSearch/Source/Runtime/Public/PoseSearch/PoseSearchLibrary.h
--- a/Engine/Plugins/Animation/PoseSearch/Source/Runtime/Public/PoseSearch/PoseSearchLibrary.h
+++ b/Engine/Plugins/Animation/PoseSearch/Source/Runtime/Public/PoseSearch/PoseSearchLibrary.h
@@ -24,8 +24,10 @@ namespace UE::PoseSearch
 {
 	struct FSearchContext;
 
-	UE_INTERNAL UE_API FRole GetCommonDefaultRole(const TConstArrayView<TObjectPtr<const UPoseSearchDatabase>> Databases);
-	UE_INTERNAL UE_API float CalculateWantedPlayRate(const FSearchResult& SearchResult, const FSearchContext& SearchContext, const FFloatInterval& PlayRate, float TrajectorySpeedMultiplier, const FPoseSearchEvent& EventToSearch);
+	// Experimental, this feature might be removed without warning, not for production use
+	UE_API FRole GetCommonDefaultRole(const TConstArrayView<TObjectPtr<const UPoseSearchDatabase>> Databases);
+	// Experimental, this feature might be removed without warning, not for production use
+	UE_API float CalculateWantedPlayRate(const FSearchResult& SearchResult, const FSearchContext& SearchContext, const FFloatInterval& PlayRate, float TrajectorySpeedMultiplier, const FPoseSearchEvent& EventToSearch);
 } // namespace UE::PoseSearch
 
 
@@ -148,7 +150,8 @@ struct FPoseSearchContinuingProperties
 	UPROPERTY(Transient, EditAnywhere, BlueprintReadWrite, Category=State)
 	TObjectPtr<const UPoseSearchDatabase> PlayingAssetDatabase = nullptr;
 
-	UE_INTERNAL UE_API void InitFrom(const FPoseSearchBlueprintResult& SearchResult, EPoseSearchInterruptMode InInterruptMode);
+	// Experimental, this feature might be removed without warning, not for production use
+	UE_API void InitFrom(const FPoseSearchBlueprintResult& SearchResult, EPoseSearchInterruptMode InInterruptMode);
 };
 
 UCLASS(MinimalAPI)
@@ -301,7 +304,8 @@ public:
 		const FPoseSearchFutureProperties& Future,
 		const FPoseSearchEvent& EventToSearch);
 	
-	static UE_INTERNAL UE_API UE::PoseSearch::FSearchResult MotionMatch(
+	// Experimental, this feature might be removed without warning, not for production use
+	static UE_API UE::PoseSearch::FSearchResult MotionMatch(
 		const TArrayView<FChooserEvaluationContext> Contexts,
 		const TArrayView<const UE::PoseSearch::FRole> Roles,
 		const TArrayView<const UE::PoseSearch::IPoseHistory*> PoseHistories, 
@@ -310,7 +314,8 @@ public:
 		const FPoseSearchFutureProperties& Future,
 		const FPoseSearchEvent& EventToSearch);
 			
-	static UE_INTERNAL UE_API UE::PoseSearch::FSearchResult MotionMatch(
+	// Experimental, this feature might be removed without warning, not for production use
+	static UE_API UE::PoseSearch::FSearchResult MotionMatch(
 		const TArrayView<FChooserEvaluationContext> Contexts,
 		const TArrayView<const UE::PoseSearch::FRole> Roles,
 		const TArrayView<const UE::PoseSearch::IPoseHistory*> PoseHistories,
@@ -319,7 +324,8 @@ public:
 		const float DesiredPermutationTimeOffset,
 		const FPoseSearchEvent& EventToSearch);
 
-	static UE_INTERNAL UE_API UE::PoseSearch::FSearchResult MotionMatch(
+	// Experimental, this feature might be removed without warning, not for production use
+	static UE_API UE::PoseSearch::FSearchResult MotionMatch(
 		const TArrayView<const UObject*> AnimContexts,
 		const TArrayView<const UE::PoseSearch::FRole> Roles,
 		const TArrayView<const UE::PoseSearch::IPoseHistory*> PoseHistories,
@@ -328,7 +334,8 @@ public:
 		const float DesiredPermutationTimeOffset,
 		const FPoseSearchEvent& EventToSearch);
 
-	static UE_INTERNAL UE_API UE::PoseSearch::FSearchResult MotionMatch(
+	// Experimental, this feature might be removed without warning, not for production use
+	static UE_API UE::PoseSearch::FSearchResult MotionMatch(
 		UE::PoseSearch::FSearchContext& SearchContext,
 		const TArrayView<const UObject*> AssetsToSearch,
 		const FPoseSearchContinuingProperties& ContinuingProperties);
