PoseSearch - FPoseSearchColumn / MM chooser integration WIP

diff --git a/Engine/Plugins/Animation/PoseSearch/Source/Editor/Private/Chooser/PoseSearchColumnEditor.cpp b/Engine/Plugins/Animation/PoseSearch/Source/Editor/Private/Chooser/PoseSearchColumnEditor.cpp
--- a/Engine/Plugins/Animation/PoseSearch/Source/Editor/Private/Chooser/PoseSearchColumnEditor.cpp
+++ b/Engine/Plugins/Animation/PoseSearch/Source/Editor/Private/Chooser/PoseSearchColumnEditor.cpp
@@ -40,9 +40,9 @@ TSharedRef<SWidget> CreatePoseSearchColumnWidget(UChooserTable* Chooser, FChoose
 		SNew(STextBlock)
 			.Text_Lambda([Row, PoseSearchColumn]()
 			{
-				if (PoseSearchColumn->RowValues.IsValidIndex(Row))
+				if (FPoseSearchDatabaseSequence* Asset = PoseSearchColumn->InternalDatabase->GetMutableDatabaseAnimationAsset<FPoseSearchDatabaseSequence>(Row))
 				{
-					if (UObject* ResultAsset = PoseSearchColumn->RowValues[Row].ResultAsset)
+					if (UObject* ResultAsset = Asset->GetAnimationAsset())
 					{
 						return FText::FromString(ResultAsset->GetName());
 					}
diff --git a/Engine/Plugins/Animation/PoseSearch/Source/Runtime/Private/Chooser/PoseSearchChooserColumn.cpp b/Engine/Plugins/Animation/PoseSearch/Source/Runtime/Private/Chooser/PoseSearchChooserColumn.cpp
--- a/Engine/Plugins/Animation/PoseSearch/Source/Runtime/Private/Chooser/PoseSearchChooserColumn.cpp
+++ b/Engine/Plugins/Animation/PoseSearch/Source/Runtime/Private/Chooser/PoseSearchChooserColumn.cpp
@@ -1,18 +1,26 @@
 // Copyright Epic Games, Inc. All Rights Reserved.
 #include "PoseSearch/Chooser/PoseSearchChooserColumn.h"
+#include "Animation/AnimNodeBase.h"
+#include "Animation/AnimSequence.h"
+#include "Chooser.h"
 #include "ChooserIndexArray.h"
 #include "ChooserPropertyAccess.h"
+#include "ChooserTrace.h"
+#include "ChooserTypes.h"
 #include "IChooserParameterBool.h"
 #include "IChooserParameterFloat.h"
-#include "ChooserTrace.h"
-#include "PoseSearch/PoseSearchLibrary.h"
+#include "PoseSearch/PoseSearchContext.h"
 #include "PoseSearch/PoseSearchDatabase.h"
+#include "PoseSearch/PoseSearchLibrary.h"
+#include "PoseSearch/PoseSearchSchema.h"
 #if WITH_EDITOR
+#include "ScopedTransaction.h"
 #include "StructUtils/PropertyBag.h"
 #endif
 
 #include UE_INLINE_GENERATED_CPP_BY_NAME(PoseSearchChooserColumn)
 
+#define LOCTEXT_NAMESPACE "FPoseSearchColumn"
 
 bool FPoseHistoryContextProperty::GetValue(FChooserEvaluationContext& Context, FPoseHistoryReference& OutResult) const
 {
@@ -33,10 +41,123 @@ FPoseSearchColumn::FPoseSearchColumn()
 #if WITH_EDITOR
 void FPoseSearchColumn::AutoPopulate(int32 RowIndex, UObject* OutputObject)
 {
-	if (RowValues.IsValidIndex(RowIndex))
+	if (FPoseSearchDatabaseSequence* Asset = InternalDatabase->GetMutableDatabaseAnimationAsset<FPoseSearchDatabaseSequence>(RowIndex))
+	{
+		UAnimSequence* AnimSequence = Cast<UAnimSequence>(OutputObject);
+		if (Asset->Sequence != AnimSequence)
+		{
+			FScopedTransaction ScopedTransaction(LOCTEXT("AutoPopulate","AutoPopulate FPoseSearchColumn"));
+			InternalDatabase->Modify();
+			Asset->Sequence = AnimSequence;
+		}
+	}
+}
+
+FName FPoseSearchColumn::RowValuesPropertyName()
+{
+	return "RowValues";
+}
+
+void FPoseSearchColumn::SetNumRows(int32 NumRows)
+{
+	if (NumRows == InternalDatabase->GetNumAnimationAssets())
 	{
-		RowValues[RowIndex].ResultAsset = Cast<UAnimationAsset>(OutputObject);
+		// nothing to do
 	}
+	else if (NumRows < InternalDatabase->GetNumAnimationAssets())
+	{
+		FScopedTransaction ScopedTransaction(LOCTEXT("SetNumRows","SetNumRows FPoseSearchColumn"));
+		InternalDatabase->Modify();
+		while (InternalDatabase->GetNumAnimationAssets() > 0 && InternalDatabase->GetNumAnimationAssets() > NumRows)
+		{
+			InternalDatabase->RemoveAnimationAssetAt(InternalDatabase->GetNumAnimationAssets() - 1);
+		}
+	}
+	else
+	{
+		FScopedTransaction ScopedTransaction(LOCTEXT("SetNumRows","SetNumRows FPoseSearchColumn"));
+		InternalDatabase->Modify();
+		while (InternalDatabase->GetNumAnimationAssets() < NumRows)
+		{
+			FPoseSearchDatabaseSequence NewAsset;
+			NewAsset.Sequence = nullptr;
+			InternalDatabase->AddAnimationAsset(FInstancedStruct::Make(NewAsset));
+		}
+	}
+}
+
+void FPoseSearchColumn::InsertRows(int Index, int Count)
+{
+	if (Count > 0)
+	{
+		FScopedTransaction ScopedTransaction(LOCTEXT("InsertRows","InsertRows FPoseSearchColumn"));
+		InternalDatabase->Modify();
+		for (int i = 0; i < Count; i++)
+		{
+			FPoseSearchDatabaseSequence NewAsset;
+			NewAsset.Sequence = nullptr;
+			InternalDatabase->InsertAnimationAssetAt(FInstancedStruct::Make(NewAsset), Index);
+		}
+	}
+}
+
+void FPoseSearchColumn::DeleteRows(const TArray<uint32>& RowIndices)
+{
+	FScopedTransaction ScopedTransaction(LOCTEXT("DeleteRows","DeleteRows FPoseSearchColumn"));
+	InternalDatabase->Modify();
+	for(uint32 Index : RowIndices)
+	{
+		InternalDatabase->RemoveAnimationAssetAt(Index);
+	}
+}
+
+void FPoseSearchColumn::MoveRow(int SourceRowIndex, int TargetRowIndex)
+{
+	if (const FPoseSearchDatabaseSequence* SourceAsset = InternalDatabase->GetDatabaseAnimationAsset<FPoseSearchDatabaseSequence>(SourceRowIndex))
+	{
+		FScopedTransaction ScopedTransaction(LOCTEXT("MoveRow","MoveRow FPoseSearchColumn"));
+		InternalDatabase->Modify();
+
+		FPoseSearchDatabaseSequence CopySourceAsset = *SourceAsset;
+		InternalDatabase->RemoveAnimationAssetAt(SourceRowIndex);
+
+		if (SourceRowIndex < TargetRowIndex)
+		{
+			TargetRowIndex--;
+		}
+		InternalDatabase->InsertAnimationAssetAt(FInstancedStruct::Make(CopySourceAsset), TargetRowIndex);
+	}
+}
+
+void FPoseSearchColumn::CopyRow(FChooserColumnBase& SourceColumn, int SourceRowIndex, int TargetRowIndex)
+{
+	if (SourceRowIndex != TargetRowIndex)
+	{
+		if (const FPoseSearchDatabaseSequence* SourceAsset = InternalDatabase->GetDatabaseAnimationAsset<FPoseSearchDatabaseSequence>(SourceRowIndex))
+		{
+			if (FPoseSearchDatabaseSequence* TargetAsset = InternalDatabase->GetMutableDatabaseAnimationAsset<FPoseSearchDatabaseSequence>(TargetRowIndex))
+			{
+				FScopedTransaction ScopedTransaction(LOCTEXT("CopyRow","CopyRow FPoseSearchColumn"));
+				InternalDatabase->Modify();
+				*TargetAsset = *SourceAsset;
+			}
+		}
+	}
+}
+
+UScriptStruct* FPoseSearchColumn::GetInputBaseType() const
+{
+	return FChooserParameterPoseHistoryBase::StaticStruct();
+}
+
+const UScriptStruct* FPoseSearchColumn::GetInputType() const
+{
+	return InputValue.IsValid() ? InputValue.GetScriptStruct() : nullptr;
+}
+
+void FPoseSearchColumn::SetInputType(const UScriptStruct* Type)
+{
+	InputValue.InitializeAs(Type);
 }
 
 bool FPoseSearchColumn::EditorTestFilter(int32 RowIndex) const
@@ -45,14 +166,18 @@ bool FPoseSearchColumn::EditorTestFilter(int32 RowIndex) const
 	return true;
 }
 
-
 float FPoseSearchColumn::EditorTestCost(int32 RowIndex) const
 {
 	// todo: not sure how to do editor display for pose search weights (might need to make a copy of the pose history struct)
 	return 0.0f;
 }
 
-#endif
+#endif // WITH_EDITOR
+
+FChooserParameterBase* FPoseSearchColumn::GetInputValue()
+{
+	return InputValue.GetMutablePtr<FChooserParameterBase>();
+}
 
 bool FPoseSearchColumn::HasFilters() const
 {
@@ -61,61 +186,91 @@ bool FPoseSearchColumn::HasFilters() const
 
 void FPoseSearchColumn::Filter(FChooserEvaluationContext& Context, const FChooserIndexArray& IndexListIn, FChooserIndexArray& IndexListOut) const
 {
+	using namespace UE::PoseSearch;
+
+	const UE::PoseSearch::IPoseHistory* PoseHistory = nullptr;
+	FPoseHistoryReference PoseHistoryReference;
+	if (const FChooserParameterPoseHistoryBase* PoseHistoryParameter = InputValue.GetPtr<FChooserParameterPoseHistoryBase>())
+	{
+		if (PoseHistoryParameter->GetValue(Context, PoseHistoryReference))
+		{
+			PoseHistory = PoseHistoryReference.PoseHistory.Get();
+		}
+	}
 
-	TArray<const UObject*, TInlineAllocator<128>> AssetsToSearch;
-	for (const FChooserIndexArray::FIndexData& IndexData : IndexListIn)
+	// searching for FChooserPlayerSettings in the Context.Params
+	const FChooserPlayerSettings* ChooserPlayerSettings = nullptr;
+	for (const FStructView& Param : Context.Params)
 	{
-		AssetsToSearch.Add(RowValues[IndexData.Index].ResultAsset);
+		ChooserPlayerSettings = Param.GetPtr<FChooserPlayerSettings>();
+		if (ChooserPlayerSettings)
+		{
+			break;
+		}
 	}
 
+	if (!PoseHistory && ChooserPlayerSettings && ChooserPlayerSettings->AnimationUpdateContext)
+	{
+		if (const FPoseHistoryProvider* PoseHistoryProvider = ChooserPlayerSettings->AnimationUpdateContext->GetMessage<FPoseHistoryProvider>())
+		{
+			PoseHistory = &PoseHistoryProvider->GetPoseHistory();
+		}
+	}
 
-	FPoseHistoryReference PoseHistoryReference;
-	if (const FChooserParameterPoseHistoryBase* PoseHistoryParameter = InputValue.GetPtr<FChooserParameterPoseHistoryBase>())
+	if (PoseHistory)
 	{
-		if (PoseHistoryParameter->GetValue(Context, PoseHistoryReference))
+		TArray<const UObject*, TInlineAllocator<128>> AssetsToConsider;
+		for (const FChooserIndexArray::FIndexData& IndexData : IndexListIn)
 		{
-			if (const UE::PoseSearch::IPoseHistory* PoseHistory = PoseHistoryReference.PoseHistory.Get())
+			if (const FPoseSearchDatabaseSequence* Asset = InternalDatabase->GetDatabaseAnimationAsset<FPoseSearchDatabaseSequence>(IndexData.Index))
 			{
-				FPoseSearchContinuingProperties ContinuingProperties;
-				FPoseSearchFutureProperties Future;
-
-				const UE::PoseSearch::FSearchResult SearchResult = UPoseSearchLibrary::MotionMatch(
-					MakeArrayView(&Context, 1),
-					MakeArrayView(&UE::PoseSearch::DefaultRole, 1),
-					MakeArrayView(&PoseHistory, 1),
-					AssetsToSearch,
-					ContinuingProperties,
-					Future,
-					FPoseSearchEvent());
-
-				 if (const UE::PoseSearch::FSearchIndexAsset* SearchIndexAsset = SearchResult.GetSearchIndexAsset())
-				 {
-				 	const UPoseSearchDatabase* Database = SearchResult.Database.Get();
-				 	if (const FPoseSearchDatabaseAnimationAssetBase* DatabaseAnimationAssetBase = Database->GetDatabaseAnimationAsset<FPoseSearchDatabaseAnimationAssetBase>(*SearchIndexAsset))
+				AssetsToConsider.Add(Asset->GetAnimationAsset());
+			}
+		}
+
+		const FFloatInterval PoseJumpThresholdTime(0.f, 0.f);
+		const UObject* DatabasesToSearch[] = { InternalDatabase.Get() };
+		FSearchContext SearchContext(0.f, nullptr /*&MotionMatchingState.PoseIndicesHistory*/, PoseJumpThresholdTime, FPoseSearchEvent());
+
+		const FRole Role = InternalDatabase->Schema ? InternalDatabase->Schema->GetDefaultRole() : DefaultRole;
+		SearchContext.AddRole(Role, &Context, PoseHistory);
+		SearchContext.SetAssetsToConsider(AssetsToConsider);
+
+		FPoseSearchContinuingProperties ContinuingProperties;
+		if (ChooserPlayerSettings)
+		{
+			ContinuingProperties.PlayingAsset = ChooserPlayerSettings->PlayingAsset;
+			ContinuingProperties.PlayingAssetAccumulatedTime = ChooserPlayerSettings->PlayingAssetAccumulatedTime;
+			ContinuingProperties.bIsPlayingAssetMirrored = ChooserPlayerSettings->bIsPlayingAssetMirrored;
+			ContinuingProperties.PlayingAssetBlendParameters = ChooserPlayerSettings->PlayingAssetBlendParameters;
+		}
+		const UE::PoseSearch::FSearchResult SearchResult = UPoseSearchLibrary::MotionMatch(SearchContext, DatabasesToSearch, ContinuingProperties);
+
+		if (const UE::PoseSearch::FSearchIndexAsset* SearchIndexAsset = SearchResult.GetSearchIndexAsset())
+		{
+			if (const FPoseSearchDatabaseAnimationAssetBase* DatabaseAnimationAssetBase = SearchResult.Database->GetDatabaseAnimationAsset<FPoseSearchDatabaseAnimationAssetBase>(*SearchIndexAsset))
+			{
+				const UObject* Result = DatabaseAnimationAssetBase->GetAnimationAsset();
+				int Index;
+				if (AssetsToConsider.Find(Result, Index))
+				{
+				 	IndexListOut.Push(IndexListIn[Index]);
+
+				 	if (const FChooserParameterFloatBase* StartTime = OutputStartTime.GetPtr<FChooserParameterFloatBase>())
 				 	{
-				 		const UObject* Result = DatabaseAnimationAssetBase->GetAnimationAsset();
-				 		int Index;
-				 		if (AssetsToSearch.Find(Result, Index))
-				 		{
-				 			IndexListOut.Push(IndexListIn[Index]);
-
-				 			if (const FChooserParameterFloatBase* StartTime = OutputStartTime.GetPtr<FChooserParameterFloatBase>())
-				 			{
-				 				StartTime->SetValue(Context, SearchResult.AssetTime);
-				 			}
+				 		StartTime->SetValue(Context, SearchResult.AssetTime);
+				 	}
 				 			
-				 			if (const FChooserParameterFloatBase* Cost = OutputCost.GetPtr<FChooserParameterFloatBase>())
-                            {
-								Cost->SetValue(Context, SearchResult.PoseCost);
-                            }
-
-				 			if (const FChooserParameterBoolBase* Mirror = OutputMirror.GetPtr<FChooserParameterBoolBase>())
-				 			{
-				 				Mirror->SetValue(Context, SearchIndexAsset->IsMirrored());
-				 			}
-				 		}
+				 	if (const FChooserParameterFloatBase* Cost = OutputCost.GetPtr<FChooserParameterFloatBase>())
+                    {
+						Cost->SetValue(Context, SearchResult.PoseCost);
+                    }
+
+				 	if (const FChooserParameterBoolBase* Mirror = OutputMirror.GetPtr<FChooserParameterBoolBase>())
+				 	{
+				 		Mirror->SetValue(Context, SearchIndexAsset->IsMirrored());
 				 	}
-				 }
+				}
 			}
 		}
 	}
@@ -145,25 +300,40 @@ void FPoseSearchColumn::Filter(FChooserEvaluationContext& Context, const FChoose
 }
 
 #if WITH_EDITOR
-	void FPoseSearchColumn::AddToDetails(FInstancedPropertyBag& PropertyBag, int32 ColumnIndex, int32 RowIndex)
-	{
-		FText DisplayName = NSLOCTEXT("PoseSearchColumn","Pose Search", "Pose Search");
-		FName PropertyName("RowData", ColumnIndex);
-		FPropertyBagPropertyDesc PropertyDesc(PropertyName,  EPropertyBagPropertyType::Struct, FChooserPoseSearchRowData::StaticStruct());
-		PropertyDesc.MetaData.Add(FPropertyBagPropertyDescMetaData("DisplayName", DisplayName.ToString()));
-		PropertyBag.AddProperties({PropertyDesc});
+void FPoseSearchColumn::Initialize(UChooserTable* OuterChooser)
+{
+	FChooserColumnBase::Initialize(OuterChooser);
+
+	InternalDatabase = NewObject<UPoseSearchDatabase>(OuterChooser, OuterChooser->GetFName(), RF_Transactional);
+}
+
+void FPoseSearchColumn::AddToDetails(FInstancedPropertyBag& PropertyBag, int32 ColumnIndex, int32 RowIndex)
+{
+	FText DisplayName = NSLOCTEXT("PoseSearchColumn","Pose Search", "Pose Search");
+	FName PropertyName("RowData", ColumnIndex);
+	FPropertyBagPropertyDesc PropertyDesc(PropertyName,  EPropertyBagPropertyType::Struct, FPoseSearchDatabaseSequence::StaticStruct());
+	PropertyDesc.MetaData.Add(FPropertyBagPropertyDescMetaData("DisplayName", DisplayName.ToString()));
+	PropertyBag.AddProperties({PropertyDesc});
 		
-		PropertyBag.SetValueStruct(PropertyName, RowValues[RowIndex]);
+	if (FPoseSearchDatabaseSequence* Asset = InternalDatabase->GetMutableDatabaseAnimationAsset<FPoseSearchDatabaseSequence>(RowIndex))
+	{
+		PropertyBag.SetValueStruct(PropertyName, *Asset);
 	}
+}
 
-	void FPoseSearchColumn::SetFromDetails(FInstancedPropertyBag& PropertyBag, int32 ColumnIndex, int32 RowIndex)
-	{
-		FName PropertyName("RowData", ColumnIndex);
+void FPoseSearchColumn::SetFromDetails(FInstancedPropertyBag& PropertyBag, int32 ColumnIndex, int32 RowIndex)
+{
+	FName PropertyName("RowData", ColumnIndex);
 		
-   		TValueOrError<FStructView, EPropertyBagResult> Result = PropertyBag.GetValueStruct(PropertyName, FChooserPoseSearchRowData::StaticStruct());
-		if (FStructView* StructView = Result.TryGetValue())
+   	TValueOrError<FStructView, EPropertyBagResult> Result = PropertyBag.GetValueStruct(PropertyName, FPoseSearchDatabaseSequence::StaticStruct());
+	if (FStructView* StructView = Result.TryGetValue())
+	{
+		if (FPoseSearchDatabaseSequence* Asset = InternalDatabase->GetMutableDatabaseAnimationAsset<FPoseSearchDatabaseSequence>(RowIndex))
 		{
-			RowValues[RowIndex] = StructView->Get<FChooserPoseSearchRowData>();
+			*Asset = StructView->Get<FPoseSearchDatabaseSequence>();
 		}
 	}
+}
 #endif
+
+#undef LOCTEXT_NAMESPACE
diff --git a/Engine/Plugins/Animation/PoseSearch/Source/Runtime/Private/PoseSearchDatabase.cpp b/Engine/Plugins/Animation/PoseSearch/Source/Runtime/Private/PoseSearchDatabase.cpp
--- a/Engine/Plugins/Animation/PoseSearch/Source/Runtime/Private/PoseSearchDatabase.cpp
+++ b/Engine/Plugins/Animation/PoseSearch/Source/Runtime/Private/PoseSearchDatabase.cpp
@@ -921,6 +921,11 @@ void UPoseSearchDatabase::AddAnimationAsset(FInstancedStruct AnimationAsset)
 	AnimationAssets.Add(AnimationAsset);
 }
 
+void UPoseSearchDatabase::InsertAnimationAssetAt(FInstancedStruct AnimationAsset, int32 AnimationAssetIndex)
+{
+	AnimationAssets.Insert(AnimationAsset, AnimationAssetIndex);
+}
+
 void UPoseSearchDatabase::RemoveAnimationAssetAt(int32 AnimationAssetIndex)
 {
 	AnimationAssets.RemoveAt(AnimationAssetIndex);
diff --git a/Engine/Plugins/Animation/PoseSearch/Source/Runtime/Private/PoseSearchDerivedData.cpp b/Engine/Plugins/Animation/PoseSearch/Source/Runtime/Private/PoseSearchDerivedData.cpp
--- a/Engine/Plugins/Animation/PoseSearch/Source/Runtime/Private/PoseSearchDerivedData.cpp
+++ b/Engine/Plugins/Animation/PoseSearch/Source/Runtime/Private/PoseSearchDerivedData.cpp
@@ -2291,8 +2291,17 @@ FAsyncPoseSearchDatabasesManagement::~FAsyncPoseSearchDatabasesManagement()
 // given Object it figures out a map of databases to UAnimSequenceBase(s) containing UAnimNotifyState_PoseSearchBranchIn
 void FAsyncPoseSearchDatabasesManagement::CollectDatabasesToSynchronize(UObject* Object)
 {
-	FScopeLock Lock(&Mutex);
+	if (!Object)
+	{
+		return;
+	}
+	
+	if (!Object->IsA<UAnimSequenceBase>() && Object->IsA<UAnimNotifyState_PoseSearchBranchIn>())
+	{
+		return;
+	}
 
+	FScopeLock Lock(&Mutex);
 	if (UAnimSequenceBase* SequenceBase = Cast<UAnimSequenceBase>(Object))
 	{
 		for (const FAnimNotifyEvent& NotifyEvent : SequenceBase->Notifies)
@@ -2400,9 +2409,26 @@ void FAsyncPoseSearchDatabasesManagement::PostModified(UObject* Object)
 	ClearPreCancelled();
 }
 
-void FAsyncPoseSearchDatabasesManagement::OnObjectTransacted(UObject* Object, const FTransactionObjectEvent& TransactionObjectEvent)
+void FAsyncPoseSearchDatabasesManagement::OnObjectTransacted(UObject*, const FTransactionObjectEvent&)
 {
-	PostModified(Object);
+	check(IsInGameThread());
+
+	FScopeLock Lock(&Mutex);
+
+	// collecting databases to synchronize, and merging the results with the PreModified collection
+	for (UObject* ObjectModifiedThisFrame : FCoreUObjectDelegates::ObjectsModifiedThisFrame)
+	{
+		CollectDatabasesToSynchronize(ObjectModifiedThisFrame);
+	}
+
+	SynchronizeDatabases();
+
+	for (UObject* ObjectModifiedThisFrame : FCoreUObjectDelegates::ObjectsModifiedThisFrame)
+	{
+		SynchronizeDatabaseChooser(ObjectModifiedThisFrame);
+	}
+
+	ClearPreCancelled();
 }
 
 void FAsyncPoseSearchDatabasesManagement::OnPackageReloaded(const EPackageReloadPhase InPackageReloadPhase, FPackageReloadedEvent* InPackageReloadedEvent)
diff --git a/Engine/Plugins/Animation/PoseSearch/Source/Runtime/Public/PoseSearch/Chooser/PoseSearchChooserColumn.h b/Engine/Plugins/Animation/PoseSearch/Source/Runtime/Public/PoseSearch/Chooser/PoseSearchChooserColumn.h
--- a/Engine/Plugins/Animation/PoseSearch/Source/Runtime/Public/PoseSearch/Chooser/PoseSearchChooserColumn.h
+++ b/Engine/Plugins/Animation/PoseSearch/Source/Runtime/Public/PoseSearch/Chooser/PoseSearchChooserColumn.h
@@ -6,6 +6,7 @@
 #include "CoreMinimal.h"
 #include "IChooserColumn.h"
 #include "ChooserPropertyAccess.h"
+#include "PoseSearch/PoseSearchDatabase.h"
 #include "PoseSearch/Chooser/ChooserParameterPoseHistoryBase.h"
 #include "StructUtils/InstancedStruct.h"
 #include "Serialization/MemoryReader.h"
@@ -17,7 +18,7 @@
 
 class UAnimationAsset;
 
-USTRUCT(DisplayName = "Pose History Property Binding")
+USTRUCT(Experimental, DisplayName = "Pose History Property Binding")
 struct FPoseHistoryContextProperty :  public FChooserParameterPoseHistoryBase
 {
 	GENERATED_BODY()
@@ -35,22 +36,15 @@ public:
 	CHOOSER_PARAMETER_BOILERPLATE();
 };
 
-
-USTRUCT()
-struct FChooserPoseSearchRowData
-{
-	GENERATED_BODY()
-
-	UPROPERTY(VisibleAnywhere, Category = "Data", Meta = (Tooltip = "The result animation asset for this row (filled by autopopulate from the result column)"));
-	TObjectPtr<UAnimationAsset> ResultAsset;
-};
-
 USTRUCT(DisplayName = "Pose Match", Meta = (Experimental, Category = "Experimental", Tooltip = "This column filters out all assets except the one which is selected by motion matching query.  Results must be AnimationAssets with a PoseSearchBranchIn notify state.  It also outputs OutputStartTime to specify the frame which matched pose best.  To work as intended it must be placed last (furthest right) in the Chooser so that other filters are applied first."))
 struct FPoseSearchColumn : public FChooserColumnBase
 {
 	GENERATED_BODY()
-	public:
+	
 	UE_API FPoseSearchColumn();
+
+	UPROPERTY(VisibleAnywhere, Category = "Data")
+	TObjectPtr<UPoseSearchDatabase> InternalDatabase;
 	
 	// Pose History
 	UPROPERTY(EditAnywhere, DisplayName = "Pose History", NoClear, Meta = (ExcludeBaseStruct, BaseStruct = "/Script/PoseSearch.ChooserParameterPoseHistoryBase"), Category = "Data")
@@ -68,20 +62,12 @@ struct FPoseSearchColumn : public FChooserColumnBase
 	UPROPERTY(EditAnywhere, NoClear, Meta = (ExcludeBaseStruct, BaseStruct = "/Script/Chooser.ChooserParameterFloatBase"), Category = "Data")
 	FInstancedStruct OutputCost;
 
-#if WITH_EDITORONLY_DATA
-	UPROPERTY(EditAnywhere, Category=Data, meta=(ToolTip="DefaultRowValue will be assigned to cells when new rows are created"));
-	FChooserPoseSearchRowData DefaultRowValue;
-#endif
-	
-	UPROPERTY()
-	// array of results (cells for this column for each row in the table)
-	// should match the length of the Results array 
-	TArray<FChooserPoseSearchRowData> RowValues;
-
 	UE_API virtual bool HasFilters() const override;
 	UE_API virtual void Filter(FChooserEvaluationContext& Context, const FChooserIndexArray& IndexListIn, FChooserIndexArray& IndexListOut) const override;
 
 #if WITH_EDITOR
+	virtual void Initialize(UChooserTable* OuterChooser) override;
+	
 	mutable double TestValue = 0.0;
 	UE_API virtual bool EditorTestFilter(int32 RowIndex) const override;
 	UE_API virtual float EditorTestCost(int32 RowIndex) const override;
@@ -97,9 +83,18 @@ struct FPoseSearchColumn : public FChooserColumnBase
 
 	virtual bool AutoPopulates() const override { return true; }
 	UE_API virtual void AutoPopulate(int32 RowIndex, UObject* OutputObject) override;
-#endif
+	virtual FName RowValuesPropertyName() override;
+	virtual void SetNumRows(int32 NumRows) override;
+	virtual void InsertRows(int Index, int Count) override;
+	virtual void DeleteRows(const TArray<uint32>& RowIndices) override;
+	virtual void MoveRow(int SourceRowIndex, int TargetRowIndex) override;
+	virtual void CopyRow(FChooserColumnBase& SourceColumn, int SourceRowIndex, int TargetRowIndex) override;
+	virtual UScriptStruct* GetInputBaseType() const override;
+	virtual const UScriptStruct* GetInputType() const override;
+	virtual void SetInputType(const UScriptStruct* Type) override;
 	
-	CHOOSER_COLUMN_BOILERPLATE(FChooserParameterPoseHistoryBase);
+#endif // WITH_EDITOR
+	virtual FChooserParameterBase* GetInputValue() override;
 };
 
 #undef UE_API
diff --git a/Engine/Plugins/Animation/PoseSearch/Source/Runtime/Public/PoseSearch/PoseSearchDatabase.h b/Engine/Plugins/Animation/PoseSearch/Source/Runtime/Public/PoseSearch/PoseSearchDatabase.h
--- a/Engine/Plugins/Animation/PoseSearch/Source/Runtime/Public/PoseSearch/PoseSearchDatabase.h
+++ b/Engine/Plugins/Animation/PoseSearch/Source/Runtime/Public/PoseSearch/PoseSearchDatabase.h
@@ -560,6 +560,8 @@ public:
 	UE_API int32 GetPoseIndex(const UObject* AnimationAsset, float AnimationAssetTime, bool bMirrored, const FVector& BlendParameters) const;
 
 	UE_API void AddAnimationAsset(FInstancedStruct AnimationAsset);
+	// Experimental, this feature might be removed without warning, not for production use
+	UE_API void InsertAnimationAssetAt(FInstancedStruct AnimationAsset, int32 AnimationAssetIndex);
 	UE_API void RemoveAnimationAssetAt(int32 AnimationAssetIndex);
 
 	template<typename TDatabaseAnimationAsset> const TDatabaseAnimationAsset* GetDatabaseAnimationAsset(int32 AnimationAssetIndex) const;
diff --git a/Engine/Plugins/Chooser/Source/Chooser/Internal/AnimNode_ChooserPlayer.h b/Engine/Plugins/Chooser/Source/Chooser/Internal/AnimNode_ChooserPlayer.h
--- a/Engine/Plugins/Chooser/Source/Chooser/Internal/AnimNode_ChooserPlayer.h
+++ b/Engine/Plugins/Chooser/Source/Chooser/Internal/AnimNode_ChooserPlayer.h
@@ -89,7 +89,6 @@ public:
 private:
 	UAnimationAsset* ChooseAsset(const FAnimationUpdateContext& Context);
 	FChooserPlayerSettings Settings;
-	FChooserEvaluationContext ChooserContext;
 	UAnimationAsset* CurrentAsset = nullptr;
 	float CurrentStartTime = 0;
 	bool CurrentMirror = false;
diff --git a/Engine/Plugins/Chooser/Source/Chooser/Private/AnimNode_ChooserPlayer.cpp b/Engine/Plugins/Chooser/Source/Chooser/Private/AnimNode_ChooserPlayer.cpp
--- a/Engine/Plugins/Chooser/Source/Chooser/Private/AnimNode_ChooserPlayer.cpp
+++ b/Engine/Plugins/Chooser/Source/Chooser/Private/AnimNode_ChooserPlayer.cpp
@@ -31,6 +31,16 @@ UAnimationAsset* FAnimNode_ChooserPlayer::ChooseAsset(const FAnimationUpdateCont
 		// reset settings to default
 		Settings = DefaultSettings;
 		
+		Settings.PlayingAsset = GetAnimAsset();
+		Settings.PlayingAssetAccumulatedTime = GetAccumulatedTime();
+		Settings.bIsPlayingAssetMirrored = GetMirror();
+		Settings.PlayingAssetBlendParameters = GetBlendParameters();
+		Settings.AnimationUpdateContext = &Context;
+
+		FChooserEvaluationContext ChooserContext;
+		ChooserContext.AddObjectParam(Context.GetAnimInstanceObject());
+		ChooserContext.AddStructParam(Settings);
+
 		const FObjectChooserBase& ChooserBase = Chooser.Get<FObjectChooserBase>();
 		if (bStartFromMatchingPose)
 		{
@@ -43,7 +53,6 @@ UAnimationAsset* FAnimNode_ChooserPlayer::ChooseAsset(const FAnimationUpdateCont
 						return FObjectChooserBase::EIteratorStatus::Continue;
 					}));
 
-
 				IPoseSearchProvider::FSearchPlayingAsset PlayingAsset;
 				PlayingAsset.Asset = GetAnimAsset();
 				PlayingAsset.AccumulatedTime = GetAccumulatedTime();
@@ -76,9 +85,6 @@ void FAnimNode_ChooserPlayer::Initialize_AnyThread(const FAnimationInitializeCon
 	
 	if(!bInitialized)
 	{
-		ChooserContext.AddObjectParam(Context.GetAnimInstanceObject());
-		ChooserContext.AddStructParam(Settings);
-
 		if (FObjectChooserBase* ObjectChooser = Chooser.GetMutablePtr<FObjectChooserBase>())
 		{
 			ObjectChooser->Compile(this,false);
diff --git a/Engine/Plugins/Chooser/Source/Chooser/Public/ChooserTypes.h b/Engine/Plugins/Chooser/Source/Chooser/Public/ChooserTypes.h
--- a/Engine/Plugins/Chooser/Source/Chooser/Public/ChooserTypes.h
+++ b/Engine/Plugins/Chooser/Source/Chooser/Public/ChooserTypes.h
@@ -79,6 +79,28 @@ struct FChooserPlayerSettings
 
 	UPROPERTY(EditAnywhere, BlueprintReadWrite, Category = Blending)
 	bool bUseInertialBlend = false;
+
+	// Experimental, this feature might be removed without warning, not for production use
+	// Currently playing animation
+	UPROPERTY(Transient)
+	TObjectPtr<const UObject> PlayingAsset = nullptr;
+
+	// Experimental, this feature might be removed without warning, not for production use
+	// Currently playing animation accumulated time
+	UPROPERTY(Transient)
+	float PlayingAssetAccumulatedTime = 0.f;
+
+	// Experimental, this feature might be removed without warning, not for production use
+	UPROPERTY(Transient)
+	bool bIsPlayingAssetMirrored = false;
+
+	// Experimental, this feature might be removed without warning, not for production use
+	// PlayingAsset associated BlendParameters (if PlayingAsset is a UBlendSpace)
+	UPROPERTY(Transient)
+	FVector PlayingAssetBlendParameters = FVector::ZeroVector;
+
+	// Experimental, this feature might be removed without warning, not for production use
+	const struct FAnimationUpdateContext* AnimationUpdateContext = nullptr;
 };
 
 template<> struct TStructOpsTypeTraits<FChooserPlayerSettings> : public TStructOpsTypeTraitsBase2<FChooserPlayerSettings>
diff --git a/Engine/Plugins/Chooser/Source/Chooser/Public/IChooserColumn.h b/Engine/Plugins/Chooser/Source/Chooser/Public/IChooserColumn.h
--- a/Engine/Plugins/Chooser/Source/Chooser/Public/IChooserColumn.h
+++ b/Engine/Plugins/Chooser/Source/Chooser/Public/IChooserColumn.h
@@ -37,7 +37,7 @@ struct FChooserColumnBase
 
 public:
 	virtual ~FChooserColumnBase() {}
-	virtual void PostLoad() {};
+	virtual void PostLoad() {}
 	virtual void Filter(FChooserEvaluationContext& Context, const FChooserIndexArray& IndexListIn, FChooserIndexArray& IndexListOut) const {}
 
 	virtual bool HasCosts() const { return false; }
@@ -60,9 +60,11 @@ public:
 		}
 	};
 
-	virtual void SetTestValue(TArrayView<const uint8> Value) { }
-	virtual FChooserParameterBase* GetInputValue() { return nullptr; };
+	virtual void SetTestValue(TArrayView<const uint8> Value) {}
+	virtual FChooserParameterBase* GetInputValue() { return nullptr; }
 #if WITH_EDITOR
+	virtual void Initialize(UChooserTable* OuterChooser) {}
+	
 	virtual FName RowValuesPropertyName() { return FName(); }
 	virtual void SetNumRows(int32 NumRows) {}
 	virtual void DeleteRows(const TArray<uint32> & RowIndices) {}
@@ -71,16 +73,16 @@ public:
 	virtual void CopyRow(FChooserColumnBase& SourceColumn, int SourceIndex, int TargetIndex) {}
 	virtual void CopyFallback(FChooserColumnBase& SourceColumn) {}
 	
-	virtual UScriptStruct* GetInputBaseType() const { return nullptr; };
-	virtual const UScriptStruct* GetInputType() const { return nullptr; };
-	virtual void SetInputType(const UScriptStruct* Type) { };
+	virtual UScriptStruct* GetInputBaseType() const { return nullptr; }
+	virtual const UScriptStruct* GetInputType() const { return nullptr; }
+	virtual void SetInputType(const UScriptStruct* Type) { }
 
 	// random columns must go last, and get a special icon
 	// using a virtual fucntion to identify them (rather than hard coding a specific type) to potentially support multiple varieties of randomization column.
 	virtual bool IsRandomizeColumn() const { return false; }
 
 	virtual bool AutoPopulates() const { return false; }
-	virtual void AutoPopulate(int32 RowIndex, UObject* OutputObject) { };
+	virtual void AutoPopulate(int32 RowIndex, UObject* OutputObject) { }
 
 	virtual float EditorTestCost(int32 RowIndex) const { return 0.0f; }
 	virtual bool EditorTestFilter(int32 RowIndex) const { return false; }
diff --git a/Engine/Plugins/Chooser/Source/ChooserEditor/Private/ChooserTableEditor.cpp b/Engine/Plugins/Chooser/Source/ChooserEditor/Private/ChooserTableEditor.cpp
--- a/Engine/Plugins/Chooser/Source/ChooserEditor/Private/ChooserTableEditor.cpp
+++ b/Engine/Plugins/Chooser/Source/ChooserEditor/Private/ChooserTableEditor.cpp
@@ -1160,6 +1160,7 @@ void FChooserTableEditor::AddColumn(const UScriptStruct* ColumnType)
 
 	FInstancedStruct NewColumn;
 	NewColumn.InitializeAs(ColumnType);
+	NewColumn.GetMutable<FChooserColumnBase>().Initialize(Chooser);
 	const FChooserColumnBase& NewColumnRef = NewColumn.Get<FChooserColumnBase>();
 	int InsertIndex = 0;
 	if (NewColumnRef.IsRandomizeColumn())
