PoseSearch - deprecating pose search database TArray<FInstancedStruct> AnimationAssets in favor of TArray<FPoseSearchDatabaseAnimationAsset> DatabaseAnimationAssets

diff --git a/Engine/Plugins/Animation/PoseSearch/Source/Editor/Private/Chooser/PoseSearchColumnEditor.cpp b/Engine/Plugins/Animation/PoseSearch/Source/Editor/Private/Chooser/PoseSearchColumnEditor.cpp
--- a/Engine/Plugins/Animation/PoseSearch/Source/Editor/Private/Chooser/PoseSearchColumnEditor.cpp
+++ b/Engine/Plugins/Animation/PoseSearch/Source/Editor/Private/Chooser/PoseSearchColumnEditor.cpp
@@ -40,7 +40,7 @@ TSharedRef<SWidget> CreatePoseSearchColumnWidget(UChooserTable* Chooser, FChoose
 		SNew(STextBlock)
 			.Text_Lambda([Row, PoseSearchColumn]()
 			{
-				if (FPoseSearchDatabaseSequence* Asset = PoseSearchColumn->InternalDatabase->GetMutableDatabaseAnimationAsset<FPoseSearchDatabaseSequence>(Row))
+				if (FPoseSearchDatabaseAnimationAsset* Asset = PoseSearchColumn->InternalDatabase->GetMutableDatabaseAnimationAsset(Row))
 				{
 					if (UObject* ResultAsset = Asset->GetAnimationAsset())
 					{
diff --git a/Engine/Plugins/Animation/PoseSearch/Source/Editor/Private/PoseSearchCustomization.cpp b/Engine/Plugins/Animation/PoseSearch/Source/Editor/Private/PoseSearchCustomization.cpp
--- a/Engine/Plugins/Animation/PoseSearch/Source/Editor/Private/PoseSearchCustomization.cpp
+++ b/Engine/Plugins/Animation/PoseSearch/Source/Editor/Private/PoseSearchCustomization.cpp
@@ -1,38 +1,37 @@
 // Copyright Epic Games, Inc. All Rights Reserved.
 
 #include "PoseSearchCustomization.h"
-#include "Engine/GameViewportClient.h"
+#include "Animation/AnimSequence.h"
+#include "Animation/BlendSpace.h"
 #include "AssetRegistry/AssetData.h"
+#include "DetailCategoryBuilder.h"
+#include "DetailLayoutBuilder.h"
+#include "DetailWidgetRow.h"
 #include "EditorClassUtils.h"
+#include "Engine/EngineBaseTypes.h"
+#include "Engine/GameViewportClient.h"
+#include "IDetailChildrenBuilder.h"
+#include "IDetailsView.h"
 #include "IPropertyUtilities.h"
-#include "PropertyHandle.h"
+#include "ObjectEditorUtils.h"
+#include "PoseSearch/MultiAnimAsset.h"
+#include "PoseSearch/PoseSearchDatabase.h"
 #include "PropertyCustomizationHelpers.h"
+#include "PropertyHandle.h"
 #include "ScopedTransaction.h"
-#include "Engine/EngineBaseTypes.h"
 #include "UObject/UnrealType.h"
 #include "Widgets/DeclarativeSyntaxSupport.h"
-#include "Widgets/SBoxPanel.h"
-#include "Widgets/Layout/SWrapBox.h"
-#include "Widgets/Layout/SBox.h"
-#include "Widgets/Text/STextBlock.h"
 #include "Widgets/Images/SImage.h"
 #include "Widgets/Input/SButton.h"
-#include "IDetailChildrenBuilder.h"
-#include "DetailLayoutBuilder.h"
-#include "DetailWidgetRow.h"
-#include "DetailCategoryBuilder.h"
-#include "IDetailsView.h"
-#include "ObjectEditorUtils.h"
-#include "Animation/AnimSequence.h"
-#include "PoseSearch/PoseSearchDatabase.h"
+#include "Widgets/Input/SCheckBox.h"
+#include "Widgets/Layout/SBox.h"
+#include "Widgets/Layout/SWrapBox.h"
+#include "Widgets/SBoxPanel.h"
+#include "Widgets/Text/STextBlock.h"
 
 #define LOCTEXT_NAMESPACE "PoseSearchCustomization"
 
-
-///////////////////////////////////////////////////////////////////////////////////////////////////////////////////////
-/// FPoseSearchDatabaseSequenceCustomization
-
-void FPoseSearchDatabaseSequenceCustomization::CustomizeHeader(TSharedRef<IPropertyHandle> InStructPropertyHandle, FDetailWidgetRow& HeaderRow, IPropertyTypeCustomizationUtils& StructCustomizationUtils)
+void FPoseSearchDatabaseAnimationAssetCustomization::CustomizeHeader(TSharedRef<IPropertyHandle> InStructPropertyHandle, FDetailWidgetRow& HeaderRow, IPropertyTypeCustomizationUtils& StructCustomizationUtils)
 {
 	FText SequenceNameText;
 
@@ -41,16 +40,16 @@ void FPoseSearchDatabaseSequenceCustomization::CustomizeHeader(TSharedRef<IPrope
 
 	if (Objects.Num() == 1)
 	{
-		FPoseSearchDatabaseSequence* PoseSearchDatabaseSequence = (FPoseSearchDatabaseSequence*)InStructPropertyHandle->GetValueBaseAddress((uint8*)Objects[0]);
-		check(PoseSearchDatabaseSequence);
+		FPoseSearchDatabaseAnimationAsset* DatabaseAnimationAsset = (FPoseSearchDatabaseAnimationAsset*)InStructPropertyHandle->GetValueBaseAddress((uint8*)Objects[0]);
+		check(DatabaseAnimationAsset);
 		
-		if (const UObject* AnimationAsset = PoseSearchDatabaseSequence->GetAnimationAsset())
+		if (const UObject* AnimationAsset = DatabaseAnimationAsset->GetAnimationAsset())
 		{
 			SequenceNameText = FText::FromName(AnimationAsset->GetFName());
 		}
 		else
 		{
-			SequenceNameText = LOCTEXT("NewSequenceLabel", "New Sequence");
+			SequenceNameText = LOCTEXT("NewDatabaseAnimationAssetLabel", "New Animation Asset");
 		}
 	}
 
@@ -69,16 +68,134 @@ void FPoseSearchDatabaseSequenceCustomization::CustomizeHeader(TSharedRef<IPrope
 	InStructPropertyHandle->SetOnChildPropertyValueChanged(OnValueChanged);
 }
 
-void FPoseSearchDatabaseSequenceCustomization::CustomizeChildren(TSharedRef<IPropertyHandle> InStructPropertyHandle, IDetailChildrenBuilder& StructBuilder, IPropertyTypeCustomizationUtils& StructCustomizationUtils)
+void FPoseSearchDatabaseAnimationAssetCustomization::CustomizeChildren(TSharedRef<IPropertyHandle> InStructPropertyHandle, IDetailChildrenBuilder& StructBuilder, IPropertyTypeCustomizationUtils& StructCustomizationUtils)
 {
 	uint32 NumChildren;
 	InStructPropertyHandle->GetNumChildren(NumChildren);
 
-	for (uint32 ChildIndex = 0; ChildIndex < NumChildren; ++ChildIndex)
+	TSharedPtr<IPropertyHandle> bEnabledHandle = InStructPropertyHandle->GetChildHandle(GET_MEMBER_NAME_CHECKED(FPoseSearchDatabaseAnimationAsset, bEnabled));
+	TSharedPtr<IPropertyHandle> bDisableReselectionHandle = InStructPropertyHandle->GetChildHandle(GET_MEMBER_NAME_CHECKED(FPoseSearchDatabaseAnimationAsset, bDisableReselection));
+	TSharedPtr<IPropertyHandle> MirrorOptionHandle = InStructPropertyHandle->GetChildHandle(GET_MEMBER_NAME_CHECKED(FPoseSearchDatabaseAnimationAsset, MirrorOption));
+	TSharedPtr<IPropertyHandle> BranchInIdHandle = InStructPropertyHandle->GetChildHandle(GET_MEMBER_NAME_CHECKED(FPoseSearchDatabaseAnimationAsset, BranchInId));
+
+	TSharedPtr<IPropertyHandle> AnimAssetHandle = InStructPropertyHandle->GetChildHandle(GET_MEMBER_NAME_CHECKED(FPoseSearchDatabaseAnimationAsset, AnimAsset));
+	TSharedPtr<IPropertyHandle> bUseSingleSampleHandle = InStructPropertyHandle->GetChildHandle(GET_MEMBER_NAME_CHECKED(FPoseSearchDatabaseAnimationAsset, bUseSingleSample));
+	TSharedPtr<IPropertyHandle> bUseGridForSamplingHandle = InStructPropertyHandle->GetChildHandle(GET_MEMBER_NAME_CHECKED(FPoseSearchDatabaseAnimationAsset, bUseGridForSampling));
+	TSharedPtr<IPropertyHandle> NumberOfHorizontalSamplesHandle = InStructPropertyHandle->GetChildHandle(GET_MEMBER_NAME_CHECKED(FPoseSearchDatabaseAnimationAsset, NumberOfHorizontalSamples));
+	TSharedPtr<IPropertyHandle> NumberOfVerticalSamplesHandle = InStructPropertyHandle->GetChildHandle(GET_MEMBER_NAME_CHECKED(FPoseSearchDatabaseAnimationAsset, NumberOfVerticalSamples));
+	TSharedPtr<IPropertyHandle> BlendParamXHandle = InStructPropertyHandle->GetChildHandle(GET_MEMBER_NAME_CHECKED(FPoseSearchDatabaseAnimationAsset, BlendParamX));
+	TSharedPtr<IPropertyHandle> BlendParamYHandle = InStructPropertyHandle->GetChildHandle(GET_MEMBER_NAME_CHECKED(FPoseSearchDatabaseAnimationAsset, BlendParamY));
+	TSharedPtr<IPropertyHandle> SamplingRangeHandle = InStructPropertyHandle->GetChildHandle(GET_MEMBER_NAME_CHECKED(FPoseSearchDatabaseAnimationAsset, SamplingRange));
+
+	const bool bIsPoseSearchDatabaseAnimationAssetLayoutValid =
+		NumChildren == 12 &&
+		bEnabledHandle.IsValid() &&
+		bDisableReselectionHandle.IsValid() &&
+		MirrorOptionHandle.IsValid() &&
+		BranchInIdHandle.IsValid() &&
+		AnimAssetHandle.IsValid() &&
+		bUseSingleSampleHandle.IsValid() &&
+		bUseGridForSamplingHandle.IsValid() &&
+		NumberOfHorizontalSamplesHandle.IsValid() &&
+		NumberOfVerticalSamplesHandle.IsValid() &&
+		BlendParamXHandle.IsValid() &&
+		BlendParamYHandle.IsValid() &&
+		SamplingRangeHandle.IsValid();
+
+	// update FPoseSearchDatabaseAnimationAssetCustomization if FPoseSearchDatabaseAnimationAsset or FPoseSearchDatabaseAnimationAssetBase changed!
+	if (!ensure(bIsPoseSearchDatabaseAnimationAssetLayoutValid))
 	{
-		StructBuilder.AddProperty(InStructPropertyHandle->GetChildHandle(ChildIndex).ToSharedRef());
+		for (uint32 ChildIndex = 0; ChildIndex < NumChildren; ++ChildIndex)
+		{
+			StructBuilder.AddProperty(InStructPropertyHandle->GetChildHandle(ChildIndex).ToSharedRef());
+		}
+		return;
 	}
-}
 
+
+	StructBuilder.AddProperty(AnimAssetHandle.ToSharedRef());
+	StructBuilder.AddProperty(bEnabledHandle.ToSharedRef());
+	StructBuilder.AddProperty(SamplingRangeHandle.ToSharedRef());
+	StructBuilder.AddProperty(bDisableReselectionHandle.ToSharedRef());
+	StructBuilder.AddProperty(MirrorOptionHandle.ToSharedRef());
+	StructBuilder.AddProperty(BranchInIdHandle.ToSharedRef());
+
+	UObject* AnimAsset = nullptr;
+	const FPropertyAccess::Result AnimAssetGetValueResult = AnimAssetHandle->GetValue(AnimAsset);
+	if (AnimAssetGetValueResult == FPropertyAccess::Success && AnimAsset)
+	{
+		if (AnimAsset->IsA<UBlendSpace>())
+		{
+			StructBuilder.AddProperty(bUseSingleSampleHandle.ToSharedRef());
+			StructBuilder.AddProperty(bUseGridForSamplingHandle.ToSharedRef());
+			StructBuilder.AddProperty(NumberOfHorizontalSamplesHandle.ToSharedRef());
+			StructBuilder.AddProperty(NumberOfVerticalSamplesHandle.ToSharedRef());
+			StructBuilder.AddProperty(BlendParamXHandle.ToSharedRef());
+			StructBuilder.AddProperty(BlendParamYHandle.ToSharedRef());
+		}
+		else if (AnimAsset->IsA<UMultiAnimAsset>())
+		{
+			StructBuilder.AddProperty(NumberOfHorizontalSamplesHandle.ToSharedRef());
+			StructBuilder.AddProperty(NumberOfVerticalSamplesHandle.ToSharedRef());
+		}
+	}
+	else if (FPropertyAccess::MultipleValues == AnimAssetGetValueResult)
+	{
+		StructBuilder.AddProperty(bUseSingleSampleHandle.ToSharedRef());
+		StructBuilder.AddProperty(bUseGridForSamplingHandle.ToSharedRef());
+		StructBuilder.AddProperty(NumberOfHorizontalSamplesHandle.ToSharedRef());
+		StructBuilder.AddProperty(NumberOfVerticalSamplesHandle.ToSharedRef());
+		StructBuilder.AddProperty(BlendParamXHandle.ToSharedRef());
+		StructBuilder.AddProperty(BlendParamYHandle.ToSharedRef());
+	}
+
+	FFloatInterval SamplingRange = FFloatInterval(0.f, 0.f);
+	void* StructData = nullptr;
+	const FPropertyAccess::Result SamplingRangeGetValueResult = SamplingRangeHandle->GetValueData(StructData);
+	if (SamplingRangeGetValueResult == FPropertyAccess::Success && StructData)
+	{
+		SamplingRange = *static_cast<FFloatInterval*>(StructData);
+	}
+
+	// Create and inject a check box to expose the animation looping state
+	const bool bLooping = FPoseSearchDatabaseAnimationAsset::IsLooping(AnimAsset, SamplingRange);
+	StructBuilder.AddProperty(InStructPropertyHandle).CustomWidget()
+		.NameContent()
+		[
+			InStructPropertyHandle->CreatePropertyNameWidget(
+				LOCTEXT("FPoseSearchDatabaseAnimationAssetCustomizationLoopingLabel", "Looping"),
+				LOCTEXT("FPoseSearchDatabaseAnimationAssetCustomizationLoopingTooltip", "Is this animation set as looping? If you want to change this you need to change it in the animation asset. Changing the sampling range will disable looping")
+			)
+		]
+		.ValueContent()
+		[
+			SNew(SCheckBox)
+				.IsChecked_Lambda([bLooping]()
+					{
+						return bLooping ? ECheckBoxState::Checked : ECheckBoxState::Unchecked;
+					})
+				.IsEnabled(false)
+		];
+
+	// Create and inject a check box to expose the animation root motion state
+	const bool bHasRootMotion = FPoseSearchDatabaseAnimationAsset::IsRootMotionEnabled(AnimAsset);
+	StructBuilder.AddProperty(InStructPropertyHandle).CustomWidget()
+		.NameContent()
+		[
+			InStructPropertyHandle->CreatePropertyNameWidget(
+				LOCTEXT("FPoseSearchDatabaseAnimationAssetCustomizationHasRootMotionLabel", "HasRootMotion"),
+				LOCTEXT("FPoseSearchDatabaseAnimationAssetCustomizationHasRootMotionTooltip", "Does this animation have root motion enabled ? If you want to change this you need to change it in the animation asset.")
+			)
+		]
+		.ValueContent()
+		[
+			SNew(SCheckBox)
+				.IsChecked_Lambda([bHasRootMotion]()
+					{ 
+						return bHasRootMotion ? ECheckBoxState::Checked : ECheckBoxState::Unchecked;
+					})
+				.IsEnabled(false)
+		];
+}
 #undef LOCTEXT_NAMESPACE
 
diff --git a/Engine/Plugins/Animation/PoseSearch/Source/Editor/Private/PoseSearchDatabaseAssetListItem.cpp b/Engine/Plugins/Animation/PoseSearch/Source/Editor/Private/PoseSearchDatabaseAssetListItem.cpp
--- a/Engine/Plugins/Animation/PoseSearch/Source/Editor/Private/PoseSearchDatabaseAssetListItem.cpp
+++ b/Engine/Plugins/Animation/PoseSearch/Source/Editor/Private/PoseSearchDatabaseAssetListItem.cpp
@@ -216,7 +216,7 @@ namespace UE::PoseSearch
 		AssetTypeColor = FColor::White;
 		if (UPoseSearchDatabase* Database = InEditorViewModel->GetPoseSearchDatabase())
 		{
-			if (const FPoseSearchDatabaseAnimationAssetBase* DatabaseAnimationAsset = Database->GetDatabaseAnimationAsset<FPoseSearchDatabaseAnimationAssetBase>(WeakAssetTreeNode.Pin()->SourceAssetIdx))
+			if (const FPoseSearchDatabaseAnimationAssetBase* DatabaseAnimationAsset = Database->GetDatabaseAnimationAsset(WeakAssetTreeNode.Pin()->SourceAssetIdx))
 			{
 				static FAssetToolsModule& AssetToolsModule = FModuleManager::LoadModuleChecked<FAssetToolsModule>(TEXT("AssetTools"));
 				if (TSharedPtr<IAssetTypeActions> AssetTypeActions = AssetToolsModule.Get().GetAssetTypeActionsForClass(DatabaseAnimationAsset->GetAnimationAssetStaticClass()).Pin())
@@ -267,30 +267,6 @@ namespace UE::PoseSearch
 			], OwnerTable);
 	}
 
-	void SDatabaseAssetListItem::OnAddSequence()
-	{
-		EditorViewModel.Pin()->AddSequenceToDatabase(nullptr);
-		SkeletonView.Pin()->RefreshTreeView(false);
-	}
-
-	void SDatabaseAssetListItem::OnAddBlendSpace()
-	{
-		EditorViewModel.Pin()->AddBlendSpaceToDatabase(nullptr);
-		SkeletonView.Pin()->RefreshTreeView(false);
-	}
-
-	void SDatabaseAssetListItem::OnAddAnimComposite()
-	{
-		EditorViewModel.Pin()->AddAnimCompositeToDatabase(nullptr);
-		SkeletonView.Pin()->RefreshTreeView(false);
-	}
-
-	void SDatabaseAssetListItem::OnAddAnimMontage()
-	{
-		EditorViewModel.Pin()->AddAnimMontageToDatabase(nullptr);
-		SkeletonView.Pin()->RefreshTreeView(false);
-	}
-
 	FReply SDatabaseAssetListItem::OnMouseButtonDoubleClick(const FGeometry& InMyGeometry, const FPointerEvent& InMouseEvent)
 	{
 		if (TSharedPtr<FDatabaseAssetTreeNode> AssetTreeNode = WeakAssetTreeNode.Pin())
@@ -298,7 +274,7 @@ namespace UE::PoseSearch
 			TSharedPtr<FDatabaseViewModel> ViewModel = EditorViewModel.Pin();
 			if (const UPoseSearchDatabase* Database = ViewModel->GetPoseSearchDatabase())
 			{
-				if (const FPoseSearchDatabaseAnimationAssetBase* DatabaseAnimationAsset = Database->GetDatabaseAnimationAsset<FPoseSearchDatabaseAnimationAssetBase>(AssetTreeNode->SourceAssetIdx))
+				if (const FPoseSearchDatabaseAnimationAssetBase* DatabaseAnimationAsset = Database->GetDatabaseAnimationAsset(AssetTreeNode->SourceAssetIdx))
 				{
 					if (UAssetEditorSubsystem* AssetEditorSS = GEditor->GetEditorSubsystem<UAssetEditorSubsystem>())
 					{
@@ -346,7 +322,7 @@ namespace UE::PoseSearch
 		{
 			if (TSharedPtr<FDatabaseAssetTreeNode> AssetTreeNode = WeakAssetTreeNode.Pin())
 			{
-				if (const FPoseSearchDatabaseAnimationAssetBase* DatabaseAnimationAsset = Database->GetDatabaseAnimationAsset<FPoseSearchDatabaseAnimationAssetBase>(AssetTreeNode->SourceAssetIdx))
+				if (const FPoseSearchDatabaseAnimationAssetBase* DatabaseAnimationAsset = Database->GetDatabaseAnimationAsset(AssetTreeNode->SourceAssetIdx))
 				{
 					return FText::FromString(DatabaseAnimationAsset->GetName());
 				}
@@ -413,7 +389,7 @@ namespace UE::PoseSearch
 				   {
 					  if (const TSharedPtr<FDatabaseAssetTreeNode> AssetTreeNode = WeakAssetTreeNode.Pin())
 					  {
-						 if (const FPoseSearchDatabaseAnimationAssetBase* DatabaseAnimationAssetBase = Database->GetDatabaseAnimationAsset<FPoseSearchDatabaseAnimationAssetBase>(AssetTreeNode->SourceAssetIdx))
+						 if (const FPoseSearchDatabaseAnimationAssetBase* DatabaseAnimationAssetBase = Database->GetDatabaseAnimationAsset(AssetTreeNode->SourceAssetIdx))
 						 {
 							if (DatabaseAnimationAssetBase->IsSynchronizedWithExternalDependency())
 							{
@@ -436,10 +412,13 @@ namespace UE::PoseSearch
 				TSharedPtr<SImage> ItemIconWidget;
 				if (UPoseSearchDatabase* Database = ViewModel->GetPoseSearchDatabase())
 				{
-					if (const FPoseSearchDatabaseAnimationAssetBase* DatabaseAnimationAsset = Database->GetDatabaseAnimationAsset<FPoseSearchDatabaseAnimationAssetBase>(SourceAssetIdx))
+					if (const FPoseSearchDatabaseAnimationAssetBase* DatabaseAnimationAsset = Database->GetDatabaseAnimationAsset(SourceAssetIdx))
 					{
 						SAssignNew(ItemIconWidget, SImage)
-						.Image(FSlateIconFinder::FindIconBrushForClass(DatabaseAnimationAsset->GetAnimationAssetStaticClass()));
+							.Image_Lambda([DatabaseAnimationAsset]
+								{
+									return FSlateIconFinder::FindIconBrushForClass(DatabaseAnimationAsset->GetAnimationAssetStaticClass());
+								});
 					}
 				}
 				
@@ -505,11 +484,12 @@ namespace UE::PoseSearch
 			TSharedPtr<SObjectPropertyEntryBox> AssetPickerWidget;
 			if (UPoseSearchDatabase* Database = ViewModel->GetPoseSearchDatabase())
 			{
-				if (const FPoseSearchDatabaseAnimationAssetBase* DatabaseAnimationAsset = Database->GetDatabaseAnimationAsset<FPoseSearchDatabaseAnimationAssetBase>(SourceAssetIdx))
+				if (const FPoseSearchDatabaseAnimationAssetBase* DatabaseAnimationAsset = Database->GetDatabaseAnimationAsset(SourceAssetIdx))
 				{
 					SAssignNew(AssetPickerWidget, SObjectPropertyEntryBox)
 					.AllowClear(false)
-					.AllowedClass(DatabaseAnimationAsset->GetAnimationAssetStaticClass())
+					// @todo: fix the AllowedClass(es). for now we filter inefficiently in OnShouldFilterAsset_Lambda
+					//.AllowedClass(DatabaseAnimationAsset->GetAnimationAssetStaticClass())
 					.DisplayThumbnail(false)
 					.IsEnabled(this, &SDatabaseAssetListItem::GetAssetPickerIsEnabled)
 					.ObjectPath(this, &SDatabaseAssetListItem::GetAssetPickerObjectPath)
@@ -742,7 +722,7 @@ namespace UE::PoseSearch
 		{
 			if (const UPoseSearchDatabase* Database = ViewModelPtr->GetPoseSearchDatabase())
 			{
-				if (const FPoseSearchDatabaseAnimationAssetBase* DatabaseAnimationAsset = Database->GetDatabaseAnimationAsset<FPoseSearchDatabaseAnimationAssetBase>(AssetTreeNode->SourceAssetIdx))
+				if (const FPoseSearchDatabaseAnimationAssetBase* DatabaseAnimationAsset = Database->GetDatabaseAnimationAsset(AssetTreeNode->SourceAssetIdx))
 				{
 					if (const UObject* AnimAsset = DatabaseAnimationAsset->GetAnimationAsset())
 					{
@@ -763,7 +743,7 @@ namespace UE::PoseSearch
 		{
 			if (const UPoseSearchDatabase* Database = ViewModelPtr->GetPoseSearchDatabase())
 			{
-				if (Database->GetDatabaseAnimationAsset<FPoseSearchDatabaseAnimationAssetBase>(AssetTreeNode->SourceAssetIdx))
+				if (Database->GetDatabaseAnimationAsset(AssetTreeNode->SourceAssetIdx))
 				{
 					return ViewModelPtr->IsEnabled(AssetTreeNode->SourceAssetIdx);
 				}
@@ -781,7 +761,7 @@ namespace UE::PoseSearch
 		{
 			if (const TSharedPtr<FDatabaseAssetTreeNode> AssetTreeNode = WeakAssetTreeNode.Pin())
 			{
-				if (const FPoseSearchDatabaseAnimationAssetBase* DatabaseAnimationAssetBase = Database->GetDatabaseAnimationAsset<FPoseSearchDatabaseAnimationAssetBase>(AssetTreeNode->SourceAssetIdx))
+				if (const FPoseSearchDatabaseAnimationAssetBase* DatabaseAnimationAssetBase = Database->GetDatabaseAnimationAsset(AssetTreeNode->SourceAssetIdx))
 				{
 					if (DatabaseAnimationAssetBase->IsEnabled())
 					{
@@ -805,7 +785,7 @@ namespace UE::PoseSearch
 		{
 			if (const TSharedPtr<FDatabaseAssetTreeNode> AssetTreeNode = WeakAssetTreeNode.Pin())
 			{
-				if (const FPoseSearchDatabaseAnimationAssetBase* DatabaseAnimationAssetBase = Database->GetDatabaseAnimationAsset<FPoseSearchDatabaseAnimationAssetBase>(AssetTreeNode->SourceAssetIdx))
+				if (const FPoseSearchDatabaseAnimationAssetBase* DatabaseAnimationAssetBase = Database->GetDatabaseAnimationAsset(AssetTreeNode->SourceAssetIdx))
 				{
 					if (DatabaseAnimationAssetBase->IsEnabled())
 					{
@@ -835,7 +815,7 @@ namespace UE::PoseSearch
 		{
 			if (TSharedPtr<FDatabaseAssetTreeNode> AssetTreeNode = WeakAssetTreeNode.Pin())
 			{
-				if (Database->GetDatabaseAnimationAsset<FPoseSearchDatabaseAnimationAssetBase>(AssetTreeNode->SourceAssetIdx))
+				if (Database->GetDatabaseAnimationAsset(AssetTreeNode->SourceAssetIdx))
 				{
 					if (ViewModelPtr->IsDisableReselection(AssetTreeNode->SourceAssetIdx))
 					{
@@ -873,7 +853,7 @@ namespace UE::PoseSearch
 		{
 			if (TSharedPtr<FDatabaseAssetTreeNode> AssetTreeNode = WeakAssetTreeNode.Pin())
 			{
-				if (Database->GetDatabaseAnimationAsset<FPoseSearchDatabaseAnimationAssetBase>(AssetTreeNode->SourceAssetIdx))
+				if (Database->GetDatabaseAnimationAsset(AssetTreeNode->SourceAssetIdx))
 				{
 					if (ViewModelPtr->IsEnabled(AssetTreeNode->SourceAssetIdx))
 					{
@@ -907,7 +887,7 @@ namespace UE::PoseSearch
 		{
 			if (const TSharedPtr<FDatabaseAssetTreeNode> AssetTreeNode = WeakAssetTreeNode.Pin())
 			{
-				if (const FPoseSearchDatabaseAnimationAssetBase* DatabaseAnimationAssetBase = Database->GetDatabaseAnimationAsset<FPoseSearchDatabaseAnimationAssetBase>(AssetTreeNode->SourceAssetIdx))
+				if (const FPoseSearchDatabaseAnimationAssetBase* DatabaseAnimationAssetBase = Database->GetDatabaseAnimationAsset(AssetTreeNode->SourceAssetIdx))
 				{
 					if (DatabaseAnimationAssetBase->IsEnabled())
 					{
diff --git a/Engine/Plugins/Animation/PoseSearch/Source/Editor/Private/PoseSearchDatabaseAssetTree.cpp b/Engine/Plugins/Animation/PoseSearch/Source/Editor/Private/PoseSearchDatabaseAssetTree.cpp
--- a/Engine/Plugins/Animation/PoseSearch/Source/Editor/Private/PoseSearchDatabaseAssetTree.cpp
+++ b/Engine/Plugins/Animation/PoseSearch/Source/Editor/Private/PoseSearchDatabaseAssetTree.cpp
@@ -59,8 +59,12 @@ namespace UE::PoseSearch
 					SNew(SPositiveActionButton)
 					.Icon(FAppStyle::Get().GetBrush("Icons.Plus"))
 					.Text(LOCTEXT("AddNew", "Add"))
-					.ToolTipText(LOCTEXT("AddNewToolTip", "Add a new Sequence, Blend Space, Anim Composite, or Anim Montage"))
-					.OnGetMenuContent(this, &SDatabaseAssetTree::CreateAddNewMenuWidget)
+					.ToolTipText(LOCTEXT("AddNewAssetToolTip", "Add a new asset to the database"))
+					.OnClicked_Lambda([this]()
+					{
+						OnAddAsset(true);
+						return FReply::Handled();
+					})
 				]
 				+ SHorizontalBox::Slot()
 				.AutoWidth()
@@ -246,13 +250,13 @@ namespace UE::PoseSearch
 
 			IndexArray.Sort([Database](int32 SequenceIdxA, int32 SequenceIdxB)
 			{
-				const FPoseSearchDatabaseAnimationAssetBase* A = Database->GetDatabaseAnimationAsset<FPoseSearchDatabaseAnimationAssetBase>(SequenceIdxA);
+				const FPoseSearchDatabaseAnimationAssetBase* A = Database->GetDatabaseAnimationAsset(SequenceIdxA);
 				if (!A)
 				{
 					return false;
 				}
 
-				const FPoseSearchDatabaseAnimationAssetBase* B = Database->GetDatabaseAnimationAsset<FPoseSearchDatabaseAnimationAssetBase>(SequenceIdxB);
+				const FPoseSearchDatabaseAnimationAssetBase* B = Database->GetDatabaseAnimationAsset(SequenceIdxB);
 				if (!B)
 				{
 					return true;
@@ -278,7 +282,7 @@ namespace UE::PoseSearch
 			{
 				const int32 MappedId = IndexArray[AnimationAssetIdx];
 
-				if (const FPoseSearchDatabaseAnimationAssetBase* DatabaseAnimationAsset = Database->GetDatabaseAnimationAsset<FPoseSearchDatabaseAnimationAssetBase>(MappedId))
+				if (const FPoseSearchDatabaseAnimationAssetBase* DatabaseAnimationAsset = Database->GetDatabaseAnimationAsset(MappedId))
 				{
 					const bool bFiltered = (DatabaseAnimationAsset->GetAnimationAsset() == nullptr || GetAssetFilterString().IsEmpty()) ? false : !DatabaseAnimationAsset->GetName().Contains(GetAssetFilterString());
 
@@ -413,31 +417,8 @@ namespace UE::PoseSearch
 				UClass* AssetClass = AssetData.GetClass();
 				UObject* Asset = AssetData.GetAsset();
 					
-				if (AssetClass->IsChildOf(UAnimSequence::StaticClass()))
-				{
-					ViewModel->AddSequenceToDatabase(Cast<UAnimSequence>(Asset));
-					++AddedAssets;
-				}
-				if (AssetClass->IsChildOf(UAnimComposite::StaticClass()))
-				{
-					ViewModel->AddAnimCompositeToDatabase(Cast<UAnimComposite>(Asset));
-					++AddedAssets;
-				}
-				else if (AssetClass->IsChildOf(UBlendSpace::StaticClass()))
-				{
-					ViewModel->AddBlendSpaceToDatabase(Cast<UBlendSpace>(Asset));
-					++AddedAssets;
-				}
-				else if (AssetClass->IsChildOf(UAnimMontage::StaticClass()))
-				{
-					ViewModel->AddAnimMontageToDatabase(Cast<UAnimMontage>(Asset));
-					++AddedAssets;
-				}
-				else if (AssetClass->IsChildOf(UMultiAnimAsset::StaticClass()))
-				{
-					ViewModel->AddMultiAnimAssetToDatabase(Cast<UMultiAnimAsset>(Asset));
-					++AddedAssets;
-				}
+				ViewModel->AddAnimationAssetToDatabase(Asset);
+				++AddedAssets;
 			}
 			
 			GWarn->EndSlowTask();
@@ -452,57 +433,6 @@ namespace UE::PoseSearch
 		return FReply::Handled();
 	}
 
-	TSharedRef<SWidget> SDatabaseAssetTree::CreateAddNewMenuWidget()
-	{
-		FMenuBuilder AddOptions(true, nullptr);
-
-		AddOptions.BeginSection("AddOptions", LOCTEXT("AssetAddOptions", "Assets"));
-		{
-			AddOptions.AddMenuEntry(
-				LOCTEXT("AddSequenceOption", "Sequence"),
-				LOCTEXT("AddSequenceToDatabaseTooltip", "Add new sequence to the database"),
-				FSlateIcon(),
-				FUIAction(FExecuteAction::CreateSP(this, &SDatabaseAssetTree::OnAddSequence, true)),
-				NAME_None,
-				EUserInterfaceActionType::Button);
-
-			AddOptions.AddMenuEntry(
-				LOCTEXT("AddBlendSpaceOption", "Blend Space"),
-				LOCTEXT("AddBlendSpaceToDatabaseTooltip", "Add new blend space to the database"),
-				FSlateIcon(),
-				FUIAction(FExecuteAction::CreateSP(this, &SDatabaseAssetTree::OnAddBlendSpace, true)),
-				NAME_None,
-				EUserInterfaceActionType::Button);
-
-			AddOptions.AddMenuEntry(
-				LOCTEXT("AddAnimCompositeOption", "Anim Composite"),
-				LOCTEXT("AddAnimCompositeToDatabaseTooltip", "Add new composite to the database"),
-				FSlateIcon(),
-				FUIAction(FExecuteAction::CreateSP(this, &SDatabaseAssetTree::OnAddAnimComposite, true)),
-				NAME_None,
-				EUserInterfaceActionType::Button);
-
-			AddOptions.AddMenuEntry(
-				LOCTEXT("AddAnimMontageOption", "Anim Montage"),
-				LOCTEXT("AddAnimMontageToDatabaseTooltip", "Add new montage to the database"),
-				FSlateIcon(),
-				FUIAction(FExecuteAction::CreateSP(this, &SDatabaseAssetTree::OnAddAnimMontage, true)),
-				NAME_None,
-				EUserInterfaceActionType::Button);
-
-			AddOptions.AddMenuEntry(
-				LOCTEXT("AddMultiAnimAssetOption", "Multi Anim Asset"),
-				LOCTEXT("AddMultiAnimAssetToDatabaseTooltip", "Add new multi anim asset to the database"),
-				FSlateIcon(),
-				FUIAction(FExecuteAction::CreateSP(this, &SDatabaseAssetTree::OnAddMultiAnimAsset, true)),
-				NAME_None,
-				EUserInterfaceActionType::Button);
-		}
-		AddOptions.EndSection();
-
-		return AddOptions.MakeWidget();
-	}
-
 	TSharedPtr<SWidget> SDatabaseAssetTree::CreateContextMenu()
 	{
 		const bool CloseAfterSelection = true;
@@ -672,11 +602,15 @@ namespace UE::PoseSearch
 		{
 			// Asset actions
 			MenuBuilder.BeginSection("Edit", LOCTEXT("EditSection", "Edit"));
-			MenuBuilder.AddWrapperSubMenu(
-				LOCTEXT("AddNewAnimAssetNoNodes", "Add"),
-				LOCTEXT("AddNewAnimAssetNoNodesToolTip", "Add a new Sequence, Blend Space, Anim Composite, or Anim Montage"),
-				FOnGetContent::CreateSP(this, &SDatabaseAssetTree::CreateAddNewMenuWidget),
-				FSlateIcon(FAppStyle::GetAppStyleSetName(), "Icons.Plus"));
+			{
+				MenuBuilder.AddMenuEntry(
+					LOCTEXT("AddAssetOption", "Add New Asset"),
+					LOCTEXT("AddAssetToDatabaseTooltip", "Add new asset to the database"),
+					FSlateIcon(),
+					FUIAction(FExecuteAction::CreateSP(this, &SDatabaseAssetTree::OnAddAsset, true)),
+					NAME_None,
+					EUserInterfaceActionType::Button);
+			}
 			MenuBuilder.EndSection();
 			
 			// Edit / Clipboard actions
@@ -717,11 +651,11 @@ namespace UE::PoseSearch
 	}
 
 
-	void SDatabaseAssetTree::OnAddSequence(bool bFinalizeChanges)
+	void SDatabaseAssetTree::OnAddAsset(bool bFinalizeChanges)
 	{
-		FScopedTransaction Transaction(LOCTEXT("AddSequence", "Add Sequence"));
+		FScopedTransaction Transaction(LOCTEXT("AddAsset", "Add Asset"));
 
-		EditorViewModel.Pin()->AddSequenceToDatabase(nullptr);
+		EditorViewModel.Pin()->AddAnimationAssetToDatabase(nullptr);
 
 		if (bFinalizeChanges)
 		{
@@ -730,54 +664,6 @@ namespace UE::PoseSearch
 		
 	}
 
-	void SDatabaseAssetTree::OnAddBlendSpace(bool bFinalizeChanges)
-	{
-		FScopedTransaction Transaction(LOCTEXT("AddBlendSpaceTransaction", "Add Blend Space"));
-		
-		EditorViewModel.Pin()->AddBlendSpaceToDatabase(nullptr);
-
-		if (bFinalizeChanges)
-		{
-			FinalizeTreeChanges();
-		}
-	}
-
-	void SDatabaseAssetTree::OnAddAnimComposite(bool bFinalizeChanges)
-	{
-		FScopedTransaction Transaction(LOCTEXT("AddAnimCompositeTransaction", "Add Anim Composite"));
-
-		EditorViewModel.Pin()->AddAnimCompositeToDatabase(nullptr);
-
-		if (bFinalizeChanges)
-		{
-			FinalizeTreeChanges();
-		}
-	}
-
-	void SDatabaseAssetTree::OnAddAnimMontage(bool bFinalizeChanges)
-	{
-		FScopedTransaction Transaction(LOCTEXT("AddAnimMontageTransaction", "Add Anim Montage"));
-
-		EditorViewModel.Pin()->AddAnimMontageToDatabase(nullptr);
-
-		if (bFinalizeChanges)
-		{
-			FinalizeTreeChanges();
-		}
-	}
-
-	void SDatabaseAssetTree::OnAddMultiAnimAsset(bool bFinalizeChanges)
-	{
-		FScopedTransaction Transaction(LOCTEXT("AddMultiAnimAssetTransaction", "Add Multi Anim Asset"));
-
-		EditorViewModel.Pin()->AddMultiAnimAssetToDatabase(nullptr);
-
-		if (bFinalizeChanges)
-		{
-			FinalizeTreeChanges();
-		}
-	}
-
 	void SDatabaseAssetTree::OnDeleteAsset(TSharedPtr<FDatabaseAssetTreeNode> Node, bool bFinalizeChanges)
 	{
 		FScopedTransaction Transaction(LOCTEXT("DeleteAsset", "Delete Asset"));
@@ -899,7 +785,7 @@ namespace UE::PoseSearch
 					{
 						if (UPoseSearchDatabase* Database = ViewModel->GetPoseSearchDatabase())
 						{
-							if (const FPoseSearchDatabaseAnimationAssetBase* DatabaseAnimationAsset = Database->GetDatabaseAnimationAsset<FPoseSearchDatabaseAnimationAssetBase>(SelectedNode->SourceAssetIdx))
+							if (const FPoseSearchDatabaseAnimationAssetBase* DatabaseAnimationAsset = Database->GetDatabaseAnimationAsset(SelectedNode->SourceAssetIdx))
 							{
 								// @todo: Support copying assets added via BranchIn notifies. 
 								if (!DatabaseAnimationAsset->IsSynchronizedWithExternalDependency())
@@ -979,7 +865,7 @@ namespace UE::PoseSearch
 				{
 					if (UPoseSearchDatabase* Database = ViewModel->GetPoseSearchDatabase())
 					{
-						const FPoseSearchDatabaseAnimationAssetBase* DatabaseAnimationAsset = Database->GetDatabaseAnimationAsset<FPoseSearchDatabaseAnimationAssetBase>(SelectedNode->SourceAssetIdx);
+						const FPoseSearchDatabaseAnimationAssetBase* DatabaseAnimationAsset = Database->GetDatabaseAnimationAsset(SelectedNode->SourceAssetIdx);
 						
 						if (DatabaseAnimationAsset && !DatabaseAnimationAsset->IsSynchronizedWithExternalDependency())
 						{
@@ -1080,7 +966,7 @@ namespace UE::PoseSearch
 				
 				for (const TSharedPtr<FDatabaseAssetTreeNode>& SelectedNode : SelectedNodes)
 				{
-					if (FPoseSearchDatabaseAnimationAssetBase* DatabaseAnimationAssetBase = PoseSearchDatabase->GetMutableDatabaseAnimationAsset<FPoseSearchDatabaseAnimationAssetBase>(SelectedNode->SourceAssetIdx))
+					if (FPoseSearchDatabaseAnimationAssetBase* DatabaseAnimationAssetBase = PoseSearchDatabase->GetMutableDatabaseAnimationAsset(SelectedNode->SourceAssetIdx))
 					{
 						if (UAnimSequenceBase* AnimSequenceBase = Cast<UAnimSequenceBase>(DatabaseAnimationAssetBase->GetAnimationAsset()))
 						{
diff --git a/Engine/Plugins/Animation/PoseSearch/Source/Editor/Private/PoseSearchDatabaseAssetTreeNode.cpp b/Engine/Plugins/Animation/PoseSearch/Source/Editor/Private/PoseSearchDatabaseAssetTreeNode.cpp
--- a/Engine/Plugins/Animation/PoseSearch/Source/Editor/Private/PoseSearchDatabaseAssetTreeNode.cpp
+++ b/Engine/Plugins/Animation/PoseSearch/Source/Editor/Private/PoseSearchDatabaseAssetTreeNode.cpp
@@ -41,7 +41,7 @@ namespace UE::PoseSearch
 	{
 		if (const UPoseSearchDatabase* Database = EditorViewModel.Pin()->GetPoseSearchDatabase())
 		{
-			if (const FPoseSearchDatabaseAnimationAssetBase* DatabaseAnimationAsset = Database->GetDatabaseAnimationAsset<FPoseSearchDatabaseAnimationAssetBase>(SourceAssetIdx))
+			if (const FPoseSearchDatabaseAnimationAssetBase* DatabaseAnimationAsset = Database->GetDatabaseAnimationAsset(SourceAssetIdx))
 			{
 				return DatabaseAnimationAsset->IsRootMotionEnabled();
 			}
@@ -54,7 +54,7 @@ namespace UE::PoseSearch
 	{
 		if (const UPoseSearchDatabase* Database = EditorViewModel.Pin()->GetPoseSearchDatabase())
 		{
-			if (const FPoseSearchDatabaseAnimationAssetBase* DatabaseAnimationAsset = Database->GetDatabaseAnimationAsset<FPoseSearchDatabaseAnimationAssetBase>(SourceAssetIdx))
+			if (const FPoseSearchDatabaseAnimationAssetBase* DatabaseAnimationAsset = Database->GetDatabaseAnimationAsset(SourceAssetIdx))
 			{
 				return DatabaseAnimationAsset->IsLooping();
 			}
@@ -67,7 +67,7 @@ namespace UE::PoseSearch
 	{
 		if (const UPoseSearchDatabase* Database = EditorViewModel.Pin()->GetPoseSearchDatabase())
 		{
-			if (const FPoseSearchDatabaseAnimationAssetBase* DatabaseAnimationAsset = Database->GetDatabaseAnimationAsset<FPoseSearchDatabaseAnimationAssetBase>(SourceAssetIdx))
+			if (const FPoseSearchDatabaseAnimationAssetBase* DatabaseAnimationAsset = Database->GetDatabaseAnimationAsset(SourceAssetIdx))
 			{
 				return DatabaseAnimationAsset->GetMirrorOption();
 			}
diff --git a/Engine/Plugins/Animation/PoseSearch/Source/Editor/Private/PoseSearchDatabaseEditor.cpp b/Engine/Plugins/Animation/PoseSearch/Source/Editor/Private/PoseSearchDatabaseEditor.cpp
--- a/Engine/Plugins/Animation/PoseSearch/Source/Editor/Private/PoseSearchDatabaseEditor.cpp
+++ b/Engine/Plugins/Animation/PoseSearch/Source/Editor/Private/PoseSearchDatabaseEditor.cpp
@@ -547,93 +547,32 @@ namespace UE::PoseSearch
 		const TArrayView<TSharedPtr<FDatabaseAssetTreeNode>>& SelectedItems,
 		ESelectInfo::Type SelectionType)
 	{
-		// Reset selected objects on all selection widgets
-		for (TPair<const UScriptStruct*, FSelectionWidget>& SelectionWidgetPair : SelectionWidgets)
-		{
-			SelectionWidgetPair.Value.SelectedReflections.Reset();
-		}
+		// Reset selected objects on the selection widget
+		SelectionWidget.SelectedReflections.Reset();
 
 		const UPoseSearchDatabase* PoseSearchDatabase = GetPoseSearchDatabase();
 		if (PoseSearchDatabase && SelectedItems.Num() > 0)
 		{
 			for (TSharedPtr<FDatabaseAssetTreeNode>& SelectedItem : SelectedItems)
 			{
-				if (!SelectedItem.IsValid() || !PoseSearchDatabase->GetDatabaseAnimationAsset<FPoseSearchDatabaseAnimationAssetBase>(SelectedItem->SourceAssetIdx))
-				{
-					continue;
-				}
-
-				if (const FPoseSearchDatabaseSequence* DatabaseSequence = PoseSearchDatabase->GetDatabaseAnimationAsset<FPoseSearchDatabaseSequence>(SelectedItem->SourceAssetIdx))
-				{
-					UPoseSearchDatabaseSequenceReflection* NewSelectionReflection = NewObject<UPoseSearchDatabaseSequenceReflection>();
-					NewSelectionReflection->AddToRoot();
-					static_cast<FPoseSearchDatabaseSequence&>(NewSelectionReflection->Sequence) = *DatabaseSequence;
-					NewSelectionReflection->Sequence.bLooping = DatabaseSequence->IsLooping();
-					NewSelectionReflection->Sequence.bHasRootMotion = DatabaseSequence->IsRootMotionEnabled();
-					NewSelectionReflection->SetSourceLink(SelectedItem, AssetTreeWidget);
-					NewSelectionReflection->SetFlags(RF_Transactional);
-					
-					FindOrAddSelectionWidget(FPoseSearchDatabaseSequence::StaticStruct()).SelectedReflections.Add(NewSelectionReflection);
-				}
-				else if (const FPoseSearchDatabaseAnimComposite* DatabaseAnimComposite = PoseSearchDatabase->GetDatabaseAnimationAsset<FPoseSearchDatabaseAnimComposite>(SelectedItem->SourceAssetIdx))
-				{
-					UPoseSearchDatabaseAnimCompositeReflection* NewSelectionReflection = NewObject<UPoseSearchDatabaseAnimCompositeReflection>();
-					NewSelectionReflection->AddToRoot();
-					static_cast<FPoseSearchDatabaseAnimComposite&>(NewSelectionReflection->AnimComposite) = *DatabaseAnimComposite;
-					NewSelectionReflection->AnimComposite.bLooping = DatabaseAnimComposite->IsLooping();
-					NewSelectionReflection->AnimComposite.bHasRootMotion = DatabaseAnimComposite->IsRootMotionEnabled();
-					NewSelectionReflection->SetSourceLink(SelectedItem, AssetTreeWidget);
-					NewSelectionReflection->SetFlags(RF_Transactional);
-					
-					FindOrAddSelectionWidget(FPoseSearchDatabaseAnimComposite::StaticStruct()).SelectedReflections.Add(NewSelectionReflection);
-				}
-				else if (const FPoseSearchDatabaseBlendSpace* DatabaseBlendSpace = PoseSearchDatabase->GetDatabaseAnimationAsset<FPoseSearchDatabaseBlendSpace>(SelectedItem->SourceAssetIdx))
+				if (SelectedItem.IsValid())
 				{
-					UPoseSearchDatabaseBlendSpaceReflection* NewSelectionReflection = NewObject<UPoseSearchDatabaseBlendSpaceReflection>();
-					NewSelectionReflection->AddToRoot();
-					static_cast<FPoseSearchDatabaseBlendSpace&>(NewSelectionReflection->BlendSpace) = *DatabaseBlendSpace;
-					NewSelectionReflection->BlendSpace.bLooping = DatabaseBlendSpace->IsLooping();
-					NewSelectionReflection->BlendSpace.bHasRootMotion = DatabaseBlendSpace->IsRootMotionEnabled();
-					NewSelectionReflection->SetSourceLink(SelectedItem, AssetTreeWidget);
-					NewSelectionReflection->SetFlags(RF_Transactional);
-					
-					FindOrAddSelectionWidget(FPoseSearchDatabaseBlendSpace::StaticStruct()).SelectedReflections.Add(NewSelectionReflection);
-				}
-				else if (const FPoseSearchDatabaseAnimMontage* DatabaseAnimMontage = PoseSearchDatabase->GetDatabaseAnimationAsset<FPoseSearchDatabaseAnimMontage>(SelectedItem->SourceAssetIdx))
-				{
-					UPoseSearchDatabaseAnimMontageReflection* NewSelectionReflection = NewObject<UPoseSearchDatabaseAnimMontageReflection>();
-					NewSelectionReflection->AddToRoot();
-					static_cast<FPoseSearchDatabaseAnimMontage&>(NewSelectionReflection->AnimMontage) = *DatabaseAnimMontage;
-					NewSelectionReflection->AnimMontage.bLooping = DatabaseAnimMontage->IsLooping();
-					NewSelectionReflection->AnimMontage.bHasRootMotion = DatabaseAnimMontage->IsRootMotionEnabled();
-					NewSelectionReflection->SetSourceLink(SelectedItem, AssetTreeWidget);
-					NewSelectionReflection->SetFlags(RF_Transactional);
-
-					FindOrAddSelectionWidget(FPoseSearchDatabaseAnimMontage::StaticStruct()).SelectedReflections.Add(NewSelectionReflection);
-				}
-				else if (const FPoseSearchDatabaseMultiAnimAsset* DatabaseMultiAnimAsset = PoseSearchDatabase->GetDatabaseAnimationAsset<FPoseSearchDatabaseMultiAnimAsset>(SelectedItem->SourceAssetIdx))
-				{
-					UPoseSearchDatabaseMultiAnimAssetReflection* NewSelectionReflection = NewObject<UPoseSearchDatabaseMultiAnimAssetReflection>();
-					NewSelectionReflection->AddToRoot();
-					static_cast<FPoseSearchDatabaseMultiAnimAsset&>(NewSelectionReflection->MultiAnimAsset) = *DatabaseMultiAnimAsset;
-					NewSelectionReflection->MultiAnimAsset.bLooping = DatabaseMultiAnimAsset->IsLooping();
-					NewSelectionReflection->MultiAnimAsset.bHasRootMotion = DatabaseMultiAnimAsset->IsRootMotionEnabled();
-					NewSelectionReflection->SetSourceLink(SelectedItem, AssetTreeWidget);
-					NewSelectionReflection->SetFlags(RF_Transactional);
-
-					FindOrAddSelectionWidget(FPoseSearchDatabaseMultiAnimAsset::StaticStruct()).SelectedReflections.Add(NewSelectionReflection);
-				}
-				else
-				{
-					checkNoEntry();
+					if (const FPoseSearchDatabaseAnimationAsset* DatabaseAnimationAsset = PoseSearchDatabase->GetDatabaseAnimationAsset(SelectedItem->SourceAssetIdx))
+					{
+						UPoseSearchDatabaseAnimationAssetReflection* NewSelectionReflection = NewObject<UPoseSearchDatabaseAnimationAssetReflection>();
+						NewSelectionReflection->AddToRoot();
+						NewSelectionReflection->DatabaseAnimationAsset = *DatabaseAnimationAsset;
+						NewSelectionReflection->SetSourceLink(SelectedItem, AssetTreeWidget);
+						NewSelectionReflection->SetFlags(RF_Transactional);
+
+						GetInitializedSelectionWidget().SelectedReflections.Add(NewSelectionReflection);
+					}
 				}
 			}
 		}
 
-		for (TPair<const UScriptStruct*, FSelectionWidget>& SelectionWidgetPair : SelectionWidgets)
+		if (SelectionWidget.DetailView.IsValid())
 		{
-			FSelectionWidget& SelectionWidget = SelectionWidgetPair.Value;
-
 			if (SelectionWidget.SelectedReflections.IsEmpty())
 			{
 				SelectionWidget.DetailView->SetObject(nullptr, true);
@@ -650,15 +589,13 @@ namespace UE::PoseSearch
 		}
 	}
 
-	FDatabaseEditor::FSelectionWidget& FDatabaseEditor::FindOrAddSelectionWidget(const UScriptStruct* ScriptStructType)
+	FDatabaseEditor::FSelectionWidget& FDatabaseEditor::GetInitializedSelectionWidget()
 	{
-		if (SelectionWidgets.Contains(ScriptStructType))
+		if (SelectionWidget.DetailView.IsValid())
 		{
-			return SelectionWidgets[ScriptStructType];
+			return SelectionWidget;
 		}
 
-		FSelectionWidget& SelectionWidget = SelectionWidgets.Add(ScriptStructType);
-
 		FDetailsViewArgs SelectionDetailsArgs;
 		{
 			SelectionDetailsArgs.bHideSelectionTip = true;
diff --git a/Engine/Plugins/Animation/PoseSearch/Source/Editor/Private/PoseSearchDatabaseEditorClipboard.cpp b/Engine/Plugins/Animation/PoseSearch/Source/Editor/Private/PoseSearchDatabaseEditorClipboard.cpp
--- a/Engine/Plugins/Animation/PoseSearch/Source/Editor/Private/PoseSearchDatabaseEditorClipboard.cpp
+++ b/Engine/Plugins/Animation/PoseSearch/Source/Editor/Private/PoseSearchDatabaseEditorClipboard.cpp
@@ -24,28 +24,11 @@ void UPoseSearchDatabaseEditorClipboardContent::CopyDatabaseItem(const FPoseSear
 	
 	UPoseSearchDatabaseItemCopyObject* CopyObj = NewObject<UPoseSearchDatabaseItemCopyObject>((UObject*)GetTransientPackage(), UPoseSearchDatabaseItemCopyObject::StaticClass(), NAME_None, RF_Transient);
 	
-	CopyObj->ClassName = InItem->GetAnimationAssetStaticClass()->GetFName();
-	
-	if (UAnimSequence::StaticClass() == InItem->GetAnimationAssetStaticClass())
-	{
-		FPoseSearchDatabaseSequence::StaticStruct()->ExportText(CopyObj->Content, InItem, nullptr, nullptr, PPF_None, nullptr);
-	}
-	else if (UAnimComposite::StaticClass() == InItem->GetAnimationAssetStaticClass())
-	{
-		FPoseSearchDatabaseAnimComposite::StaticStruct()->ExportText(CopyObj->Content, InItem, nullptr, nullptr, PPF_None, nullptr);
-	}
-	else if (UAnimMontage::StaticClass() == InItem->GetAnimationAssetStaticClass())
-	{
-		FPoseSearchDatabaseAnimMontage::StaticStruct()->ExportText(CopyObj->Content, InItem, nullptr, nullptr, PPF_None, nullptr);
-	}
-	else if (UBlendSpace::StaticClass() == InItem->GetAnimationAssetStaticClass())
-	{
-		FPoseSearchDatabaseBlendSpace::StaticStruct()->ExportText(CopyObj->Content, InItem, nullptr, nullptr, PPF_None, nullptr);
-	}
-	else if (UMultiAnimAsset::StaticClass() == InItem->GetAnimationAssetStaticClass())
-	{
-		FPoseSearchDatabaseMultiAnimAsset::StaticStruct()->ExportText(CopyObj->Content, InItem, nullptr, nullptr, PPF_None, nullptr);
-	}
+	UClass* AnimationAssetStaticClass = InItem->GetAnimationAssetStaticClass();
+	check(AnimationAssetStaticClass);
+	CopyObj->ClassName = AnimationAssetStaticClass->GetFName();
+	FPoseSearchDatabaseAnimationAsset::StaticStruct()->ExportText(CopyObj->Content, InItem, nullptr, nullptr, PPF_None, nullptr);
+
 
 	DatabaseItems.Add(CopyObj);
 }
@@ -70,41 +53,9 @@ void UPoseSearchDatabaseEditorClipboardContent::PasteToDatabase(UPoseSearchDatab
 	
 	for (TObjectPtr<UPoseSearchDatabaseItemCopyObject> Item : DatabaseItems)
 	{
-		if (Item->ClassName == UAnimSequence::StaticClass()->GetFName())
-		{
-			FPoseSearchDatabaseSequence NewAsset;
-			FPoseSearchDatabaseSequence::StaticStruct()->ImportText(*Item->Content, &NewAsset, nullptr, PPF_None, GLog, FPoseSearchDatabaseSequence::StaticStruct()->GetName());
-
-			InTargetDatabase->AddAnimationAsset(FInstancedStruct::Make(NewAsset));
-		}
-		else if (Item->ClassName == UAnimComposite::StaticClass()->GetFName())
-		{
-			FPoseSearchDatabaseAnimComposite NewAsset;
-			FPoseSearchDatabaseAnimComposite::StaticStruct()->ImportText(*Item->Content, &NewAsset, nullptr, PPF_None, GLog, FPoseSearchDatabaseAnimComposite::StaticStruct()->GetName());
-
-			InTargetDatabase->AddAnimationAsset(FInstancedStruct::Make(NewAsset));
-		}
-		else if (Item->ClassName == UBlendSpace::StaticClass()->GetFName())
-		{
-			FPoseSearchDatabaseBlendSpace NewAsset;
-			FPoseSearchDatabaseBlendSpace::StaticStruct()->ImportText(*Item->Content, &NewAsset, nullptr, PPF_None, GLog, FPoseSearchDatabaseBlendSpace::StaticStruct()->GetName());
-
-			InTargetDatabase->AddAnimationAsset(FInstancedStruct::Make(NewAsset));
-		}
-		else if (Item->ClassName == UAnimMontage::StaticClass()->GetFName())
-		{
-			FPoseSearchDatabaseAnimMontage NewAsset;
-			FPoseSearchDatabaseAnimMontage::StaticStruct()->ImportText(*Item->Content, &NewAsset, nullptr, PPF_None, GLog, FPoseSearchDatabaseAnimMontage::StaticStruct()->GetName());
-
-			InTargetDatabase->AddAnimationAsset(FInstancedStruct::Make(NewAsset));
-		}
-		else if (Item->ClassName == UMultiAnimAsset::StaticClass()->GetFName())
-		{
-			FPoseSearchDatabaseMultiAnimAsset NewAsset;
-			FPoseSearchDatabaseMultiAnimAsset::StaticStruct()->ImportText(*Item->Content, &NewAsset, nullptr, PPF_None, GLog, FPoseSearchDatabaseMultiAnimAsset::StaticStruct()->GetName());
-				
-			InTargetDatabase->AddAnimationAsset(FInstancedStruct::Make(NewAsset));
-		}
+		FPoseSearchDatabaseAnimationAsset NewAsset;
+		FPoseSearchDatabaseAnimationAsset::StaticStruct()->ImportText(*Item->Content, &NewAsset, nullptr, PPF_None, GLog, FPoseSearchDatabaseAnimationAsset::StaticStruct()->GetName());
+		InTargetDatabase->AddAnimationAsset(NewAsset);
 	}
 }
 
diff --git a/Engine/Plugins/Animation/PoseSearch/Source/Editor/Private/PoseSearchDatabaseEditorReflection.cpp b/Engine/Plugins/Animation/PoseSearch/Source/Editor/Private/PoseSearchDatabaseEditorReflection.cpp
--- a/Engine/Plugins/Animation/PoseSearch/Source/Editor/Private/PoseSearchDatabaseEditorReflection.cpp
+++ b/Engine/Plugins/Animation/PoseSearch/Source/Editor/Private/PoseSearchDatabaseEditorReflection.cpp
@@ -1,12 +1,13 @@
 // Copyright Epic Games, Inc. All Rights Reserved.
 
 #include "PoseSearchDatabaseEditorReflection.h"
-#include "StructUtils/InstancedStruct.h"
+#include "Animation/AnimSequenceBase.h"
+#include "PoseSearch/PoseSearchDerivedData.h"
+#include "PoseSearch/PoseSearchSchema.h"
 #include "PoseSearchDatabaseAssetTree.h"
 #include "PoseSearchDatabaseAssetTreeNode.h"
 #include "PoseSearchDatabaseViewModel.h"
-#include "PoseSearch/PoseSearchDerivedData.h"
-#include "PoseSearch/PoseSearchSchema.h"
+#include "StructUtils/InstancedStruct.h"
 
 #include UE_INLINE_GENERATED_CPP_BY_NAME(PoseSearchDatabaseEditorReflection)
 
@@ -50,125 +51,21 @@ void UPoseSearchDatabaseReflectionBase::PostEditChangeProperty(FPropertyChangedE
 	}
 }
 
-bool UPoseSearchDatabaseSequenceReflection::ApplyChanges() 
-{
-	const float PlayLength = Sequence.GetPlayLength(FVector::ZeroVector);
-	const FFloatInterval ClampedSamplingRange(FMath::Clamp(Sequence.SamplingRange.Min, 0.0f, PlayLength), FMath::Clamp(Sequence.SamplingRange.Max, 0.0f, PlayLength));
-	if (ClampedSamplingRange != Sequence.SamplingRange)
-	{
-		Sequence.SamplingRange = ClampedSamplingRange;
-	}
-
-	if (const TSharedPtr<UE::PoseSearch::FDatabaseAssetTreeNode> AssetTreeNode = WeakAssetTreeNode.Pin())
-	{
-		if (const TSharedPtr<UE::PoseSearch::FDatabaseViewModel> ViewModel = AssetTreeNode->EditorViewModel.Pin())
-		{
-			UPoseSearchDatabase* Database = ViewModel->GetPoseSearchDatabase();
-			if (IsValid(Database))
-			{
-				if (FPoseSearchDatabaseSequence* DatabaseSequence = Database->GetMutableDatabaseAnimationAsset<FPoseSearchDatabaseSequence>(AssetTreeNode->SourceAssetIdx))
-				{
-					if (*DatabaseSequence != Sequence)
-					{
-						*DatabaseSequence = Sequence;
-						return true;
-					}
-				}
-			}
-		}
-	}
-
-	return false;
-}
-
-bool UPoseSearchDatabaseBlendSpaceReflection::ApplyChanges()
-{
-	if (const TSharedPtr<UE::PoseSearch::FDatabaseAssetTreeNode> AssetTreeNode = WeakAssetTreeNode.Pin())
-	{
-		if (const TSharedPtr<UE::PoseSearch::FDatabaseViewModel> ViewModel = AssetTreeNode->EditorViewModel.Pin())
-		{
-			UPoseSearchDatabase* Database = ViewModel->GetPoseSearchDatabase();
-			if (IsValid(Database))
-			{
-				if (FPoseSearchDatabaseBlendSpace* DatabaseBlendSpace = Database->GetMutableDatabaseAnimationAsset<FPoseSearchDatabaseBlendSpace>(AssetTreeNode->SourceAssetIdx))
-				{
-					if (*DatabaseBlendSpace != BlendSpace)
-					{
-						*DatabaseBlendSpace = BlendSpace;
-						return true;
-					}
-				}
-			}
-		}
-	}
-
-	return false;
-}
-
-bool UPoseSearchDatabaseAnimCompositeReflection::ApplyChanges()
-{
-	const float PlayLength = AnimComposite.GetPlayLength(FVector::ZeroVector);
-	const FFloatInterval ClampedSamplingRange(FMath::Clamp(AnimComposite.SamplingRange.Min, 0.0f, PlayLength), FMath::Clamp(AnimComposite.SamplingRange.Max, 0.0f, PlayLength));
-	if (ClampedSamplingRange != AnimComposite.SamplingRange)
-	{
-		AnimComposite.SamplingRange = ClampedSamplingRange;
-	}
-
-	if (const TSharedPtr<UE::PoseSearch::FDatabaseAssetTreeNode> AssetTreeNode = WeakAssetTreeNode.Pin())
-	{
-		if (const TSharedPtr<UE::PoseSearch::FDatabaseViewModel> ViewModel = AssetTreeNode->EditorViewModel.Pin())
-		{
-			UPoseSearchDatabase* Database = ViewModel->GetPoseSearchDatabase();
-			if (IsValid(Database))
-			{
-				if (FPoseSearchDatabaseAnimComposite* DatabaseAnimComposite = Database->GetMutableDatabaseAnimationAsset<FPoseSearchDatabaseAnimComposite>(AssetTreeNode->SourceAssetIdx))
-				{
-					if (*DatabaseAnimComposite != AnimComposite)
-					{
-						*DatabaseAnimComposite = AnimComposite;
-						return true;
-					}
-				}
-			}
-		}
-	}
-
-	return false;
-}
-
-bool UPoseSearchDatabaseAnimMontageReflection::ApplyChanges()
+bool UPoseSearchDatabaseAnimationAssetReflection::ApplyChanges() 
 {
-	const float PlayLength = AnimMontage.GetPlayLength(FVector::ZeroVector);
-	const FFloatInterval ClampedSamplingRange(FMath::Clamp(AnimMontage.SamplingRange.Min, 0.0f, PlayLength), FMath::Clamp(AnimMontage.SamplingRange.Max, 0.0f, PlayLength));
-	if (ClampedSamplingRange != AnimMontage.SamplingRange)
+	if (UObject* AnimationAsset = DatabaseAnimationAsset.GetAnimationAsset())
 	{
-		AnimMontage.SamplingRange = ClampedSamplingRange;
-	}
-
-	if (const TSharedPtr<UE::PoseSearch::FDatabaseAssetTreeNode> AssetTreeNode = WeakAssetTreeNode.Pin())
-	{
-		if (const TSharedPtr<UE::PoseSearch::FDatabaseViewModel> ViewModel = AssetTreeNode->EditorViewModel.Pin())
+		if (AnimationAsset->IsA<UAnimSequenceBase>())
 		{
-			UPoseSearchDatabase* Database = ViewModel->GetPoseSearchDatabase();
-			if (IsValid(Database))
+			const float PlayLength = DatabaseAnimationAsset.GetPlayLength(FVector::ZeroVector);
+			const FFloatInterval ClampedSamplingRange(FMath::Clamp(DatabaseAnimationAsset.SamplingRange.Min, 0.0f, PlayLength), FMath::Clamp(DatabaseAnimationAsset.SamplingRange.Max, 0.0f, PlayLength));
+			if (ClampedSamplingRange != DatabaseAnimationAsset.SamplingRange)
 			{
-				if (FPoseSearchDatabaseAnimMontage* DatabaseAnimMontage = Database->GetMutableDatabaseAnimationAsset<FPoseSearchDatabaseAnimMontage>(AssetTreeNode->SourceAssetIdx))
-				{
-					if (*DatabaseAnimMontage != AnimMontage)
-					{
-						*DatabaseAnimMontage = AnimMontage;
-						return true;
-					}
-				}
+				DatabaseAnimationAsset.SamplingRange = ClampedSamplingRange;
 			}
 		}
 	}
 
-	return false;
-}
-
-bool UPoseSearchDatabaseMultiAnimAssetReflection::ApplyChanges()
-{
 	if (const TSharedPtr<UE::PoseSearch::FDatabaseAssetTreeNode> AssetTreeNode = WeakAssetTreeNode.Pin())
 	{
 		if (const TSharedPtr<UE::PoseSearch::FDatabaseViewModel> ViewModel = AssetTreeNode->EditorViewModel.Pin())
@@ -176,11 +73,11 @@ bool UPoseSearchDatabaseMultiAnimAssetReflection::ApplyChanges()
 			UPoseSearchDatabase* Database = ViewModel->GetPoseSearchDatabase();
 			if (IsValid(Database))
 			{
-				if (FPoseSearchDatabaseMultiAnimAsset* DatabaseMultiAnimAsset = Database->GetMutableDatabaseAnimationAsset<FPoseSearchDatabaseMultiAnimAsset>(AssetTreeNode->SourceAssetIdx))
+				if (FPoseSearchDatabaseAnimationAsset* DatabaseAnimationAssetPtr = Database->GetMutableDatabaseAnimationAsset(AssetTreeNode->SourceAssetIdx))
 				{
-					if (*DatabaseMultiAnimAsset != MultiAnimAsset)
+					if (*DatabaseAnimationAssetPtr != DatabaseAnimationAsset)
 					{
-						*DatabaseMultiAnimAsset = MultiAnimAsset;
+						*DatabaseAnimationAssetPtr = DatabaseAnimationAsset;
 						return true;
 					}
 				}
@@ -261,7 +158,7 @@ void UPoseSearchDatabaseStatistics::Initialize(const UPoseSearchDatabase* PoseSe
 			Analyzed.Reserve(NumAnimationAssets);
 			for (int32 AnimationAssetIndex = 0; AnimationAssetIndex < NumAnimationAssets; ++AnimationAssetIndex)
 			{
-				if (const FPoseSearchDatabaseAnimationAssetBase* DatabaseAnimationAssetBase = PoseSearchDatabase->GetDatabaseAnimationAsset<FPoseSearchDatabaseAnimationAssetBase>(AnimationAssetIndex))
+				if (const FPoseSearchDatabaseAnimationAssetBase* DatabaseAnimationAssetBase = PoseSearchDatabase->GetDatabaseAnimationAsset(AnimationAssetIndex))
 				{
 					bool bAlreadyAnalyzed = false;
 					Analyzed.Add(DatabaseAnimationAssetBase->GetAnimationAsset(), &bAlreadyAnalyzed);
diff --git a/Engine/Plugins/Animation/PoseSearch/Source/Editor/Private/PoseSearchDatabaseEditorUtils.cpp b/Engine/Plugins/Animation/PoseSearch/Source/Editor/Private/PoseSearchDatabaseEditorUtils.cpp
--- a/Engine/Plugins/Animation/PoseSearch/Source/Editor/Private/PoseSearchDatabaseEditorUtils.cpp
+++ b/Engine/Plugins/Animation/PoseSearch/Source/Editor/Private/PoseSearchDatabaseEditorUtils.cpp
@@ -22,9 +22,15 @@ namespace UE::PoseSearch
 			return false;
 		}
 
+		const UClass* InAssetDataClass = InAssetData.GetClass();
+		if (!InAssetDataClass)
+		{
+			return false;
+		}
+
 		if (RoledSkeletons.Num() == 1)
 		{
-			if (RoledSkeletons[0].Skeleton && InAssetData.GetClass()->IsChildOf(UAnimationAsset::StaticClass()) &&
+			if (RoledSkeletons[0].Skeleton && InAssetDataClass->IsChildOf(UAnimationAsset::StaticClass()) &&
 				RoledSkeletons[0].Skeleton->IsCompatibleForEditor(InAssetData))
 			{
 				// We found a compatible skeleton in the schema.
@@ -34,7 +40,7 @@ namespace UE::PoseSearch
 			return false;
 		}
 				
-		if (!InAssetData.GetClass()->IsChildOf(UMultiAnimAsset::StaticClass()))
+		if (!InAssetDataClass->IsChildOf(UMultiAnimAsset::StaticClass()))
 		{
 			return false;
 		}
diff --git a/Engine/Plugins/Animation/PoseSearch/Source/Editor/Private/PoseSearchDatabaseViewModel.cpp b/Engine/Plugins/Animation/PoseSearch/Source/Editor/Private/PoseSearchDatabaseViewModel.cpp
--- a/Engine/Plugins/Animation/PoseSearch/Source/Editor/Private/PoseSearchDatabaseViewModel.cpp
+++ b/Engine/Plugins/Animation/PoseSearch/Source/Editor/Private/PoseSearchDatabaseViewModel.cpp
@@ -48,7 +48,7 @@ bool FDatabasePreviewActor::SpawnPreviewActor(UWorld* World, const UPoseSearchDa
 	check(PoseSearchDatabase && PoseSearchDatabase->Schema);
 	const FSearchIndex& SearchIndex = PoseSearchDatabase->GetSearchIndex();
 	const FSearchIndexAsset& IndexAsset = SearchIndex.Assets[IndexAssetIdx];
-	const FPoseSearchDatabaseAnimationAssetBase* DatabaseAnimationAsset = PoseSearchDatabase->GetDatabaseAnimationAsset<FPoseSearchDatabaseAnimationAssetBase>(IndexAsset.GetSourceAssetIdx());
+	const FPoseSearchDatabaseAnimationAssetBase* DatabaseAnimationAsset = PoseSearchDatabase->GetDatabaseAnimationAsset(IndexAsset.GetSourceAssetIdx());
 	check(DatabaseAnimationAsset);
 		
 	USkeleton* Skeleton = PoseSearchDatabase->Schema->GetSkeleton(Role);
@@ -189,7 +189,7 @@ void FDatabasePreviewActor::UpdatePreviewActor(const UPoseSearchDatabase* PoseSe
 	}
 
 	const FSearchIndexAsset& IndexAsset = SearchIndex.Assets[IndexAssetIndex];
-	const FPoseSearchDatabaseAnimationAssetBase* DatabaseAnimationAsset = PoseSearchDatabase->GetDatabaseAnimationAsset<FPoseSearchDatabaseAnimationAssetBase>(IndexAsset.GetSourceAssetIdx());
+	const FPoseSearchDatabaseAnimationAssetBase* DatabaseAnimationAsset = PoseSearchDatabase->GetDatabaseAnimationAsset(IndexAsset.GetSourceAssetIdx());
 
 	USkeletalMesh* PreviewMesh = DatabaseAnimationAsset->GetPreviewMeshForRole(ActorRole);
 	if (!PreviewMesh)
@@ -639,63 +639,15 @@ void FDatabaseViewModel::RemovePreviewActors()
 	PreviewActors.Reset();
 }
 
-void FDatabaseViewModel::AddSequenceToDatabase(UAnimSequence* AnimSequence)
+void FDatabaseViewModel::AddAnimationAssetToDatabase(UObject* AnimationAsset)
 {
 	if (UPoseSearchDatabase* Database = GetPoseSearchDatabase())
 	{
 		Database->Modify();
 		
-		FPoseSearchDatabaseSequence NewAsset;
-		NewAsset.Sequence = AnimSequence;
-		Database->AddAnimationAsset(FInstancedStruct::Make(NewAsset));
-	}
-}
-
-void FDatabaseViewModel::AddBlendSpaceToDatabase(UBlendSpace* BlendSpace)
-{
-	if (UPoseSearchDatabase* Database = GetPoseSearchDatabase())
-	{
-		Database->Modify();
-		
-		FPoseSearchDatabaseBlendSpace NewAsset;
-		NewAsset.BlendSpace = BlendSpace;
-		Database->AddAnimationAsset(FInstancedStruct::Make(NewAsset));
-	}
-}
-
-void FDatabaseViewModel::AddAnimCompositeToDatabase(UAnimComposite* AnimComposite)
-{
-	if (UPoseSearchDatabase* Database = GetPoseSearchDatabase())
-	{
-		Database->Modify();
-
-		FPoseSearchDatabaseAnimComposite NewAsset;
-		NewAsset.AnimComposite = AnimComposite;
-		Database->AddAnimationAsset(FInstancedStruct::Make(NewAsset));
-	}
-}
-
-void FDatabaseViewModel::AddAnimMontageToDatabase(UAnimMontage* AnimMontage)
-{
-	if (UPoseSearchDatabase* Database = GetPoseSearchDatabase())
-	{
-		Database->Modify();
-		
-		FPoseSearchDatabaseAnimMontage NewAsset;
-		NewAsset.AnimMontage = AnimMontage;
-		Database->AddAnimationAsset(FInstancedStruct::Make(NewAsset));
-	}
-}
-
-void FDatabaseViewModel::AddMultiAnimAssetToDatabase(UMultiAnimAsset* MultiAnimAsset)
-{
-	if (UPoseSearchDatabase* Database = GetPoseSearchDatabase())
-	{
-		Database->Modify();
-		
-		FPoseSearchDatabaseMultiAnimAsset NewAsset;
-		NewAsset.MultiAnimAsset = MultiAnimAsset;
-		Database->AddAnimationAsset(FInstancedStruct::Make(NewAsset));
+		FPoseSearchDatabaseAnimationAsset NewAsset;
+		NewAsset.AnimAsset = AnimationAsset;
+		Database->AddAnimationAsset(NewAsset);
 	}
 }
 
@@ -703,7 +655,7 @@ bool FDatabaseViewModel::DeleteFromDatabase(int32 AnimationAssetIndex)
 {
 	if (UPoseSearchDatabase* Database = GetPoseSearchDatabase())
 	{
-		if (const FPoseSearchDatabaseAnimationAssetBase* DatabaseAnimationAssetBase = Database->GetDatabaseAnimationAsset<FPoseSearchDatabaseAnimationAssetBase>(AnimationAssetIndex))
+		if (const FPoseSearchDatabaseAnimationAssetBase* DatabaseAnimationAssetBase = Database->GetDatabaseAnimationAsset(AnimationAssetIndex))
 		{
 			Database->Modify();
 			
@@ -757,7 +709,7 @@ void FDatabaseViewModel::SetDisableReselection(int32 AnimationAssetIndex, bool b
 {
 	if (UPoseSearchDatabase* Database = GetPoseSearchDatabase())
 	{
-		if (FPoseSearchDatabaseAnimationAssetBase* DatabaseAnimationAsset = Database->GetMutableDatabaseAnimationAsset<FPoseSearchDatabaseAnimationAssetBase>(AnimationAssetIndex))
+		if (FPoseSearchDatabaseAnimationAssetBase* DatabaseAnimationAsset = Database->GetMutableDatabaseAnimationAsset(AnimationAssetIndex))
 		{
 			Database->Modify();
 			
@@ -770,7 +722,7 @@ bool FDatabaseViewModel::IsDisableReselection(int32 AnimationAssetIndex) const
 {
 	if (const UPoseSearchDatabase* Database = GetPoseSearchDatabase())
 	{
-		if (const FPoseSearchDatabaseAnimationAssetBase* DatabaseAnimationAsset = Database->GetDatabaseAnimationAsset<FPoseSearchDatabaseAnimationAssetBase>(AnimationAssetIndex))
+		if (const FPoseSearchDatabaseAnimationAssetBase* DatabaseAnimationAsset = Database->GetDatabaseAnimationAsset(AnimationAssetIndex))
 		{
 			return DatabaseAnimationAsset->IsDisableReselection();
 		}
@@ -783,7 +735,7 @@ void FDatabaseViewModel::SetIsEnabled(int32 AnimationAssetIndex, bool bEnabled)
 {
 	if (UPoseSearchDatabase* Database = GetPoseSearchDatabase())
 	{
-		if (FPoseSearchDatabaseAnimationAssetBase* DatabaseAnimationAsset = Database->GetMutableDatabaseAnimationAsset<FPoseSearchDatabaseAnimationAssetBase>(AnimationAssetIndex))
+		if (FPoseSearchDatabaseAnimationAssetBase* DatabaseAnimationAsset = Database->GetMutableDatabaseAnimationAsset(AnimationAssetIndex))
 		{
 			Database->Modify();
 
@@ -796,7 +748,7 @@ bool FDatabaseViewModel::IsEnabled(int32 AnimationAssetIndex) const
 {
 	if (const UPoseSearchDatabase* Database = GetPoseSearchDatabase())
 	{
-		if (const FPoseSearchDatabaseAnimationAssetBase* DatabaseAnimationAsset = Database->GetDatabaseAnimationAsset<FPoseSearchDatabaseAnimationAssetBase>(AnimationAssetIndex))
+		if (const FPoseSearchDatabaseAnimationAssetBase* DatabaseAnimationAsset = Database->GetDatabaseAnimationAsset(AnimationAssetIndex))
 		{
 			return DatabaseAnimationAsset->IsEnabled();
 		}
@@ -811,52 +763,11 @@ bool FDatabaseViewModel::SetAnimationAsset(int32 AnimationAssetIndex, UObject* A
 	{
 		if (UPoseSearchDatabase* Database = GetPoseSearchDatabase())
 		{
-			if (FPoseSearchDatabaseAnimationAssetBase* DatabaseAnimationAsset = Database->GetMutableDatabaseAnimationAsset<FPoseSearchDatabaseAnimationAssetBase>(AnimationAssetIndex))
+			if (FPoseSearchDatabaseAnimationAsset* DatabaseAnimationAsset = Database->GetMutableDatabaseAnimationAsset(AnimationAssetIndex))
 			{
-				// Ensure that our target database item matches the input object's class.
-				const UClass* AssetClass = AnimAsset->GetClass();
-				if (AssetClass->IsChildOf(DatabaseAnimationAsset->GetAnimationAssetStaticClass()))
-				{
-					if (AssetClass->IsChildOf( UAnimSequence::StaticClass()))
-					{
-						Database->Modify();
-						FPoseSearchDatabaseSequence* DatabaseSequenceAsset = static_cast<FPoseSearchDatabaseSequence*>(DatabaseAnimationAsset);
-						DatabaseSequenceAsset->Sequence = Cast<UAnimSequence>(AnimAsset);
-						return true;
-					}
-					
-					if (AssetClass->IsChildOf(UAnimComposite::StaticClass()))
-					{
-						Database->Modify();
-						FPoseSearchDatabaseAnimComposite* DatabaseCompositeAsset = static_cast<FPoseSearchDatabaseAnimComposite*>(DatabaseAnimationAsset);
-						DatabaseCompositeAsset->AnimComposite = Cast<UAnimComposite>(AnimAsset);
-						return true;
-					}
-					
-					if (AssetClass->IsChildOf(UAnimMontage::StaticClass()))
-					{
-						Database->Modify();
-						FPoseSearchDatabaseAnimMontage* DatabaseMontageAsset = static_cast<FPoseSearchDatabaseAnimMontage*>(DatabaseAnimationAsset);
-						DatabaseMontageAsset->AnimMontage = Cast<UAnimMontage>(AnimAsset);
-						return true;
-					}
-					
-					if (AssetClass->IsChildOf(UBlendSpace::StaticClass()))
-					{
-						Database->Modify();
-						FPoseSearchDatabaseBlendSpace* DatabaseBlendSpaceAsset = static_cast<FPoseSearchDatabaseBlendSpace*>(DatabaseAnimationAsset);
-						DatabaseBlendSpaceAsset->BlendSpace = Cast<UBlendSpace>(AnimAsset);
-						return true;
-					}
-					
-					if (AssetClass->IsChildOf(UMultiAnimAsset::StaticClass()))
-					{
-						Database->Modify();
-						FPoseSearchDatabaseMultiAnimAsset* DatabaseMultiAnimAssetAsset = static_cast<FPoseSearchDatabaseMultiAnimAsset*>(DatabaseAnimationAsset);
-						DatabaseMultiAnimAssetAsset->MultiAnimAsset = Cast<UMultiAnimAsset>(AnimAsset);
-						return true;
-					}
-				}
+				Database->Modify();
+				DatabaseAnimationAsset->AnimAsset = AnimAsset;
+				return true;
 			}
 		}
 	}
@@ -868,7 +779,7 @@ void FDatabaseViewModel::SetMirrorOption(int32 AnimationAssetIndex, EPoseSearchM
 {
 	if (UPoseSearchDatabase* Database = GetPoseSearchDatabase())
 	{
-		if (FPoseSearchDatabaseAnimationAssetBase* DatabaseAnimationAsset = Database->GetMutableDatabaseAnimationAsset<FPoseSearchDatabaseAnimationAssetBase>(AnimationAssetIndex))
+		if (FPoseSearchDatabaseAnimationAssetBase* DatabaseAnimationAsset = Database->GetMutableDatabaseAnimationAsset(AnimationAssetIndex))
 		{
 			Database->Modify();
 
@@ -881,7 +792,7 @@ EPoseSearchMirrorOption FDatabaseViewModel::GetMirrorOption(int32 AnimationAsset
 {
 	if (const UPoseSearchDatabase* Database = GetPoseSearchDatabase())
 	{
-		if (const FPoseSearchDatabaseAnimationAssetBase* DatabaseAnimationAsset = Database->GetDatabaseAnimationAsset<FPoseSearchDatabaseAnimationAssetBase>(AnimationAssetIndex))
+		if (const FPoseSearchDatabaseAnimationAssetBase* DatabaseAnimationAsset = Database->GetDatabaseAnimationAsset(AnimationAssetIndex))
 		{
 			return DatabaseAnimationAsset->MirrorOption;
 		}
@@ -914,7 +825,7 @@ int32 FDatabaseViewModel::SetSelectedNode(int32 PoseIdx, bool bClearSelection, b
 				if (SearchIndex.Assets.IsValidIndex(IndexAssetIndex))
 				{
 					const FSearchIndexAsset& IndexAsset = SearchIndex.Assets[IndexAssetIndex];
-					const FPoseSearchDatabaseAnimationAssetBase* DatabaseAnimationAsset = Database->GetDatabaseAnimationAsset<FPoseSearchDatabaseAnimationAssetBase>(IndexAsset.GetSourceAssetIdx());
+					const FPoseSearchDatabaseAnimationAssetBase* DatabaseAnimationAsset = Database->GetDatabaseAnimationAsset(IndexAsset.GetSourceAssetIdx());
 					check(DatabaseAnimationAsset);
 					int32 PreviewActorGroupIndex = INDEX_NONE;
 					for (int32 RoleIndex = 0; RoleIndex < DatabaseAnimationAsset->GetNumRoles(); ++RoleIndex)
@@ -977,7 +888,7 @@ void FDatabaseViewModel::SetSelectedNodes(const TArrayView<TSharedPtr<FDatabaseA
 				const FSearchIndexAsset& IndexAsset = SearchIndex.Assets[IndexAssetIndex];
 				if (AssociatedAssetIndices.Find(IndexAsset.GetSourceAssetIdx()))
 				{
-					const FPoseSearchDatabaseAnimationAssetBase* DatabaseAnimationAsset = Database->GetDatabaseAnimationAsset<FPoseSearchDatabaseAnimationAssetBase>(IndexAsset.GetSourceAssetIdx());
+					const FPoseSearchDatabaseAnimationAssetBase* DatabaseAnimationAsset = Database->GetDatabaseAnimationAsset(IndexAsset.GetSourceAssetIdx());
 					check(DatabaseAnimationAsset);
 					int32 PreviewActorGroupIndex = INDEX_NONE;
 					for (int32 RoleIndex = 0; RoleIndex < DatabaseAnimationAsset->GetNumRoles(); ++RoleIndex)
@@ -1118,10 +1029,12 @@ bool FDatabaseViewModel::GetAnimationTime(int32 SourceAssetIdx, float& CurrentPl
 					CurrentPlayTime = PlayTime + IndexAsset.GetFirstSampleTime(Database->Schema->SampleRate);
 					BlendParameters = IndexAsset.GetBlendParameters();
 
-					const bool bIsBlendSpace = Database->GetDatabaseAnimationAsset<FPoseSearchDatabaseBlendSpace>(IndexAsset) != nullptr;
-					if (bIsBlendSpace && !FMath::IsNearlyEqual(MaxPreviewPlayLength, MinPreviewPlayLength))
+					if (const UObject* AnimationAsset = Database->GetDatabaseAnimationAsset(IndexAsset)->GetAnimationAsset())
 					{
-						CurrentPlayTime = (CurrentPlayTime - MaxPreviewPlayLength) / (MaxPreviewPlayLength - MinPreviewPlayLength);
+						if (AnimationAsset->IsA<UBlendSpace>() && !FMath::IsNearlyEqual(MaxPreviewPlayLength, MinPreviewPlayLength))
+						{
+							CurrentPlayTime = (CurrentPlayTime - MaxPreviewPlayLength) / (MaxPreviewPlayLength - MinPreviewPlayLength);
+						}
 					}
 					return true;
 				}
diff --git a/Engine/Plugins/Animation/PoseSearch/Source/Editor/Private/PoseSearchDebuggerDatabaseView.cpp b/Engine/Plugins/Animation/PoseSearch/Source/Editor/Private/PoseSearchDebuggerDatabaseView.cpp
--- a/Engine/Plugins/Animation/PoseSearch/Source/Editor/Private/PoseSearchDebuggerDatabaseView.cpp
+++ b/Engine/Plugins/Animation/PoseSearch/Source/Editor/Private/PoseSearchDebuggerDatabaseView.cpp
@@ -212,7 +212,7 @@ static void AddUnfilteredDatabaseRow(const UPoseSearchDatabase* Database, TConst
 			}
 		}
 
-		if (const FPoseSearchDatabaseAnimationAssetBase* DatabaseAsset = Database->GetDatabaseAnimationAsset<FPoseSearchDatabaseAnimationAssetBase>(*SearchIndexAsset))
+		if (const FPoseSearchDatabaseAnimationAssetBase* DatabaseAsset = Database->GetDatabaseAnimationAsset(*SearchIndexAsset))
 		{
 			const FVector BlendParameters = SearchIndexAsset->GetBlendParameters();
 			const float PlayLength = DatabaseAsset->GetPlayLength(BlendParameters);
diff --git a/Engine/Plugins/Animation/PoseSearch/Source/Editor/Private/PoseSearchEditor.cpp b/Engine/Plugins/Animation/PoseSearch/Source/Editor/Private/PoseSearchEditor.cpp
--- a/Engine/Plugins/Animation/PoseSearch/Source/Editor/Private/PoseSearchEditor.cpp
+++ b/Engine/Plugins/Animation/PoseSearch/Source/Editor/Private/PoseSearchEditor.cpp
@@ -132,7 +132,7 @@ void FEditorModule::RegisterCustomPropertyTypeLayout(FName PropertyTypeName, FOn
 
 void FEditorModule::RegisterPropertyTypeCustomizations()
 {
-	RegisterCustomPropertyTypeLayout("PoseSearchDatabaseSequence", FOnGetPropertyTypeCustomizationInstance::CreateStatic(&FPoseSearchDatabaseSequenceCustomization::MakeInstance));
+	RegisterCustomPropertyTypeLayout("PoseSearchDatabaseAnimationAsset", FOnGetPropertyTypeCustomizationInstance::CreateStatic(&FPoseSearchDatabaseAnimationAssetCustomization::MakeInstance));
 }
 
 void FEditorModule::UnregisterCustomizations()
diff --git a/Engine/Plugins/Animation/PoseSearch/Source/Editor/Public/PoseSearchCustomization.h b/Engine/Plugins/Animation/PoseSearch/Source/Editor/Public/PoseSearchCustomization.h
--- a/Engine/Plugins/Animation/PoseSearch/Source/Editor/Public/PoseSearchCustomization.h
+++ b/Engine/Plugins/Animation/PoseSearch/Source/Editor/Public/PoseSearchCustomization.h
@@ -13,18 +13,15 @@
 class IPropertyHandle;
 class IDetailLayoutBuilder;
 
-/**
- * Customizes FPoseSearchDatabaseSequence to show the name of the sequence and the groups in the header
- */
-class FPoseSearchDatabaseSequenceCustomization : public IPropertyTypeCustomization
+class FPoseSearchDatabaseAnimationAssetCustomization : public IPropertyTypeCustomization
 {
 public:
 	static TSharedRef<IPropertyTypeCustomization> MakeInstance()
 	{
-		return MakeShareable(new FPoseSearchDatabaseSequenceCustomization);
+		return MakeShareable(new FPoseSearchDatabaseAnimationAssetCustomization);
 	}
 
-	/** IPropertyTypeCustomization interface */
+	// IPropertyTypeCustomization interface
 	virtual void CustomizeHeader(TSharedRef<class IPropertyHandle> InStructPropertyHandle, class FDetailWidgetRow& HeaderRow, IPropertyTypeCustomizationUtils& StructCustomizationUtils) override;
 	virtual void CustomizeChildren(TSharedRef<class IPropertyHandle> InStructPropertyHandle, class IDetailChildrenBuilder& StructBuilder, IPropertyTypeCustomizationUtils& StructCustomizationUtils) override;
 };
diff --git a/Engine/Plugins/Animation/PoseSearch/Source/Editor/Public/PoseSearchDatabaseAssetListItem.h b/Engine/Plugins/Animation/PoseSearch/Source/Editor/Public/PoseSearchDatabaseAssetListItem.h
--- a/Engine/Plugins/Animation/PoseSearch/Source/Editor/Public/PoseSearchDatabaseAssetListItem.h
+++ b/Engine/Plugins/Animation/PoseSearch/Source/Editor/Public/PoseSearchDatabaseAssetListItem.h
@@ -33,12 +33,7 @@ namespace UE::PoseSearch
 		void ConstructGroupItem(const TSharedRef<STableViewBase>& OwnerTable);
 		void ConstructAssetItem(const TSharedRef<STableViewBase>& OwnerTable);
 
-		void OnAddSequence();
-		void OnAddBlendSpace();
-		void OnAddAnimComposite();
-		void OnAddAnimMontage();
 		virtual FReply OnMouseButtonDoubleClick(const FGeometry& InMyGeometry, const FPointerEvent& InMouseEvent) override;
-
 		
 		EVisibility GetSelectedActorIconVisibility() const;
 
diff --git a/Engine/Plugins/Animation/PoseSearch/Source/Editor/Public/PoseSearchDatabaseAssetTree.h b/Engine/Plugins/Animation/PoseSearch/Source/Editor/Public/PoseSearchDatabaseAssetTree.h
--- a/Engine/Plugins/Animation/PoseSearch/Source/Editor/Public/PoseSearchDatabaseAssetTree.h
+++ b/Engine/Plugins/Animation/PoseSearch/Source/Editor/Public/PoseSearchDatabaseAssetTree.h
@@ -70,7 +70,6 @@ namespace UE::PoseSearch
 			EItemDropZone DropZone,
 			TSharedPtr<FDatabaseAssetTreeNode> TargetItem);
 
-		TSharedRef<SWidget> CreateAddNewMenuWidget();
 		TSharedPtr<SWidget> CreateContextMenu();
 
 		TSharedRef<SWidget> GenerateFilterBoxWidget();
@@ -84,12 +83,7 @@ namespace UE::PoseSearch
 		//String Filter For Database View
 		FString AssetFilterString;
 
-		void OnAddSequence(bool bFinalizeChanges = true);
-		void OnAddBlendSpace(bool bFinalizeChanges = true);
-		void OnAddAnimComposite(bool bFinalizeChanges = true);
-		void OnAddAnimMontage(bool bFinalizeChanges = true);
-		void OnAddMultiAnimAsset(bool bFinalizeChanges = true);
-
+		void OnAddAsset(bool bFinalizeChanges = true);
 		void OnDeleteAsset(TSharedPtr<FDatabaseAssetTreeNode> Node, bool bFinalizeChanges = true);
 		void CreateCommandList();
 
diff --git a/Engine/Plugins/Animation/PoseSearch/Source/Editor/Public/PoseSearchDatabaseEditor.h b/Engine/Plugins/Animation/PoseSearch/Source/Editor/Public/PoseSearchDatabaseEditor.h
--- a/Engine/Plugins/Animation/PoseSearch/Source/Editor/Public/PoseSearchDatabaseEditor.h
+++ b/Engine/Plugins/Animation/PoseSearch/Source/Editor/Public/PoseSearchDatabaseEditor.h
@@ -78,7 +78,7 @@ namespace UE::PoseSearch
 			TSharedPtr<IDetailsView> DetailView;
 			TArray<TWeakObjectPtr<UObject>> SelectedReflections;
 		};
-		FSelectionWidget& FindOrAddSelectionWidget(const UScriptStruct* ScriptStructType);
+		FSelectionWidget& GetInitializedSelectionWidget();
 		
 		TSharedPtr<SDatabasePreview> PreviewWidget;
 
@@ -91,7 +91,7 @@ namespace UE::PoseSearch
 		TSharedPtr<IDetailsView> EditingAssetWidget;
 
 		TSharedPtr<SVerticalBox> DetailsContainer;
-		TMap<const UScriptStruct*, FSelectionWidget> SelectionWidgets;
+		FSelectionWidget SelectionWidget;
 
 		TSharedPtr<IDetailsView> StatisticsOverviewWidget;
 		
diff --git a/Engine/Plugins/Animation/PoseSearch/Source/Editor/Public/PoseSearchDatabaseEditorReflection.h b/Engine/Plugins/Animation/PoseSearch/Source/Editor/Public/PoseSearchDatabaseEditorReflection.h
--- a/Engine/Plugins/Animation/PoseSearch/Source/Editor/Public/PoseSearchDatabaseEditorReflection.h
+++ b/Engine/Plugins/Animation/PoseSearch/Source/Editor/Public/PoseSearchDatabaseEditorReflection.h
@@ -32,137 +32,14 @@ protected:
 	TSharedPtr<UE::PoseSearch::SDatabaseAssetTree> AssetTreeWidget;
 };
 
-USTRUCT()
-struct FPoseSearchDatabaseSequenceEx : public FPoseSearchDatabaseSequence
-{
-	GENERATED_BODY()
-
-	// Is this animation set as looping? If you want to change this you need to change it in the base Anim Sequence.
-	// Changing the sampling range will disable looping
-	UPROPERTY(VisibleAnywhere, Category="Sequence", meta = (DisplayPriority = 10))
-	bool bLooping = false;
-
-	// Does this animation have root motion enabled ? If you want to change this you need to change it in the base Anim Sequence.
-	UPROPERTY(VisibleAnywhere, Category="Sequence", meta = (DisplayPriority = 11))
-	bool bHasRootMotion = false;
-};
-
-UCLASS()
-class UPoseSearchDatabaseSequenceReflection : public UPoseSearchDatabaseReflectionBase
-{
-	GENERATED_BODY()
-
-public:
-	UPROPERTY(EditAnywhere, Category = "Selected Sequence")
-	FPoseSearchDatabaseSequenceEx Sequence;
-	
-	virtual bool ApplyChanges() override;
-};
-
-USTRUCT()
-struct FPoseSearchDatabaseBlendSpaceEx : public FPoseSearchDatabaseBlendSpace
-{
-	GENERATED_BODY()
-
-	// Is this animation set as looping? If you want to change this you need to change it in the base Anim Sequence.
-	// Changing the sampling range will disable looping
-	UPROPERTY(VisibleAnywhere, Category="Sequence", meta = (DisplayPriority = 10))
-	bool bLooping = false;
-
-	// Does this animation have root motion enabled ? If you want to change this you need to change it in the base Anim Sequence.
-	UPROPERTY(VisibleAnywhere, Category="Sequence", meta = (DisplayPriority = 11))
-	bool bHasRootMotion = false;
-};
-
-UCLASS()
-class UPoseSearchDatabaseBlendSpaceReflection : public UPoseSearchDatabaseReflectionBase
-{
-	GENERATED_BODY()
-
-public:
-	UPROPERTY(EditAnywhere, Category = "Selected Blend Space")
-	FPoseSearchDatabaseBlendSpaceEx BlendSpace;
-	
-	virtual bool ApplyChanges() override;
-};
-
-USTRUCT()
-struct FPoseSearchDatabaseAnimCompositeEx : public FPoseSearchDatabaseAnimComposite
-{
-	GENERATED_BODY()
-
-	// Is this animation set as looping? If you want to change this you need to change it in the base Anim Sequence.
-	// Changing the sampling range will disable looping
-	UPROPERTY(VisibleAnywhere, Category="Sequence", meta = (DisplayPriority = 10))
-	bool bLooping = false;
-
-	// Does this animation have root motion enabled ? If you want to change this you need to change it in the base Anim Sequence.
-	UPROPERTY(VisibleAnywhere, Category="Sequence", meta = (DisplayPriority = 11))
-	bool bHasRootMotion = false;
-};
-
-UCLASS()
-class UPoseSearchDatabaseAnimCompositeReflection : public UPoseSearchDatabaseReflectionBase
-{
-	GENERATED_BODY()
-
-public:
-	UPROPERTY(EditAnywhere, Category = "Selected Anim Composite")
-	FPoseSearchDatabaseAnimCompositeEx AnimComposite;
-	
-	virtual bool ApplyChanges() override;
-};
-
-USTRUCT()
-struct FPoseSearchDatabaseAnimMontageEx : public FPoseSearchDatabaseAnimMontage
-{
-	GENERATED_BODY()
-
-	// Is this animation set as looping? If you want to change this you need to change it in the base Anim Sequence.
-	// Changing the sampling range will disable looping
-	UPROPERTY(VisibleAnywhere, Category="Sequence", meta = (DisplayPriority = 10))
-	bool bLooping = false;
-
-	// Does this animation have root motion enabled ? If you want to change this you need to change it in the base Anim Sequence.
-	UPROPERTY(VisibleAnywhere, Category="Sequence", meta = (DisplayPriority = 11))
-	bool bHasRootMotion = false;
-};
-
-UCLASS()
-class UPoseSearchDatabaseAnimMontageReflection : public UPoseSearchDatabaseReflectionBase
-{
-	GENERATED_BODY()
-
-public:
-	UPROPERTY(EditAnywhere, Category = "Selected Anim Montage")
-	FPoseSearchDatabaseAnimMontageEx AnimMontage;
-	
-	virtual bool ApplyChanges() override;
-};
-
-USTRUCT()
-struct FPoseSearchDatabaseMultiAnimAssetEx : public FPoseSearchDatabaseMultiAnimAsset
-{
-	GENERATED_BODY()
-
-	// Is this animation set as looping? If you want to change this you need to change it in the base Anim Sequence.
-	// Changing the sampling range will disable looping
-	UPROPERTY(VisibleAnywhere, Category="Sequence", meta = (DisplayPriority = 10))
-	bool bLooping = false;
-
-	// Does this animation have root motion enabled ? If you want to change this you need to change it in the base Anim Sequence.
-	UPROPERTY(VisibleAnywhere, Category="Sequence", meta = (DisplayPriority = 11))
-	bool bHasRootMotion = false;
-};
-
 UCLASS()
-class UPoseSearchDatabaseMultiAnimAssetReflection : public UPoseSearchDatabaseReflectionBase
+class UPoseSearchDatabaseAnimationAssetReflection : public UPoseSearchDatabaseReflectionBase
 {
 	GENERATED_BODY()
 
 public:
-	UPROPERTY(EditAnywhere, Category = "Selected Sequence")
-	FPoseSearchDatabaseMultiAnimAssetEx MultiAnimAsset;
+	UPROPERTY(EditAnywhere, Category = "Selected Animation Asset")
+	FPoseSearchDatabaseAnimationAsset DatabaseAnimationAsset;
 	
 	virtual bool ApplyChanges() override;
 };
diff --git a/Engine/Plugins/Animation/PoseSearch/Source/Editor/Public/PoseSearchDatabaseViewModel.h b/Engine/Plugins/Animation/PoseSearch/Source/Editor/Public/PoseSearchDatabaseViewModel.h
--- a/Engine/Plugins/Animation/PoseSearch/Source/Editor/Public/PoseSearchDatabaseViewModel.h
+++ b/Engine/Plugins/Animation/PoseSearch/Source/Editor/Public/PoseSearchDatabaseViewModel.h
@@ -109,12 +109,7 @@ namespace UE::PoseSearch
 		void ToggleDisplayEventData() { bDisplayEventData = !bDisplayEventData; }
 		bool IsDisplayEventDataChecked() const { return bDisplayEventData; }
 
-		void AddSequenceToDatabase(UAnimSequence* AnimSequence);
-		void AddBlendSpaceToDatabase(UBlendSpace* BlendSpace);
-		void AddAnimCompositeToDatabase(UAnimComposite* AnimComposite);
-		void AddAnimMontageToDatabase(UAnimMontage* AnimMontage);
-		void AddMultiAnimAssetToDatabase(UMultiAnimAsset* MultiAnimAsset);
-
+		void AddAnimationAssetToDatabase(UObject* AnimationAsset);
 		bool DeleteFromDatabase(int32 AnimationAssetIndex);
 
 		void SetDisableReselection(int32 AnimationAssetIndex, bool bEnabled);
diff --git a/Engine/Plugins/Animation/PoseSearch/Source/Runtime/Private/Chooser/PoseSearchChooserColumn.cpp b/Engine/Plugins/Animation/PoseSearch/Source/Runtime/Private/Chooser/PoseSearchChooserColumn.cpp
--- a/Engine/Plugins/Animation/PoseSearch/Source/Runtime/Private/Chooser/PoseSearchChooserColumn.cpp
+++ b/Engine/Plugins/Animation/PoseSearch/Source/Runtime/Private/Chooser/PoseSearchChooserColumn.cpp
@@ -38,19 +38,38 @@ FPoseSearchColumn::FPoseSearchColumn()
 {
 }
 
-#if WITH_EDITOR
-void FPoseSearchColumn::AutoPopulate(int32 RowIndex, UObject* OutputObject)
+
+void FPoseSearchColumn::CheckConsistency() const
 {
-	if (FPoseSearchDatabaseSequence* Asset = InternalDatabase->GetMutableDatabaseAnimationAsset<FPoseSearchDatabaseSequence>(RowIndex))
+#if DO_CHECK && WITH_EDITORONLY_DATA
+
+	check(InternalDatabase);
+	UChooserTable* OuterChooser = Cast<UChooserTable>(InternalDatabase->GetOuter());
+	check(OuterChooser);
+
+	const int32 NumAnimationAssets = InternalDatabase->GetNumAnimationAssets();
+	check(NumAnimationAssets == OuterChooser->ResultsStructs.Num());
+
+	for (int32 AnimationAssetIndex = 0; AnimationAssetIndex < NumAnimationAssets; ++AnimationAssetIndex)
 	{
-		UAnimSequence* AnimSequence = Cast<UAnimSequence>(OutputObject);
-		if (Asset->Sequence != AnimSequence)
-		{
-			FScopedTransaction ScopedTransaction(LOCTEXT("AutoPopulate","AutoPopulate FPoseSearchColumn"));
-			InternalDatabase->Modify();
-			Asset->Sequence = AnimSequence;
-		}
+		const FPoseSearchDatabaseAnimationAsset* DatabaseAnimationAsset = InternalDatabase->GetDatabaseAnimationAsset(AnimationAssetIndex);
+		check(DatabaseAnimationAsset);
+
+		const UObject* ReferencedObject = OuterChooser->ResultsStructs[AnimationAssetIndex].Get<FObjectChooserBase>().GetReferencedObject();
+		const UObject* AnimationAsset = DatabaseAnimationAsset->GetAnimationAsset();
+		check(AnimationAsset == ReferencedObject);
 	}
+
+#endif // DO_CHECK && WITH_EDITORONLY_DATA
+}
+
+#if WITH_EDITOR
+
+UObject* FPoseSearchColumn::GetReferencedObject(int32 RowIndex) const
+{
+	UChooserTable* OuterChooser = Cast<UChooserTable>(InternalDatabase->GetOuter());
+	check(OuterChooser);
+	return OuterChooser->ResultsStructs[RowIndex].Get<FObjectChooserBase>().GetReferencedObject();
 }
 
 FName FPoseSearchColumn::RowValuesPropertyName()
@@ -60,13 +79,15 @@ FName FPoseSearchColumn::RowValuesPropertyName()
 
 void FPoseSearchColumn::SetNumRows(int32 NumRows)
 {
+	check(InternalDatabase);
+	
 	if (NumRows == InternalDatabase->GetNumAnimationAssets())
 	{
 		// nothing to do
 	}
 	else if (NumRows < InternalDatabase->GetNumAnimationAssets())
 	{
-		FScopedTransaction ScopedTransaction(LOCTEXT("SetNumRows","SetNumRows FPoseSearchColumn"));
+		FScopedTransaction ScopedTransaction(LOCTEXT("SetNumRows", "SetNumRows FPoseSearchColumn"));
 		InternalDatabase->Modify();
 		while (InternalDatabase->GetNumAnimationAssets() > 0 && InternalDatabase->GetNumAnimationAssets() > NumRows)
 		{
@@ -75,67 +96,83 @@ void FPoseSearchColumn::SetNumRows(int32 NumRows)
 	}
 	else
 	{
-		FScopedTransaction ScopedTransaction(LOCTEXT("SetNumRows","SetNumRows FPoseSearchColumn"));
+		FScopedTransaction ScopedTransaction(LOCTEXT("SetNumRows", "SetNumRows FPoseSearchColumn"));
 		InternalDatabase->Modify();
 		while (InternalDatabase->GetNumAnimationAssets() < NumRows)
 		{
-			FPoseSearchDatabaseSequence NewAsset;
-			NewAsset.Sequence = nullptr;
-			InternalDatabase->AddAnimationAsset(FInstancedStruct::Make(NewAsset));
+			FPoseSearchDatabaseAnimationAsset NewAsset;
+			NewAsset.AnimAsset = GetReferencedObject(InternalDatabase->GetNumAnimationAssets());
+			InternalDatabase->AddAnimationAsset(NewAsset);
 		}
 	}
+
+	// NumRows == 0 has a special meaning of resetting the column and usually make the column inconsistent with the chooser asset, so we don't CheckConsistency
+	if (NumRows != 0)
+	{
+		CheckConsistency();
+	}
 }
 
 void FPoseSearchColumn::InsertRows(int Index, int Count)
 {
+	check(InternalDatabase);
 	if (Count > 0)
 	{
 		FScopedTransaction ScopedTransaction(LOCTEXT("InsertRows","InsertRows FPoseSearchColumn"));
 		InternalDatabase->Modify();
 		for (int i = 0; i < Count; i++)
 		{
-			FPoseSearchDatabaseSequence NewAsset;
-			NewAsset.Sequence = nullptr;
-			InternalDatabase->InsertAnimationAssetAt(FInstancedStruct::Make(NewAsset), Index);
+			FPoseSearchDatabaseAnimationAsset NewAsset;
+			NewAsset.AnimAsset = GetReferencedObject(Index);
+			InternalDatabase->InsertAnimationAssetAt(NewAsset, Index);
 		}
 	}
+
+	CheckConsistency();
 }
 
 void FPoseSearchColumn::DeleteRows(const TArray<uint32>& RowIndices)
 {
+	check(InternalDatabase);
 	FScopedTransaction ScopedTransaction(LOCTEXT("DeleteRows","DeleteRows FPoseSearchColumn"));
 	InternalDatabase->Modify();
 	for(uint32 Index : RowIndices)
 	{
 		InternalDatabase->RemoveAnimationAssetAt(Index);
 	}
+	
+	CheckConsistency();
 }
 
 void FPoseSearchColumn::MoveRow(int SourceRowIndex, int TargetRowIndex)
 {
-	if (const FPoseSearchDatabaseSequence* SourceAsset = InternalDatabase->GetDatabaseAnimationAsset<FPoseSearchDatabaseSequence>(SourceRowIndex))
+	check(InternalDatabase);
+	if (const FPoseSearchDatabaseAnimationAsset* SourceAsset = InternalDatabase->GetDatabaseAnimationAsset(SourceRowIndex))
 	{
 		FScopedTransaction ScopedTransaction(LOCTEXT("MoveRow","MoveRow FPoseSearchColumn"));
 		InternalDatabase->Modify();
 
-		FPoseSearchDatabaseSequence CopySourceAsset = *SourceAsset;
+		FPoseSearchDatabaseAnimationAsset CopySourceAsset = *SourceAsset;
 		InternalDatabase->RemoveAnimationAssetAt(SourceRowIndex);
 
 		if (SourceRowIndex < TargetRowIndex)
 		{
 			TargetRowIndex--;
 		}
-		InternalDatabase->InsertAnimationAssetAt(FInstancedStruct::Make(CopySourceAsset), TargetRowIndex);
+		InternalDatabase->InsertAnimationAssetAt(CopySourceAsset, TargetRowIndex);
 	}
+
+	CheckConsistency();
 }
 
 void FPoseSearchColumn::CopyRow(FChooserColumnBase& SourceColumn, int SourceRowIndex, int TargetRowIndex)
 {
+	check(InternalDatabase);
 	if (SourceRowIndex != TargetRowIndex)
 	{
-		if (const FPoseSearchDatabaseSequence* SourceAsset = InternalDatabase->GetDatabaseAnimationAsset<FPoseSearchDatabaseSequence>(SourceRowIndex))
+		if (const FPoseSearchDatabaseAnimationAsset* SourceAsset = InternalDatabase->GetDatabaseAnimationAsset(SourceRowIndex))
 		{
-			if (FPoseSearchDatabaseSequence* TargetAsset = InternalDatabase->GetMutableDatabaseAnimationAsset<FPoseSearchDatabaseSequence>(TargetRowIndex))
+			if (FPoseSearchDatabaseAnimationAsset* TargetAsset = InternalDatabase->GetMutableDatabaseAnimationAsset(TargetRowIndex))
 			{
 				FScopedTransaction ScopedTransaction(LOCTEXT("CopyRow","CopyRow FPoseSearchColumn"));
 				InternalDatabase->Modify();
@@ -143,6 +180,8 @@ void FPoseSearchColumn::CopyRow(FChooserColumnBase& SourceColumn, int SourceRowI
 			}
 		}
 	}
+
+	CheckConsistency();
 }
 
 UScriptStruct* FPoseSearchColumn::GetInputBaseType() const
@@ -188,6 +227,8 @@ void FPoseSearchColumn::Filter(FChooserEvaluationContext& Context, const FChoose
 {
 	using namespace UE::PoseSearch;
 
+	CheckConsistency();
+
 	const UE::PoseSearch::IPoseHistory* PoseHistory = nullptr;
 	FPoseHistoryReference PoseHistoryReference;
 	if (const FChooserParameterPoseHistoryBase* PoseHistoryParameter = InputValue.GetPtr<FChooserParameterPoseHistoryBase>())
@@ -222,7 +263,7 @@ void FPoseSearchColumn::Filter(FChooserEvaluationContext& Context, const FChoose
 		TArray<const UObject*, TInlineAllocator<128>> AssetsToConsider;
 		for (const FChooserIndexArray::FIndexData& IndexData : IndexListIn)
 		{
-			if (const FPoseSearchDatabaseSequence* Asset = InternalDatabase->GetDatabaseAnimationAsset<FPoseSearchDatabaseSequence>(IndexData.Index))
+			if (const FPoseSearchDatabaseAnimationAsset* Asset = InternalDatabase->GetDatabaseAnimationAsset(IndexData.Index))
 			{
 				AssetsToConsider.Add(Asset->GetAnimationAsset());
 			}
@@ -248,9 +289,9 @@ void FPoseSearchColumn::Filter(FChooserEvaluationContext& Context, const FChoose
 
 		if (const UE::PoseSearch::FSearchIndexAsset* SearchIndexAsset = SearchResult.GetSearchIndexAsset())
 		{
-			if (const FPoseSearchDatabaseAnimationAssetBase* DatabaseAnimationAssetBase = SearchResult.Database->GetDatabaseAnimationAsset<FPoseSearchDatabaseAnimationAssetBase>(*SearchIndexAsset))
+			if (const FPoseSearchDatabaseAnimationAsset* DatabaseAnimationAsset = SearchResult.Database->GetDatabaseAnimationAsset(*SearchIndexAsset))
 			{
-				const UObject* Result = DatabaseAnimationAssetBase->GetAnimationAsset();
+				const UObject* Result = DatabaseAnimationAsset->GetAnimationAsset();
 				int Index;
 				if (AssetsToConsider.Find(Result, Index))
 				{
@@ -303,19 +344,20 @@ void FPoseSearchColumn::Filter(FChooserEvaluationContext& Context, const FChoose
 void FPoseSearchColumn::Initialize(UChooserTable* OuterChooser)
 {
 	FChooserColumnBase::Initialize(OuterChooser);
-
 	InternalDatabase = NewObject<UPoseSearchDatabase>(OuterChooser, OuterChooser->GetFName(), RF_Transactional);
 }
 
 void FPoseSearchColumn::AddToDetails(FInstancedPropertyBag& PropertyBag, int32 ColumnIndex, int32 RowIndex)
 {
+	check(InternalDatabase);
+
 	FText DisplayName = NSLOCTEXT("PoseSearchColumn","Pose Search", "Pose Search");
 	FName PropertyName("RowData", ColumnIndex);
-	FPropertyBagPropertyDesc PropertyDesc(PropertyName,  EPropertyBagPropertyType::Struct, FPoseSearchDatabaseSequence::StaticStruct());
+	FPropertyBagPropertyDesc PropertyDesc(PropertyName,  EPropertyBagPropertyType::Struct, FPoseSearchDatabaseAnimationAsset::StaticStruct());
 	PropertyDesc.MetaData.Add(FPropertyBagPropertyDescMetaData("DisplayName", DisplayName.ToString()));
 	PropertyBag.AddProperties({PropertyDesc});
 		
-	if (FPoseSearchDatabaseSequence* Asset = InternalDatabase->GetMutableDatabaseAnimationAsset<FPoseSearchDatabaseSequence>(RowIndex))
+	if (const FPoseSearchDatabaseAnimationAsset* Asset = InternalDatabase->GetDatabaseAnimationAsset(RowIndex))
 	{
 		PropertyBag.SetValueStruct(PropertyName, *Asset);
 	}
@@ -323,16 +365,20 @@ void FPoseSearchColumn::AddToDetails(FInstancedPropertyBag& PropertyBag, int32 C
 
 void FPoseSearchColumn::SetFromDetails(FInstancedPropertyBag& PropertyBag, int32 ColumnIndex, int32 RowIndex)
 {
+	check(InternalDatabase);
+
 	FName PropertyName("RowData", ColumnIndex);
 		
-   	TValueOrError<FStructView, EPropertyBagResult> Result = PropertyBag.GetValueStruct(PropertyName, FPoseSearchDatabaseSequence::StaticStruct());
+   	TValueOrError<FStructView, EPropertyBagResult> Result = PropertyBag.GetValueStruct(PropertyName, FPoseSearchDatabaseAnimationAsset::StaticStruct());
 	if (FStructView* StructView = Result.TryGetValue())
 	{
-		if (FPoseSearchDatabaseSequence* Asset = InternalDatabase->GetMutableDatabaseAnimationAsset<FPoseSearchDatabaseSequence>(RowIndex))
+		if (FPoseSearchDatabaseAnimationAsset* Asset = InternalDatabase->GetMutableDatabaseAnimationAsset(RowIndex))
 		{
-			*Asset = StructView->Get<FPoseSearchDatabaseSequence>();
+			*Asset = StructView->Get<FPoseSearchDatabaseAnimationAsset>();
 		}
 	}
+
+	CheckConsistency();
 }
 #endif
 
diff --git a/Engine/Plugins/Animation/PoseSearch/Source/Runtime/Private/PoseSearchDatabase.cpp b/Engine/Plugins/Animation/PoseSearch/Source/Runtime/Private/PoseSearchDatabase.cpp
--- a/Engine/Plugins/Animation/PoseSearch/Source/Runtime/Private/PoseSearchDatabase.cpp
+++ b/Engine/Plugins/Animation/PoseSearch/Source/Runtime/Private/PoseSearchDatabase.cpp
@@ -67,6 +67,86 @@ static FAutoConsoleVariableRef CVarMotionMatchProfileMaxHeapKNNSearch(TEXT("a.Mo
 static int32 GVarMotionMatchDebugWeightGroupID = 0;
 static FAutoConsoleVariableRef CVarMotionMatchDebugWeightGroupID(TEXT("a.MotionMatch.DebugWeightGroupID"), GVarMotionMatchDebugWeightGroupID, TEXT("Only the channels with this or negative DebugWeightGroupID will have non zero weights"));
 
+FPoseSearchDatabaseAnimationAsset UpgradeToDatabaseAnimationAsset(FInstancedStruct& AnimationAsset)
+{
+	FPoseSearchDatabaseAnimationAsset PoseSearchDatabaseAnimationAsset;
+
+	PRAGMA_DISABLE_DEPRECATION_WARNINGS
+
+	if (const FPoseSearchDatabaseSequence* DatabaseSequence = AnimationAsset.GetPtr<FPoseSearchDatabaseSequence>())
+	{
+		PoseSearchDatabaseAnimationAsset.AnimAsset = DatabaseSequence->GetAnimationAsset();
+#if WITH_EDITORONLY_DATA
+		PoseSearchDatabaseAnimationAsset.bEnabled = DatabaseSequence->bEnabled;
+		PoseSearchDatabaseAnimationAsset.bDisableReselection = DatabaseSequence->bDisableReselection;
+		PoseSearchDatabaseAnimationAsset.MirrorOption = DatabaseSequence->MirrorOption;
+		PoseSearchDatabaseAnimationAsset.BranchInId = DatabaseSequence->BranchInId;
+
+		PoseSearchDatabaseAnimationAsset.SamplingRange = DatabaseSequence->SamplingRange;
+#endif // WITH_EDITORONLY_DATA
+	}
+	else if (const FPoseSearchDatabaseBlendSpace* DatabaseBlendSpace = AnimationAsset.GetPtr<FPoseSearchDatabaseBlendSpace>())
+	{
+		PoseSearchDatabaseAnimationAsset.AnimAsset = DatabaseBlendSpace->GetAnimationAsset();
+#if WITH_EDITORONLY_DATA
+		PoseSearchDatabaseAnimationAsset.bEnabled = DatabaseBlendSpace->bEnabled;
+		PoseSearchDatabaseAnimationAsset.bDisableReselection = DatabaseBlendSpace->bDisableReselection;
+		PoseSearchDatabaseAnimationAsset.MirrorOption = DatabaseBlendSpace->MirrorOption;
+		PoseSearchDatabaseAnimationAsset.BranchInId = DatabaseBlendSpace->BranchInId;
+
+		PoseSearchDatabaseAnimationAsset.bUseSingleSample = DatabaseBlendSpace->bUseSingleSample;
+		PoseSearchDatabaseAnimationAsset.bUseGridForSampling = DatabaseBlendSpace->bUseGridForSampling;
+		PoseSearchDatabaseAnimationAsset.NumberOfHorizontalSamples = DatabaseBlendSpace->NumberOfHorizontalSamples;
+		PoseSearchDatabaseAnimationAsset.NumberOfVerticalSamples = DatabaseBlendSpace->NumberOfVerticalSamples;
+		PoseSearchDatabaseAnimationAsset.BlendParamX = DatabaseBlendSpace->BlendParamX;
+		PoseSearchDatabaseAnimationAsset.BlendParamY = DatabaseBlendSpace->BlendParamY;
+		PoseSearchDatabaseAnimationAsset.SamplingRange = DatabaseBlendSpace->SamplingRange;
+#endif // WITH_EDITORONLY_DATA
+	}
+	else if (const FPoseSearchDatabaseAnimComposite* DatabaseAnimComposite = AnimationAsset.GetPtr<FPoseSearchDatabaseAnimComposite>())
+	{
+		PoseSearchDatabaseAnimationAsset.AnimAsset = DatabaseAnimComposite->GetAnimationAsset();
+#if WITH_EDITORONLY_DATA
+		PoseSearchDatabaseAnimationAsset.bEnabled = DatabaseAnimComposite->bEnabled;
+		PoseSearchDatabaseAnimationAsset.bDisableReselection = DatabaseAnimComposite->bDisableReselection;
+		PoseSearchDatabaseAnimationAsset.MirrorOption = DatabaseAnimComposite->MirrorOption;
+		PoseSearchDatabaseAnimationAsset.BranchInId = DatabaseAnimComposite->BranchInId;
+
+		PoseSearchDatabaseAnimationAsset.SamplingRange = DatabaseAnimComposite->SamplingRange;
+#endif // WITH_EDITORONLY_DATA
+	}
+	else if (const FPoseSearchDatabaseAnimMontage* DatabaseAnimMontage = AnimationAsset.GetPtr<FPoseSearchDatabaseAnimMontage>())
+	{
+		PoseSearchDatabaseAnimationAsset.AnimAsset = DatabaseAnimMontage->GetAnimationAsset();
+#if WITH_EDITORONLY_DATA
+		PoseSearchDatabaseAnimationAsset.bEnabled = DatabaseAnimMontage->bEnabled;
+		PoseSearchDatabaseAnimationAsset.bDisableReselection = DatabaseAnimMontage->bDisableReselection;
+		PoseSearchDatabaseAnimationAsset.MirrorOption = DatabaseAnimMontage->MirrorOption;
+		PoseSearchDatabaseAnimationAsset.BranchInId = DatabaseAnimMontage->BranchInId;
+
+		PoseSearchDatabaseAnimationAsset.SamplingRange = DatabaseAnimMontage->SamplingRange;
+#endif // WITH_EDITORONLY_DATA
+	}
+	else if (const FPoseSearchDatabaseMultiAnimAsset* DatabaseMultiAnimAsset = AnimationAsset.GetPtr<FPoseSearchDatabaseMultiAnimAsset>())
+	{
+		PoseSearchDatabaseAnimationAsset.AnimAsset = DatabaseMultiAnimAsset->GetAnimationAsset();
+#if WITH_EDITORONLY_DATA
+		PoseSearchDatabaseAnimationAsset.bEnabled = DatabaseMultiAnimAsset->bEnabled;
+		PoseSearchDatabaseAnimationAsset.bDisableReselection = DatabaseMultiAnimAsset->bDisableReselection;
+		PoseSearchDatabaseAnimationAsset.MirrorOption = DatabaseMultiAnimAsset->MirrorOption;
+		PoseSearchDatabaseAnimationAsset.BranchInId = DatabaseMultiAnimAsset->BranchInId;
+
+		PoseSearchDatabaseAnimationAsset.NumberOfHorizontalSamples = DatabaseMultiAnimAsset->NumberOfHorizontalSamples;
+		PoseSearchDatabaseAnimationAsset.NumberOfVerticalSamples = DatabaseMultiAnimAsset->NumberOfVerticalSamples;
+		PoseSearchDatabaseAnimationAsset.SamplingRange = DatabaseMultiAnimAsset->SamplingRange;
+#endif // WITH_EDITORONLY_DATA
+	}
+		
+	PRAGMA_ENABLE_DEPRECATION_WARNINGS
+
+	return PoseSearchDatabaseAnimationAsset;
+}
+
 struct FSearchFilters
 {
 	FSearchFilters(const UPoseSearchSchema* Schema, TConstArrayView<int32> NonSelectableIdx, TConstArrayView<int32> SelectableAssetIdx, bool bAddBlockTransitionFilter)
@@ -441,6 +521,181 @@ FFloatInterval FPoseSearchDatabaseAnimationAssetBase::GetEffectiveSamplingRange(
 }
 #endif // WITH_EDITORONLY_DATA
 
+//////////////////////////////////////////////////////////////////////////
+// FPoseSearchDatabaseAnimationAsset
+UObject* FPoseSearchDatabaseAnimationAsset::GetAnimationAsset() const
+{
+	return AnimAsset.Get();
+}
+
+float FPoseSearchDatabaseAnimationAsset::GetPlayLength(const FVector& BlendParameters) const
+{
+	if (!AnimAsset)
+	{
+		return 0.f;
+	}
+
+	if (const UBlendSpace* BlendSpace = Cast<UBlendSpace>(AnimAsset))
+	{
+		int32 TriangulationIndex = 0;
+		TArray<FBlendSampleData> BlendSamples;
+		BlendSpace->GetSamplesFromBlendInput(BlendParameters, BlendSamples, TriangulationIndex, true);
+		const float PlayLength = BlendSpace->GetAnimationLengthFromSampleData(BlendSamples);
+		return PlayLength;
+	}
+
+	if (const UMultiAnimAsset* MultiAnimAsset = Cast<UMultiAnimAsset>(AnimAsset))
+	{
+		return MultiAnimAsset->GetPlayLength(BlendParameters);
+	}
+
+	return FPoseSearchDatabaseAnimationAssetBase::GetPlayLength(BlendParameters);
+}
+
+#if WITH_EDITORONLY_DATA
+// @todo: Fix me!
+UClass* FPoseSearchDatabaseAnimationAsset::GetAnimationAssetStaticClass() const
+{
+	if (!AnimAsset)
+	{
+		return nullptr;
+	}
+
+	return AnimAsset->GetClass();
+}
+
+bool FPoseSearchDatabaseAnimationAsset::IsLooping(const UObject* AnimAsset, const FFloatInterval& SamplingRange)
+{
+	if (!AnimAsset)
+	{
+		return false;
+	}
+
+	if (const UAnimSequenceBase* AnimSequenceBase = Cast<UAnimSequenceBase>(AnimAsset))
+	{
+		return AnimSequenceBase->bLoop && SamplingRange.Min == 0.f && SamplingRange.Max == 0.f;
+	}
+
+	if (const UBlendSpace* BlendSpace = Cast<UBlendSpace>(AnimAsset))
+	{
+		return BlendSpace->bLoop;
+	}
+
+	if (const UMultiAnimAsset* MultiAnimAsset = Cast<UMultiAnimAsset>(AnimAsset))
+	{
+		return MultiAnimAsset->IsLooping() && SamplingRange.Min == 0.f && SamplingRange.Max == 0.f;
+	}
+
+	return false;
+}
+
+bool FPoseSearchDatabaseAnimationAsset::IsLooping() const
+{
+	return IsLooping(AnimAsset, SamplingRange);
+}
+
+bool FPoseSearchDatabaseAnimationAsset::IsRootMotionEnabled(const UObject* AnimAsset)
+{
+	if (!AnimAsset)
+	{
+		return false;
+	}
+
+	if (const UAnimSequenceBase* AnimSequenceBase = Cast<UAnimSequenceBase>(AnimAsset))
+	{
+		return AnimSequenceBase->HasRootMotion();
+	}
+
+	if (const UBlendSpace* BlendSpace = Cast<UBlendSpace>(AnimAsset))
+	{
+		return UE::PoseSearch::IsBlendSpaceRootMotionEnabled(BlendSpace);
+	}
+
+	if (const UMultiAnimAsset* MultiAnimAsset = Cast<UMultiAnimAsset>(AnimAsset))
+	{
+		return MultiAnimAsset->HasRootMotion();
+	}
+
+	return false;
+}
+
+bool FPoseSearchDatabaseAnimationAsset::IsRootMotionEnabled() const
+{
+	return IsRootMotionEnabled(AnimAsset);
+}
+
+FFloatInterval FPoseSearchDatabaseAnimationAsset::GetEffectiveSamplingRange(const FVector& BlendParameters) const
+{
+	if (const UBlendSpace* BlendSpace = Cast<UBlendSpace>(AnimAsset))
+	{
+		const float PlayLength = GetPlayLength(BlendParameters);
+
+		// scaling blend space SamplingRange from the space [0, 1] to [0, PlayLength] with PlayLength calculated from the BlendSamples
+		FFloatInterval ScaledSamplingRange = SamplingRange;
+		ScaledSamplingRange.Min *= PlayLength;
+		ScaledSamplingRange.Max *= PlayLength;
+
+		return FPoseSearchDatabaseAnimationAssetBase::GetEffectiveSamplingRange(PlayLength, ScaledSamplingRange);
+	}
+
+	return FPoseSearchDatabaseAnimationAssetBase::GetEffectiveSamplingRange(BlendParameters);
+}
+
+void FPoseSearchDatabaseAnimationAsset::IterateOverSamplingParameter(const TFunction<void(const FVector& BlendParameters)>& ProcessSamplingParameter) const
+{
+	if (const UBlendSpace* BlendSpace = Cast<UBlendSpace>(AnimAsset))
+	{
+		UE::PoseSearch::IterateOverBlendSpaceSamplingParameter(BlendSpace, bUseSingleSample, FVector(BlendParamX, BlendParamY, 0.f),
+			bUseGridForSampling, NumberOfHorizontalSamples, NumberOfVerticalSamples, ProcessSamplingParameter);
+	}
+	else if (const UMultiAnimAsset* MultiAnimAsset = Cast<UMultiAnimAsset>(AnimAsset))
+	{
+		const UBlendSpace* MultiAnimAssetBlendSpace = nullptr;
+		for (int32 RoleIndex = 0; RoleIndex < MultiAnimAsset->GetNumRoles(); ++RoleIndex)
+		{
+			if (const UAnimationAsset* AnimationAsset = MultiAnimAsset->GetAnimationAsset(MultiAnimAsset->GetRole(RoleIndex)))
+			{
+				if (AnimationAsset->GetClass()->IsChildOf(UBlendSpace::StaticClass()))
+				{
+					// @todo: right now we just sample using the first blend space, but we should probably make sure if there're multiple MultiAnimAssetBlendSpace they are consistent with each other
+					MultiAnimAssetBlendSpace = Cast<UBlendSpace>(AnimationAsset);
+					check(MultiAnimAssetBlendSpace);
+					break;
+				}
+			}
+		}
+
+		if (MultiAnimAssetBlendSpace)
+		{
+			UE::PoseSearch::IterateOverBlendSpaceSamplingParameter(MultiAnimAssetBlendSpace, false, FVector::ZeroVector,
+				false, NumberOfHorizontalSamples, NumberOfVerticalSamples, ProcessSamplingParameter);
+		}
+		else
+		{
+			ProcessSamplingParameter(FVector::ZeroVector);
+		}
+	}
+	else
+	{
+		FPoseSearchDatabaseAnimationAssetBase::IterateOverSamplingParameter(ProcessSamplingParameter);
+	}
+}
+
+#endif // WITH_EDITORONLY_DATA
+
+#if WITH_EDITOR
+USkeletalMesh* FPoseSearchDatabaseAnimationAsset::GetPreviewMeshForRole(const UE::PoseSearch::FRole& Role) const
+{
+	if (UMultiAnimAsset* MultiAnimAsset = Cast<UMultiAnimAsset>(AnimAsset))
+	{
+		return MultiAnimAsset->GetPreviewMesh(Role);
+	}
+
+	return FPoseSearchDatabaseAnimationAssetBase::GetPreviewMeshForRole(Role);
+}
+#endif // WITH_EDITOR
+
+PRAGMA_DISABLE_DEPRECATION_WARNINGS
 //////////////////////////////////////////////////////////////////////////
 // FPoseSearchDatabaseSequence
 UObject* FPoseSearchDatabaseSequence::GetAnimationAsset() const
@@ -484,14 +739,6 @@ float FPoseSearchDatabaseBlendSpace::GetPlayLength(const FVector& BlendParameter
 	return PlayLength;
 }
 
-//#if WITH_EDITOR
-//int32 FPoseSearchDatabaseBlendSpace::GetFrameAtTime(float Time) const
-//{
-//	// returning the percentage of time as value to diplay in the pose search debugger (NoTe: BlendSpace->GetPlayLength() is one)
-//	return FMath::RoundToInt(Time * 100.f);
-//}
-//#endif // WITH_EDITOR
-
 #if WITH_EDITORONLY_DATA
 UClass* FPoseSearchDatabaseBlendSpace::GetAnimationAssetStaticClass() const
 {
@@ -575,9 +822,7 @@ FVector FPoseSearchDatabaseBlendSpace::BlendParameterForSampleRanges(int32 Horiz
 	const bool bWrapInputOnVerticalAxis = BlendSpace->GetBlendParameter(1).bWrapInput;
 
 	int32 HorizontalBlendNum, VerticalBlendNum;
-	PRAGMA_DISABLE_DEPRECATION_WARNINGS
 	GetBlendSpaceParameterSampleRanges(HorizontalBlendNum, VerticalBlendNum);
-	PRAGMA_ENABLE_DEPRECATION_WARNINGS
 
 	if (bWrapInputOnHorizontalAxis)
 	{
@@ -756,7 +1001,6 @@ void FPoseSearchDatabaseMultiAnimAsset::IterateOverSamplingParameter(const TFunc
 
 //////////////////////////////////////////////////////////////////////////
 // UPoseSearchDatabase
-PRAGMA_DISABLE_DEPRECATION_WARNINGS
 UPoseSearchDatabase::~UPoseSearchDatabase()
 {
 }
@@ -779,7 +1023,7 @@ void UPoseSearchDatabase::UpdateCachedProperties()
 	{
 		const FSearchIndexAsset& SearchIndexAsset = SearchIndexPrivate.Assets[AssetIdx];
 
-		if (const FPoseSearchDatabaseAnimationAssetBase* DatabaseAnimationAssetBase = GetDatabaseAnimationAsset<FPoseSearchDatabaseAnimationAssetBase>(SearchIndexAsset))
+		if (const FPoseSearchDatabaseAnimationAssetBase* DatabaseAnimationAssetBase = GetDatabaseAnimationAsset(SearchIndexAsset))
 		{
 			CachedAssetMap.FindOrAdd(DatabaseAnimationAssetBase->GetAnimationAsset()).Add(AssetIdx);
 		}
@@ -803,7 +1047,7 @@ TConstArrayView<int32> UPoseSearchDatabase::GetAssetIndexesForSourceAsset(const
 		for (int32 AssetIndex : *IndexesForSourceAsset)
 		{
 			const FSearchIndexAsset& SearchIndexAsset = SearchIndex.Assets[AssetIndex];
-			const FPoseSearchDatabaseAnimationAssetBase* DatabaseAnimationAssetBase = GetDatabaseAnimationAsset<FPoseSearchDatabaseAnimationAssetBase>(SearchIndexAsset);
+			const FPoseSearchDatabaseAnimationAssetBase* DatabaseAnimationAssetBase = GetDatabaseAnimationAsset(SearchIndexAsset);
 	
 			// if those checks fail the calling code hasn't been protected by FAsyncPoseSearchDatabasesManagement::RequestAsyncBuildIndex
 			check(DatabaseAnimationAssetBase);
@@ -901,7 +1145,7 @@ int32 UPoseSearchDatabase::GetPoseIndex(const UObject* AnimationAsset, float Ani
 				{
 					MinSquaredLength = BlendParametersSquaredLength;
 
-					const FPoseSearchDatabaseAnimationAssetBase* DatabaseAnimationAssetBase = GetDatabaseAnimationAsset<FPoseSearchDatabaseAnimationAssetBase>(SearchIndexAsset);
+					const FPoseSearchDatabaseAnimationAssetBase* DatabaseAnimationAssetBase = GetDatabaseAnimationAsset(SearchIndexAsset);
 
 					check(DatabaseAnimationAssetBase);
 					check(DatabaseAnimationAssetBase->GetAnimationAsset() == AnimationAsset);
@@ -918,17 +1162,53 @@ int32 UPoseSearchDatabase::GetPoseIndex(const UObject* AnimationAsset, float Ani
 
 void UPoseSearchDatabase::AddAnimationAsset(FInstancedStruct AnimationAsset)
 {
-	AnimationAssets.Add(AnimationAsset);
+	using namespace UE::PoseSearch;
+	PRAGMA_DISABLE_DEPRECATION_WARNINGS
+	DatabaseAnimationAssets.Add(UpgradeToDatabaseAnimationAsset(AnimationAsset));
+	PRAGMA_ENABLE_DEPRECATION_WARNINGS
+}
+
+void UPoseSearchDatabase::AddAnimationAsset(const FPoseSearchDatabaseAnimationAsset& AnimationAsset)
+{
+	DatabaseAnimationAssets.Add(AnimationAsset);
 }
 
-void UPoseSearchDatabase::InsertAnimationAssetAt(FInstancedStruct AnimationAsset, int32 AnimationAssetIndex)
+void UPoseSearchDatabase::InsertAnimationAssetAt(const FPoseSearchDatabaseAnimationAsset& AnimationAsset, int32 AnimationAssetIndex)
 {
-	AnimationAssets.Insert(AnimationAsset, AnimationAssetIndex);
+	DatabaseAnimationAssets.Insert(AnimationAsset, AnimationAssetIndex);
 }
 
 void UPoseSearchDatabase::RemoveAnimationAssetAt(int32 AnimationAssetIndex)
 {
-	AnimationAssets.RemoveAt(AnimationAssetIndex);
+	DatabaseAnimationAssets.RemoveAt(AnimationAssetIndex);
+}
+
+const FPoseSearchDatabaseAnimationAsset* UPoseSearchDatabase::GetDatabaseAnimationAsset(int32 AnimationAssetIndex) const
+{
+	PRAGMA_DISABLE_DEPRECATION_WARNINGS
+	return GetDatabaseAnimationAsset<FPoseSearchDatabaseAnimationAsset>(AnimationAssetIndex);
+	PRAGMA_ENABLE_DEPRECATION_WARNINGS
+}
+
+const FPoseSearchDatabaseAnimationAsset* UPoseSearchDatabase::GetDatabaseAnimationAsset(const UE::PoseSearch::FSearchIndexAsset& SearchIndexAsset) const
+{
+	PRAGMA_DISABLE_DEPRECATION_WARNINGS
+	return GetDatabaseAnimationAsset<FPoseSearchDatabaseAnimationAsset>(SearchIndexAsset);
+	PRAGMA_ENABLE_DEPRECATION_WARNINGS
+}
+
+FPoseSearchDatabaseAnimationAsset* UPoseSearchDatabase::GetMutableDatabaseAnimationAsset(int32 AnimationAssetIndex)
+{
+	PRAGMA_DISABLE_DEPRECATION_WARNINGS
+	return GetMutableDatabaseAnimationAsset<FPoseSearchDatabaseAnimationAsset>(AnimationAssetIndex);
+	PRAGMA_ENABLE_DEPRECATION_WARNINGS
+}
+
+FPoseSearchDatabaseAnimationAsset* UPoseSearchDatabase::GetMutableDatabaseAnimationAsset(const UE::PoseSearch::FSearchIndexAsset& SearchIndexAsset)
+{
+	PRAGMA_DISABLE_DEPRECATION_WARNINGS
+	return GetMutableDatabaseAnimationAsset<FPoseSearchDatabaseAnimationAsset>(SearchIndexAsset);
+	PRAGMA_ENABLE_DEPRECATION_WARNINGS
 }
 
 #if WITH_EDITOR
@@ -951,9 +1231,22 @@ bool UPoseSearchDatabase::GetSkipSearchIfPossible() const
 
 void UPoseSearchDatabase::PostLoad()
 {
-#if WITH_EDITOR
 	using namespace UE::PoseSearch;
 
+	PRAGMA_DISABLE_DEPRECATION_WARNINGS
+	if (!AnimationAssets_DEPRECATED.IsEmpty())
+	{
+		DatabaseAnimationAssets.Reserve(AnimationAssets_DEPRECATED.Num());
+		for (FInstancedStruct& AnimationAsset : AnimationAssets_DEPRECATED)
+		{
+			DatabaseAnimationAssets.Add(UpgradeToDatabaseAnimationAsset(AnimationAsset));
+		}
+		AnimationAssets_DEPRECATED.Empty();
+	}
+	PRAGMA_ENABLE_DEPRECATION_WARNINGS
+
+#if WITH_EDITOR
+
 	// todo: should we SynchronizeWithExternalDependencies() here?
 	SynchronizeChooser();
 
@@ -987,17 +1280,14 @@ bool UPoseSearchDatabase::Contains(const UObject* Object) const
 
 int32 UPoseSearchDatabase::GetNumAnimationAssets() const
 {
-	return AnimationAssets.Num();
+	return DatabaseAnimationAssets.Num();
 }
 
 UObject* UPoseSearchDatabase::GetAnimationAsset(int32 Index) const
 {
-	if (AnimationAssets.IsValidIndex(Index))
+	if (DatabaseAnimationAssets.IsValidIndex(Index))
 	{
-		if (const FPoseSearchDatabaseAnimationAssetBase* AnimationAssetBase = AnimationAssets[Index].GetPtr<FPoseSearchDatabaseAnimationAssetBase>())
-		{
-			return AnimationAssetBase->GetAnimationAsset();
-		}
+		return DatabaseAnimationAssets[Index].GetAnimationAsset();
 	}
 	return nullptr;
 }
@@ -1055,18 +1345,16 @@ void UPoseSearchDatabase::SynchronizeWithExternalDependencies()
 void UPoseSearchDatabase::SynchronizeWithExternalDependencies(TConstArrayView<UAnimSequenceBase*> SequencesBase)
 {
 	// cannot use TSet since FInstancedStruct doesn't implement GetTypeHash
-	TArray<FInstancedStruct> NewAnimationAssets;
+	TArray<FPoseSearchDatabaseAnimationAsset> NewDatabaseAnimationAssets;
 
 	// collecting all the database AnimationAsset(s) that don't require synchronization
 	
-	for (FInstancedStruct& AnimationAsset : AnimationAssets)
+	for (FPoseSearchDatabaseAnimationAsset& DatabaseAnimationAsset : DatabaseAnimationAssets)
 	{
-		FPoseSearchDatabaseAnimationAssetBase& AnimationAssetBase = AnimationAsset.GetMutable<FPoseSearchDatabaseAnimationAssetBase>();
-
-		const bool bRequiresSynchronization = AnimationAssetBase.IsSynchronizedWithExternalDependency() && SequencesBase.Contains(AnimationAssetBase.GetAnimationAsset());
+		const bool bRequiresSynchronization = DatabaseAnimationAsset.IsSynchronizedWithExternalDependency() && SequencesBase.Contains(DatabaseAnimationAsset.GetAnimationAsset());
 		if (!bRequiresSynchronization)
 		{
-			NewAnimationAssets.Add(AnimationAsset);
+			NewDatabaseAnimationAssets.Add(DatabaseAnimationAsset);
 		}
 	}
 
@@ -1091,47 +1379,28 @@ void UPoseSearchDatabase::SynchronizeWithExternalDependencies(TConstArrayView<UA
 							return SamplingRange;
 						};
 
-						if (UAnimSequence* Sequence = Cast<UAnimSequence>(SequenceBase))
-						{
-							FPoseSearchDatabaseSequence DatabaseSequence;
-							DatabaseSequence.Sequence = Sequence;
-							DatabaseSequence.SamplingRange = GetSamplingRange(NotifyEvent, SequenceBase);
-							DatabaseSequence.BranchInId = PoseSearchBranchIn->GetBranchInId();
-							NewAnimationAssets.Add(FInstancedStruct::Make(DatabaseSequence));
-						}
-						else if (UAnimComposite* AnimComposite = Cast<UAnimComposite>(SequenceBase))
-						{
-							FPoseSearchDatabaseAnimComposite DatabaseAnimComposite;
-							DatabaseAnimComposite.AnimComposite = AnimComposite;
-							DatabaseAnimComposite.SamplingRange = GetSamplingRange(NotifyEvent, SequenceBase);
-							DatabaseAnimComposite.BranchInId = PoseSearchBranchIn->GetBranchInId();
-							NewAnimationAssets.Add(FInstancedStruct::Make(DatabaseAnimComposite));
-						}
-						else if (UAnimMontage* AnimMontage = Cast<UAnimMontage>(SequenceBase))
-						{
-							FPoseSearchDatabaseAnimMontage DatabaseAnimMontage;
-							DatabaseAnimMontage.AnimMontage = AnimMontage;
-							DatabaseAnimMontage.SamplingRange = GetSamplingRange(NotifyEvent, SequenceBase);
-							DatabaseAnimMontage.BranchInId = PoseSearchBranchIn->GetBranchInId();
-							NewAnimationAssets.Add(FInstancedStruct::Make(DatabaseAnimMontage));
-						}
+						FPoseSearchDatabaseAnimationAsset DatabaseAnimationAsset;
+						DatabaseAnimationAsset.AnimAsset = SequenceBase;
+						DatabaseAnimationAsset.SamplingRange = GetSamplingRange(NotifyEvent, SequenceBase);
+						DatabaseAnimationAsset.BranchInId = PoseSearchBranchIn->GetBranchInId();
+						NewDatabaseAnimationAssets.Add(DatabaseAnimationAsset);
 					}
 				}
 			}
 		}
 	}
 
-	// updating AnimationAssets from NewAnimationAssets preserving the original sorting
+	// updating DatabaseAnimationAssets from NewDatabaseAnimationAssets preserving the original sorting
 	bool bModified = false;
 	for (int32 AnimationAssetIndex = GetNumAnimationAssets() - 1; AnimationAssetIndex >= 0; --AnimationAssetIndex)
 	{
-		FPoseSearchDatabaseAnimationAssetBase& AnimationAsset = AnimationAssets[AnimationAssetIndex].GetMutable<FPoseSearchDatabaseAnimationAssetBase>();
+		FPoseSearchDatabaseAnimationAsset& DatabaseAnimationAsset = DatabaseAnimationAssets[AnimationAssetIndex];
 		
 		int32 FoundIndex = -1;
-		for(int i=0; i < NewAnimationAssets.Num(); i++)
+		for(int i=0; i < NewDatabaseAnimationAssets.Num(); i++)
 		{
-			const FPoseSearchDatabaseAnimationAssetBase& NewAnimationAsset = NewAnimationAssets[i].Get<FPoseSearchDatabaseAnimationAssetBase>();
-			if (AnimationAsset.UpdateFrom(NewAnimationAsset))
+			const FPoseSearchDatabaseAnimationAssetBase& NewDatabaseAnimationAsset = NewDatabaseAnimationAssets[i];
+			if (DatabaseAnimationAsset.UpdateFrom(NewDatabaseAnimationAsset))
 			{
 				FoundIndex = i;
 				break;
@@ -1141,19 +1410,19 @@ void UPoseSearchDatabase::SynchronizeWithExternalDependencies(TConstArrayView<UA
 		
 		if (FoundIndex >= 0)
 		{
-			NewAnimationAssets.RemoveAt(FoundIndex); 
+			NewDatabaseAnimationAssets.RemoveAt(FoundIndex); 
 		}
 		else
 		{
-			AnimationAssets.RemoveAt(AnimationAssetIndex);
+			DatabaseAnimationAssets.RemoveAt(AnimationAssetIndex);
 			bModified = true;
 		}
 	}
 
-	// adding the remaining AnimationAsset(s) from AnimationAssetsSet
-	for (const FInstancedStruct& AnimationAsset : NewAnimationAssets)
+	// adding the remaining NewDatabaseAnimationAsset(s) from NewDatabaseAnimationAssets
+	for (const FPoseSearchDatabaseAnimationAsset& NewDatabaseAnimationAsset : NewDatabaseAnimationAssets)
 	{
-		AnimationAssets.Add(AnimationAsset);
+		DatabaseAnimationAssets.Add(NewDatabaseAnimationAsset);
 		bModified = true;
 	}
 
@@ -1191,68 +1460,35 @@ void UPoseSearchDatabase::SynchronizeChooser()
 			}));
 
 		TArray<UObject*, TInlineAllocator<128>> DatabaseAssets;
-		for (FInstancedStruct& AnimationAsset : AnimationAssets)
+		for (FPoseSearchDatabaseAnimationAsset& DatabaseAnimationAsset : DatabaseAnimationAssets)
 		{
-			if (const FPoseSearchDatabaseAnimationAssetBase* AnimationAssetBase = AnimationAsset.GetPtr<FPoseSearchDatabaseAnimationAssetBase>())
-			{
-				DatabaseAssets.Add(AnimationAssetBase->GetAnimationAsset());
-			}
+			DatabaseAssets.Add(DatabaseAnimationAsset.GetAnimationAsset());
 		}
 
 		if (ChooserAssets != DatabaseAssets)
 		{
-			TArray<FInstancedStruct> AnimationAssetsBackup = MoveTemp(AnimationAssets);
+			TArray<FPoseSearchDatabaseAnimationAsset> DatabaseAnimationAssetsBackup = MoveTemp(DatabaseAnimationAssets);
 		
 			for (UObject* ChooserAsset : ChooserAssets)
 			{
 				// searching ChooserAsset in AnimationAssetsBackup
 				bool bFoundInAnimationAssetsBackup = false;
-				for (int32 AnimationAssetBackupIndex = 0; AnimationAssetBackupIndex < AnimationAssetsBackup.Num(); ++AnimationAssetBackupIndex)
+				for (int32 DatabaseAnimationAssetBackupIndex = 0; DatabaseAnimationAssetBackupIndex < DatabaseAnimationAssetsBackup.Num(); ++DatabaseAnimationAssetBackupIndex)
 				{
-					if (const FPoseSearchDatabaseAnimationAssetBase* AnimationAssetBase = AnimationAssetsBackup[AnimationAssetBackupIndex].GetPtr<FPoseSearchDatabaseAnimationAssetBase>())
+					if (DatabaseAnimationAssetsBackup[DatabaseAnimationAssetBackupIndex].GetAnimationAsset() == ChooserAsset)
 					{
-						if (AnimationAssetBase->GetAnimationAsset() == ChooserAsset)
-						{
-							bFoundInAnimationAssetsBackup = true;
-							AnimationAssets.Add(AnimationAssetsBackup[AnimationAssetBackupIndex]);
-							AnimationAssetsBackup.RemoveAt(AnimationAssetBackupIndex);
-							break;
-						}
+						bFoundInAnimationAssetsBackup = true;
+						DatabaseAnimationAssets.Add(DatabaseAnimationAssetsBackup[DatabaseAnimationAssetBackupIndex]);
+						DatabaseAnimationAssetsBackup.RemoveAt(DatabaseAnimationAssetBackupIndex);
+						break;
 					}
 				}
 
 				if (!bFoundInAnimationAssetsBackup)
 				{
-					if (UAnimSequence* Sequence = Cast<UAnimSequence>(ChooserAsset))
-					{
-						FPoseSearchDatabaseSequence DatabaseSequence;
-						DatabaseSequence.Sequence = Sequence;
-						AnimationAssets.Add(FInstancedStruct::Make(DatabaseSequence));
-					}
-					else if (UAnimComposite* AnimComposite = Cast<UAnimComposite>(ChooserAsset))
-					{
-						FPoseSearchDatabaseAnimComposite DatabaseAnimComposite;
-						DatabaseAnimComposite.AnimComposite = AnimComposite;
-						AnimationAssets.Add(FInstancedStruct::Make(DatabaseAnimComposite));
-					}
-					else if (UAnimMontage* AnimMontage = Cast<UAnimMontage>(ChooserAsset))
-					{
-						FPoseSearchDatabaseAnimMontage DatabaseAnimMontage;
-						DatabaseAnimMontage.AnimMontage = AnimMontage;
-						AnimationAssets.Add(FInstancedStruct::Make(DatabaseAnimMontage));
-					}
-					else if (UBlendSpace* BlendSpace = Cast<UBlendSpace>(ChooserAsset))
-					{
-						FPoseSearchDatabaseBlendSpace DatabaseBlendSpace;
-						DatabaseBlendSpace.BlendSpace = BlendSpace;
-						AnimationAssets.Add(FInstancedStruct::Make(DatabaseBlendSpace));
-					}
-					else if (UMultiAnimAsset* MultiAnimAsset = Cast<UMultiAnimAsset>(ChooserAsset))
-					{
-						FPoseSearchDatabaseMultiAnimAsset DatabaseMultiAnimAsset;
-						DatabaseMultiAnimAsset.MultiAnimAsset = MultiAnimAsset;
-						AnimationAssets.Add(FInstancedStruct::Make(DatabaseMultiAnimAsset));
-					}
+					FPoseSearchDatabaseAnimationAsset NewDatabaseAnimationAsset;
+					NewDatabaseAnimationAsset.AnimAsset = ChooserAsset;
+					DatabaseAnimationAssets.Add(NewDatabaseAnimationAsset);
 				}
 			}
 
@@ -1280,13 +1516,14 @@ bool UPoseSearchDatabase::IsCachedCookedPlatformDataLoaded(const ITargetPlatform
 #if WITH_EDITOR && ENABLE_ANIM_DEBUG
 void UPoseSearchDatabase::TestSynchronizeWithExternalDependencies()
 {
-	TArray<FInstancedStruct> AnimationAssetsCopy = AnimationAssets;
+	TArray<FPoseSearchDatabaseAnimationAsset> DatabaseAnimationAssetsCopy = DatabaseAnimationAssets;
+	
 	SynchronizeWithExternalDependencies();
 
-	if (AnimationAssetsCopy != AnimationAssets)
+	if (DatabaseAnimationAssetsCopy != DatabaseAnimationAssets)
 	{
 		UE_LOG(LogPoseSearch, Error, TEXT("TestSynchronizeWithExternalDependencies failed"));
-		AnimationAssets = AnimationAssetsCopy;
+		DatabaseAnimationAssets = DatabaseAnimationAssetsCopy;
 	}
 }
 #endif // WITH_EDITOR && ENABLE_ANIM_DEBUG
@@ -1456,7 +1693,7 @@ void UPoseSearchDatabase::PopulateNonSelectableIdx(FNonSelectableIdx& NonSelecta
 			if (CurrentIndexAsset->IsDisableReselection())
 			{
 				// excluding all the poses with CurrentIndexAsset->GetSourceAssetIdx()
-				const FPoseSearchDatabaseAnimationAssetBase* DatabaseAnimationAssetBase = GetDatabaseAnimationAsset<FPoseSearchDatabaseAnimationAssetBase>(CurrentIndexAsset->GetSourceAssetIdx());
+				const FPoseSearchDatabaseAnimationAssetBase* DatabaseAnimationAssetBase = GetDatabaseAnimationAsset(CurrentIndexAsset->GetSourceAssetIdx());
 				check(DatabaseAnimationAssetBase);
 
 				for (int32 AssetIndex : GetAssetIndexesForSourceAsset(DatabaseAnimationAssetBase->GetAnimationAsset()))
@@ -1599,7 +1836,7 @@ UE::PoseSearch::FSearchResult UPoseSearchDatabase::SearchContinuingPose(UE::Pose
 	const FSearchIndex& SearchIndex = GetSearchIndex();
 	const int32 ContinuingPoseIdx = SearchContext.GetContinuingPoseSearchResult().PoseIdx;
 	const FSearchIndexAsset& SearchIndexAsset = SearchIndex.GetAssetForPose(ContinuingPoseIdx);
-	const FPoseSearchDatabaseAnimationAssetBase* DatabaseAnimationAssetBase = GetDatabaseAnimationAsset<FPoseSearchDatabaseAnimationAssetBase>(SearchIndexAsset);
+	const FPoseSearchDatabaseAnimationAssetBase* DatabaseAnimationAssetBase = GetDatabaseAnimationAsset(SearchIndexAsset);
 	check(DatabaseAnimationAssetBase);
 
 	float ContinuingPoseCostAddend = ContinuingPoseCostBias;
diff --git a/Engine/Plugins/Animation/PoseSearch/Source/Runtime/Private/PoseSearchDerivedData.cpp b/Engine/Plugins/Animation/PoseSearch/Source/Runtime/Private/PoseSearchDerivedData.cpp
--- a/Engine/Plugins/Animation/PoseSearch/Source/Runtime/Private/PoseSearchDerivedData.cpp
+++ b/Engine/Plugins/Animation/PoseSearch/Source/Runtime/Private/PoseSearchDerivedData.cpp
@@ -388,7 +388,7 @@ static bool InitSearchIndexAssets(FSearchIndexBase& SearchIndex, const UPoseSear
 	int32 TotalPoses = 0;
 	for (int32 AnimationAssetIndex = 0; AnimationAssetIndex < DatabaseToLookForAssets->GetNumAnimationAssets(); ++AnimationAssetIndex)
 	{
-		if (const FPoseSearchDatabaseAnimationAssetBase* DatabaseAsset = DatabaseToLookForAssets->GetDatabaseAnimationAsset<FPoseSearchDatabaseAnimationAssetBase>(AnimationAssetIndex))
+		if (const FPoseSearchDatabaseAnimationAssetBase* DatabaseAsset = DatabaseToLookForAssets->GetDatabaseAnimationAsset(AnimationAssetIndex))
 		{
 			if (!DatabaseAsset->IsEnabled() || !DatabaseAsset->GetAnimationAsset())
 			{
@@ -962,7 +962,7 @@ static void InitRoledBoneContainers(TMap<FRole, FBoneContainer>& RoledBoneContai
 
 		for (int32 AnimationAssetIndex = 0; AnimationAssetIndex < DatabaseToLookForAssets->GetNumAnimationAssets(); ++AnimationAssetIndex)
 		{
-			if (const FPoseSearchDatabaseAnimationAssetBase* DatabaseAsset = DatabaseToLookForAssets->GetDatabaseAnimationAsset<FPoseSearchDatabaseAnimationAssetBase>(AnimationAssetIndex))
+			if (const FPoseSearchDatabaseAnimationAssetBase* DatabaseAsset = DatabaseToLookForAssets->GetDatabaseAnimationAsset(AnimationAssetIndex))
 			{
 				if (!DatabaseAsset->IsEnabled() || !DatabaseAsset->GetAnimationAsset())
 				{
@@ -1011,7 +1011,7 @@ static bool IndexDatabase(FSearchIndexBase& SearchIndexBase, const UPoseSearchDa
 	{
 		FSearchIndexAsset& SearchIndexAsset = SearchIndexBase.Assets[AssetIdx];
 
-		const FPoseSearchDatabaseAnimationAssetBase* DatabaseAnimationAssetBase = DatabaseToLookForAssets->GetDatabaseAnimationAsset<FPoseSearchDatabaseAnimationAssetBase>(SearchIndexAsset.GetSourceAssetIdx());
+		const FPoseSearchDatabaseAnimationAssetBase* DatabaseAnimationAssetBase = DatabaseToLookForAssets->GetDatabaseAnimationAsset(SearchIndexAsset.GetSourceAssetIdx());
 		check(DatabaseAnimationAssetBase);
 
 		const int32 NumRoles = DatabaseAnimationAssetBase->GetNumRoles();
@@ -1062,7 +1062,7 @@ static bool IndexDatabase(FSearchIndexBase& SearchIndexBase, const UPoseSearchDa
 		FSearchIndexAsset& SearchIndexAsset = SearchIndexBase.Assets[AssetIdx];
 		check(SearchIndexAsset.GetFirstPoseIdx() == TotalPoses);
 
-		const FPoseSearchDatabaseAnimationAssetBase* DatabaseAnimationAssetBase = DatabaseToLookForAssets->GetDatabaseAnimationAsset<FPoseSearchDatabaseAnimationAssetBase>(SearchIndexAsset.GetSourceAssetIdx());
+		const FPoseSearchDatabaseAnimationAssetBase* DatabaseAnimationAssetBase = DatabaseToLookForAssets->GetDatabaseAnimationAsset(SearchIndexAsset.GetSourceAssetIdx());
 		check(DatabaseAnimationAssetBase);
 
 		TempAssetSamplers.Reset();
@@ -1219,7 +1219,7 @@ static bool ValidateSearchInfoAgainstDatabase(const FSearchIndex& SearchIndex, c
 
 	for (int32 AnimationAssetIndex = 0; AnimationAssetIndex < Database->GetNumAnimationAssets(); ++AnimationAssetIndex)
 	{
-		if (const FPoseSearchDatabaseAnimationAssetBase* DatabaseAsset = Database->GetDatabaseAnimationAsset<FPoseSearchDatabaseAnimationAssetBase>(AnimationAssetIndex))
+		if (const FPoseSearchDatabaseAnimationAssetBase* DatabaseAsset = Database->GetDatabaseAnimationAsset(AnimationAssetIndex))
 		{
 			if (const UObject* AnimationAsset = DatabaseAsset->GetAnimationAsset())
 			{
@@ -1912,7 +1912,7 @@ void FPoseSearchDatabaseAsyncCacheTask::OnGetComplete(UE::DerivedData::FCacheGet
 					{
 						for (int32 AnimationAssetIndex = 0; AnimationAssetIndex < DependentDatabase->GetNumAnimationAssets(); ++AnimationAssetIndex)
 						{
-							if (const FPoseSearchDatabaseAnimationAssetBase* DatabaseAsset = DependentDatabase->GetDatabaseAnimationAsset<FPoseSearchDatabaseAnimationAssetBase>(AnimationAssetIndex))
+							if (const FPoseSearchDatabaseAnimationAssetBase* DatabaseAsset = DependentDatabase->GetDatabaseAnimationAsset(AnimationAssetIndex))
 							{
 								if (DatabaseAsset->GetMirrorOption() == EPoseSearchMirrorOption::MirroredOnly || DatabaseAsset->GetMirrorOption() == EPoseSearchMirrorOption::UnmirroredAndMirrored)
 								{
@@ -1927,7 +1927,7 @@ void FPoseSearchDatabaseAsyncCacheTask::OnGetComplete(UE::DerivedData::FCacheGet
 					
 					for (int32 AnimationAssetIndex = 0; AnimationAssetIndex < DependentDatabase->GetNumAnimationAssets(); ++AnimationAssetIndex)
 					{
-						if (const FPoseSearchDatabaseAnimationAssetBase* DatabaseAsset = DependentDatabase->GetDatabaseAnimationAsset<FPoseSearchDatabaseAnimationAssetBase>(AnimationAssetIndex))
+						if (const FPoseSearchDatabaseAnimationAssetBase* DatabaseAsset = DependentDatabase->GetDatabaseAnimationAsset(AnimationAssetIndex))
 						{
 							if (DatabaseAsset->GetAnimationAsset() == nullptr)
 							{
@@ -2532,7 +2532,7 @@ void FAsyncPoseSearchDatabasesManagement::Tick(float DeltaTime)
 
 					for (int32 AnimationAssetIndex = 0; AnimationAssetIndex < Database->GetNumAnimationAssets(); ++AnimationAssetIndex)
 					{
-						if (const FPoseSearchDatabaseAnimationAssetBase* DatabaseAsset = Database->GetDatabaseAnimationAsset<FPoseSearchDatabaseAnimationAssetBase>(AnimationAssetIndex))
+						if (const FPoseSearchDatabaseAnimationAssetBase* DatabaseAsset = Database->GetDatabaseAnimationAsset(AnimationAssetIndex))
 						{
 							if (DatabaseAsset->IsEnabled())
 							{
diff --git a/Engine/Plugins/Animation/PoseSearch/Source/Runtime/Private/PoseSearchInteractionIsland.cpp b/Engine/Plugins/Animation/PoseSearch/Source/Runtime/Private/PoseSearchInteractionIsland.cpp
--- a/Engine/Plugins/Animation/PoseSearch/Source/Runtime/Private/PoseSearchInteractionIsland.cpp
+++ b/Engine/Plugins/Animation/PoseSearch/Source/Runtime/Private/PoseSearchInteractionIsland.cpp
@@ -1086,7 +1086,7 @@ bool FInteractionIsland::GetResult_AnyThread(const UObject* AnimContext, FPoseSe
 				const FSearchIndexAsset* SearchIndexAsset = SearchResult.GetSearchIndexAsset();
 				check(SearchIndexAsset);
 
-				const FPoseSearchDatabaseAnimationAssetBase* DatabaseAnimationAssetBase = Database->GetDatabaseAnimationAsset<FPoseSearchDatabaseAnimationAssetBase>(*SearchIndexAsset);
+				const FPoseSearchDatabaseAnimationAssetBase* DatabaseAnimationAssetBase = Database->GetDatabaseAnimationAsset(*SearchIndexAsset);
 				check(DatabaseAnimationAssetBase);
 
 				Result.SelectedAnim = DatabaseAnimationAssetBase->GetAnimationAsset();
diff --git a/Engine/Plugins/Animation/PoseSearch/Source/Runtime/Private/PoseSearchInteractionSubsystem.cpp b/Engine/Plugins/Animation/PoseSearch/Source/Runtime/Private/PoseSearchInteractionSubsystem.cpp
--- a/Engine/Plugins/Animation/PoseSearch/Source/Runtime/Private/PoseSearchInteractionSubsystem.cpp
+++ b/Engine/Plugins/Animation/PoseSearch/Source/Runtime/Private/PoseSearchInteractionSubsystem.cpp
@@ -415,7 +415,7 @@ void UPoseSearchInteractionSubsystem::PopulateContinuingProperties(float DeltaSe
 				{
 					if (const UE::PoseSearch::FSearchIndexAsset* SearchIndexAsset = SearchResult->GetSearchIndexAsset())
 					{
-						if (const FPoseSearchDatabaseAnimationAssetBase* DatabaseAsset = SearchResult->Database->GetDatabaseAnimationAsset<FPoseSearchDatabaseAnimationAssetBase>(*SearchIndexAsset))
+						if (const FPoseSearchDatabaseAnimationAssetBase* DatabaseAsset = SearchResult->Database->GetDatabaseAnimationAsset(*SearchIndexAsset))
 						{
 							check(SearchIndexAsset->GetToRealTimeFactor() > UE_KINDA_SMALL_NUMBER);
 							// in case DatabaseAsset->GetAnimationAsset() is a blendspace, SearchResult->AssetTime is a normalized time in the interval [0,1] 
diff --git a/Engine/Plugins/Animation/PoseSearch/Source/Runtime/Private/PoseSearchLibrary.cpp b/Engine/Plugins/Animation/PoseSearch/Source/Runtime/Private/PoseSearchLibrary.cpp
--- a/Engine/Plugins/Animation/PoseSearch/Source/Runtime/Private/PoseSearchLibrary.cpp
+++ b/Engine/Plugins/Animation/PoseSearch/Source/Runtime/Private/PoseSearchLibrary.cpp
@@ -205,37 +205,50 @@ namespace UE::PoseSearch
 	{
 		if (const UObject* PlayingAnimationAsset = ContinuingProperties.PlayingAsset.Get())
 		{
-			// checking if any of the AssetsToSearch (databases) or ContinuingProperties.PlayingAssetDatabase contain PlayingAnimationAsset
-			if (ContinuingProperties.PlayingAssetDatabase)
+			// checking if PlayingAnimationAsset can be considered or filtered out
+			bool bCanBeConsidered = true;
+			if (!SearchContext.GetAssetsToConsider().IsEmpty())
 			{
-				if (AddToSearchForDatabase(ContinuingPoseAssetsToSearchPerDatabase, PlayingAnimationAsset, ContinuingProperties.PlayingAssetDatabase, false))
+				if (!SearchContext.GetAssetsToConsider().Contains(PlayingAnimationAsset))
 				{
-#if WITH_EDITOR
-					SearchContext.SetAsyncBuildIndexInProgress();
-#endif // WITH_EDITOR
+					bCanBeConsidered = false;
 				}
 			}
 
-			for (const UObject* AssetToSearch : AssetsToSearch)
+			if (bCanBeConsidered)
 			{
-				if (const UPoseSearchDatabase* Database = Cast<UPoseSearchDatabase>(AssetToSearch))
+				// checking if any of the AssetsToSearch (databases) or ContinuingProperties.PlayingAssetDatabase contain PlayingAnimationAsset
+				if (ContinuingProperties.PlayingAssetDatabase)
 				{
-					// since it cannot be a database we can directly add it to ContinuingPoseAssetsToSearchPerDatabase
-					if (AddToSearchForDatabase(ContinuingPoseAssetsToSearchPerDatabase, PlayingAnimationAsset, Database, false))
+					if (AddToSearchForDatabase(ContinuingPoseAssetsToSearchPerDatabase, PlayingAnimationAsset, ContinuingProperties.PlayingAssetDatabase, false))
 					{
 #if WITH_EDITOR
 						SearchContext.SetAsyncBuildIndexInProgress();
 #endif // WITH_EDITOR
 					}
 				}
-			}
 
-			// checking if PlayingAnimationAsset has an associated database
-			if (AddToSearch(ContinuingPoseAssetsToSearchPerDatabase, PlayingAnimationAsset, bUsePoseSearchBranchIn))
-			{
+				for (const UObject* AssetToSearch : AssetsToSearch)
+				{
+					if (const UPoseSearchDatabase* Database = Cast<UPoseSearchDatabase>(AssetToSearch))
+					{
+						// since it cannot be a database we can directly add it to ContinuingPoseAssetsToSearchPerDatabase
+						if (AddToSearchForDatabase(ContinuingPoseAssetsToSearchPerDatabase, PlayingAnimationAsset, Database, false))
+						{
 #if WITH_EDITOR
-				SearchContext.SetAsyncBuildIndexInProgress();
+							SearchContext.SetAsyncBuildIndexInProgress();
 #endif // WITH_EDITOR
+						}
+					}
+				}
+
+				// checking if PlayingAnimationAsset has an associated database
+				if (AddToSearch(ContinuingPoseAssetsToSearchPerDatabase, PlayingAnimationAsset, bUsePoseSearchBranchIn))
+				{
+#if WITH_EDITOR
+					SearchContext.SetAsyncBuildIndexInProgress();
+#endif // WITH_EDITOR
+				}
 			}
 		}
 	}
@@ -251,6 +264,43 @@ namespace UE::PoseSearch
 #endif // WITH_EDITOR
 			}
 		}
+
+		// intersecting AssetsToSearchPerDatabase with asset to consider
+		if (!SearchContext.GetAssetsToConsider().IsEmpty())
+		{
+			for (FAssetsToSearchPerDatabasePair& AssetsToSearchPerDatabasePair : AssetsToSearchPerDatabase)
+			{
+				const UPoseSearchDatabase* Database = AssetsToSearchPerDatabasePair.Key;
+				check(Database);
+
+				FAssetsToSearch& AssetsToSearchForDatabase = AssetsToSearchPerDatabasePair.Value;
+				if (!AssetsToSearchForDatabase.IsEmpty())
+				{
+					// doing an intersection between AssetsToSearchForDatabase and asset to consider
+					FAssetsToSearch AssetsToSearchForDatabaseIntersection;
+					AssetsToSearchForDatabaseIntersection.Reserve(AssetsToSearchForDatabase.Num());
+					for (const UObject* AssetToSearchForDatabase : AssetsToSearchForDatabase)
+					{
+						if (SearchContext.GetAssetsToConsider().Contains(AssetToSearchForDatabase))
+						{
+							AssetsToSearchForDatabaseIntersection.Add(AssetToSearchForDatabase);
+						}
+					}
+					AssetsToSearchForDatabase = AssetsToSearchForDatabaseIntersection;
+				}
+				else
+				{
+					// since all the database assets can be searched, we need to add only the one we can consider:
+					for (const UObject* AssetToConsider : SearchContext.GetAssetsToConsider())
+					{
+						if (Database->Contains(AssetToConsider))
+						{
+							AssetsToSearchForDatabase.AddUnique(AssetToConsider);
+						}
+					}
+				}
+			}
+		}
 	}
 
 	// @todo: refine this logic. Currently if AssetsToSearch contains ONLY UPoseSearchDatabase we don't have to look for other databases referenced by other assets UAnimNotifyState_PoseSearchBranchIn(s)
@@ -669,7 +719,7 @@ void UPoseSearchLibrary::TraceMotionMatching(
 		const UPoseSearchDatabase* CurrentResultDatabase = SearchResult.Database.Get();
 		if (SearchIndexAsset && CurrentResultDatabase)
 		{
-			const FPoseSearchDatabaseAnimationAssetBase* DatabaseAsset = CurrentResultDatabase->GetDatabaseAnimationAsset<FPoseSearchDatabaseAnimationAssetBase>(*SearchIndexAsset);
+			const FPoseSearchDatabaseAnimationAssetBase* DatabaseAsset = CurrentResultDatabase->GetDatabaseAnimationAsset(*SearchIndexAsset);
 			check(DatabaseAsset);
 			if (UAnimationAsset* AnimationAsset = Cast<UAnimationAsset>(DatabaseAsset->GetAnimationAsset()))
 			{
@@ -1249,6 +1299,7 @@ UE::PoseSearch::FSearchResult UPoseSearchLibrary::MotionMatch(
 
 	FSearchResult SearchResult;
 
+	TConstArrayView<const UObject*> CurrentAssetToConsider = SearchContext.GetAssetsToConsider();
 	SearchContext.SetIsContinuingInteraction(ContinuingProperties.bIsContinuingInteraction);
 
 	// collecting all the databases searches in AssetsToSearchPerDatabase
@@ -1334,9 +1385,6 @@ UE::PoseSearch::FSearchResult UPoseSearchLibrary::MotionMatch(
 #endif // WITH_EDITOR && ENABLE_ANIM_DEBUG && UE_POSE_SEARCH_TRACE_ENABLED
 	}
 
-	// clean up the asset to consider array view
-	SearchContext.SetAssetsToConsider(TConstArrayView<const UObject*>());
-
 #if ENABLE_DRAW_DEBUG && ENABLE_ANIM_DEBUG
 	if (SearchResult.IsValid())
 	{
@@ -1375,6 +1423,9 @@ UE::PoseSearch::FSearchResult UPoseSearchLibrary::MotionMatch(
 	TraceMotionMatching(SearchContext, SearchResult, 0.f, DeltaSeconds, true, 1.f, EPoseSearchInterruptMode::DoNotInterrupt);
 #endif // UE_POSE_SEARCH_TRACE_ENABLED
 
+	// restoring the asset to consider array view
+	SearchContext.SetAssetsToConsider(CurrentAssetToConsider);
+
 	return SearchResult;
 }
 
diff --git a/Engine/Plugins/Animation/PoseSearch/Source/Runtime/Private/PoseSearchModule.cpp b/Engine/Plugins/Animation/PoseSearch/Source/Runtime/Private/PoseSearchModule.cpp
--- a/Engine/Plugins/Animation/PoseSearch/Source/Runtime/Private/PoseSearchModule.cpp
+++ b/Engine/Plugins/Animation/PoseSearch/Source/Runtime/Private/PoseSearchModule.cpp
@@ -81,7 +81,7 @@ public:
 		if (const UE::PoseSearch::FSearchIndexAsset* SearchIndexAsset = SearchResult.GetSearchIndexAsset())
 		{
 			const UPoseSearchDatabase* Database = SearchResult.Database.Get();
-			if (const FPoseSearchDatabaseAnimationAssetBase* DatabaseAnimationAssetBase = Database->GetDatabaseAnimationAsset<FPoseSearchDatabaseAnimationAssetBase>(*SearchIndexAsset))
+			if (const FPoseSearchDatabaseAnimationAssetBase* DatabaseAnimationAssetBase = Database->GetDatabaseAnimationAsset(*SearchIndexAsset))
 			{
 				ProviderResult.SelectedAsset = DatabaseAnimationAssetBase->GetAnimationAsset();
 				ProviderResult.Dissimilarity = SearchResult.PoseCost;
diff --git a/Engine/Plugins/Animation/PoseSearch/Source/Runtime/Private/PoseSearchResult.cpp b/Engine/Plugins/Animation/PoseSearch/Source/Runtime/Private/PoseSearchResult.cpp
--- a/Engine/Plugins/Animation/PoseSearch/Source/Runtime/Private/PoseSearchResult.cpp
+++ b/Engine/Plugins/Animation/PoseSearch/Source/Runtime/Private/PoseSearchResult.cpp
@@ -33,7 +33,7 @@ const UAnimationAsset* FDatabasePoseIdx::GetCurrentResultAnimationAsset() const
 {
 	if (const FSearchIndexAsset* SearchIndexAsset = GetSearchIndexAsset())
 	{
-		return Database->GetDatabaseAnimationAsset<FPoseSearchDatabaseAnimationAssetBase>(*SearchIndexAsset)->GetAnimationAssetForRole(Database->Schema->GetDefaultRole());
+		return Database->GetDatabaseAnimationAsset(*SearchIndexAsset)->GetAnimationAssetForRole(Database->Schema->GetDefaultRole());
 	}
 	return nullptr;
 }
@@ -42,7 +42,7 @@ const UAnimationAsset* FDatabasePoseIdx::GetCurrentResultAnimationAsset(const FR
 {
 	if (const FSearchIndexAsset* SearchIndexAsset = GetSearchIndexAsset())
 	{
-		return Database->GetDatabaseAnimationAsset<FPoseSearchDatabaseAnimationAssetBase>(*SearchIndexAsset)->GetAnimationAssetForRole(Role);
+		return Database->GetDatabaseAnimationAsset(*SearchIndexAsset)->GetAnimationAssetForRole(Role);
 	}
 	return nullptr;
 }
@@ -60,7 +60,7 @@ bool FSearchResult::DebugValidate() const
 	{
 		const FSearchIndex& SearchIndex = Database->GetSearchIndex();
 		const FSearchIndexAsset& SearchIndexAsset = SearchIndex.GetAssetForPose(PoseIdx);
-		if (const FPoseSearchDatabaseAnimationAssetBase* DatabaseAnimationAssetBase = Database->GetDatabaseAnimationAsset<FPoseSearchDatabaseAnimationAssetBase>(SearchIndexAsset))
+		if (const FPoseSearchDatabaseAnimationAssetBase* DatabaseAnimationAssetBase = Database->GetDatabaseAnimationAsset(SearchIndexAsset))
 		{
 			const float RealAssetTime = AssetTime * SearchIndexAsset.GetToRealTimeFactor();
 			const int32 RecalculatedPoseIdx = SearchIndexAsset.GetPoseIndexFromTime(RealAssetTime, Database->Schema->SampleRate);
@@ -91,7 +91,7 @@ void FSearchResult::UpdateWithNormalizedTime(float NormalizedTime)
 		const float RealTime = NormalizedTime * SearchIndexAsset.GetToRealTimeFactor();
 		
 #if WITH_EDITOR && DO_CHECK
-		const FPoseSearchDatabaseAnimationAssetBase* DatabaseAnimationAssetBase = Database->GetDatabaseAnimationAsset<FPoseSearchDatabaseAnimationAssetBase>(SearchIndexAsset);
+		const FPoseSearchDatabaseAnimationAssetBase* DatabaseAnimationAssetBase = Database->GetDatabaseAnimationAsset(SearchIndexAsset);
 		check(DatabaseAnimationAssetBase);
 		if (Cast<UBlendSpace>(DatabaseAnimationAssetBase->GetAnimationAssetForRole(Database->Schema->GetDefaultRole())))
 		{
@@ -136,7 +136,7 @@ void FSearchResult::UpdateWithRealTime(float RealTime)
 		const float NormalizedTime = RealTime / SearchIndexAsset.GetToRealTimeFactor();
 
 #if WITH_EDITOR && DO_CHECK
-		const FPoseSearchDatabaseAnimationAssetBase* DatabaseAnimationAssetBase = Database->GetDatabaseAnimationAsset<FPoseSearchDatabaseAnimationAssetBase>(SearchIndexAsset);
+		const FPoseSearchDatabaseAnimationAssetBase* DatabaseAnimationAssetBase = Database->GetDatabaseAnimationAsset(SearchIndexAsset);
 		check(DatabaseAnimationAssetBase);
 		if (Cast<UBlendSpace>(DatabaseAnimationAssetBase->GetAnimationAssetForRole(Database->Schema->GetDefaultRole())))
 		{
@@ -223,7 +223,7 @@ bool FPoseSearchBlueprintResult::InitFrom(const UE::PoseSearch::FSearchResult& S
 	{
 		const UPoseSearchDatabase* Database = SearchResult.Database.Get();
 
-		const FPoseSearchDatabaseAnimationAssetBase* DatabaseAsset = Database->GetDatabaseAnimationAsset<FPoseSearchDatabaseAnimationAssetBase>(*SearchIndexAsset);
+		const FPoseSearchDatabaseAnimationAssetBase* DatabaseAsset = Database->GetDatabaseAnimationAsset(*SearchIndexAsset);
 		check(DatabaseAsset);
 		
 #if WITH_EDITORONLY_DATA
diff --git a/Engine/Plugins/Animation/PoseSearch/Source/Runtime/Public/PoseSearch/Chooser/PoseSearchChooserColumn.h b/Engine/Plugins/Animation/PoseSearch/Source/Runtime/Public/PoseSearch/Chooser/PoseSearchChooserColumn.h
--- a/Engine/Plugins/Animation/PoseSearch/Source/Runtime/Public/PoseSearch/Chooser/PoseSearchChooserColumn.h
+++ b/Engine/Plugins/Animation/PoseSearch/Source/Runtime/Public/PoseSearch/Chooser/PoseSearchChooserColumn.h
@@ -81,8 +81,6 @@ struct FPoseSearchColumn : public FChooserColumnBase
 	UE_API virtual void AddToDetails(FInstancedPropertyBag& PropertyBag, int32 ColumnIndex, int32 RowIndex) override;
 	UE_API virtual void SetFromDetails(FInstancedPropertyBag& PropertyBag, int32 ColumnIndex, int32 RowIndex) override;
 
-	virtual bool AutoPopulates() const override { return true; }
-	UE_API virtual void AutoPopulate(int32 RowIndex, UObject* OutputObject) override;
 	virtual FName RowValuesPropertyName() override;
 	virtual void SetNumRows(int32 NumRows) override;
 	virtual void InsertRows(int Index, int Count) override;
@@ -95,6 +93,12 @@ struct FPoseSearchColumn : public FChooserColumnBase
 	
 #endif // WITH_EDITOR
 	virtual FChooserParameterBase* GetInputValue() override;
+
+private:
+#if WITH_EDITOR
+	UObject* GetReferencedObject(int32 RowIndex) const;
+#endif // WITH_EDITOR
+	void CheckConsistency() const;
 };
 
 #undef UE_API
diff --git a/Engine/Plugins/Animation/PoseSearch/Source/Runtime/Public/PoseSearch/PoseSearchDatabase.h b/Engine/Plugins/Animation/PoseSearch/Source/Runtime/Public/PoseSearch/PoseSearchDatabase.h
--- a/Engine/Plugins/Animation/PoseSearch/Source/Runtime/Public/PoseSearch/PoseSearchDatabase.h
+++ b/Engine/Plugins/Animation/PoseSearch/Source/Runtime/Public/PoseSearch/PoseSearchDatabase.h
@@ -27,6 +27,9 @@ class UPoseSearchNormalizationSet;
 namespace UE::PoseSearch
 {
 	struct FSearchContext;
+
+	UE_DEPRECATED(5.7, "UpgradeToDatabaseAnimationAsset will be removed once all the classes inheriting from FPoseSearchDatabaseAnimationAssetBase are removed")
+	UE_API FPoseSearchDatabaseAnimationAsset UpgradeToDatabaseAnimationAsset(FInstancedStruct& AnimationAsset);
 } // namespace UE::PoseSearch
 
 UENUM()
@@ -142,12 +145,99 @@ template<> struct TStructOpsTypeTraits<FPoseSearchDatabaseAnimationAssetBase> :
 	enum { WithIdenticalViaEquality = true };
 };
 
-/** A sequence entry in a UPoseSearchDatabase. */
 USTRUCT(BlueprintType, Category = "Animation|Pose Search")
-struct FPoseSearchDatabaseSequence : public FPoseSearchDatabaseAnimationAssetBase
+struct FPoseSearchDatabaseAnimationAsset : public FPoseSearchDatabaseAnimationAssetBase
+{
+	GENERATED_BODY()
+
+	UPROPERTY(EditAnywhere, Category = "Settings", meta = (DisplayPriority = 0))
+	TObjectPtr<UObject> AnimAsset;
+
+#if WITH_EDITORONLY_DATA
+
+	// If true this BlendSpace will output a single segment in the database.
+	UPROPERTY(EditAnywhere, Category = "Settings", meta = (DisplayPriority = 4))
+	bool bUseSingleSample = false;
+
+	// When turned on, this will use the set grid samples of the blend space asset for sampling. This will override the Number of Horizontal/Vertical Samples.
+	UPROPERTY(EditAnywhere, Category = "Settings", meta = (EditCondition = "!bUseSingleSample", EditConditionHides, DisplayPriority = 5))
+	bool bUseGridForSampling = false;
+
+	// Sets the number of horizontal samples in the blend space to pull the animation data coverage from. The larger the samples the more the data, but also the more memory and performance it takes.
+	UPROPERTY(EditAnywhere, Category = "Settings", meta = (EditCondition = "!bUseSingleSample && !bUseGridForSampling", EditConditionHides, ClampMin = "1", UIMin = "1", UIMax = "25", DisplayPriority = 6))
+	int32 NumberOfHorizontalSamples = 9;
+	
+	// Sets the number of vertical samples in the blend space to pull the animation data coverage from.The larger the samples the more the data, but also the more memory and performance it takes.
+	UPROPERTY(EditAnywhere, Category = "Settings", meta = (EditCondition = "!bUseSingleSample && !bUseGridForSampling", EditConditionHides, ClampMin = "1", UIMin = "1", UIMax = "25", DisplayPriority = 7))
+	int32 NumberOfVerticalSamples = 2;
+
+	// BlendParams used to sample this BlendSpace
+	UPROPERTY(EditAnywhere, Category = "Settings", meta = (EditCondition = "bUseSingleSample", EditConditionHides, DisplayPriority = 8))
+	float BlendParamX = 0.f;
+
+	// BlendParams used to sample this BlendSpace
+	UPROPERTY(EditAnywhere, Category = "Settings", meta = (EditCondition = "bUseSingleSample", EditConditionHides, DisplayPriority = 9))
+	float BlendParamY = 0.f;
+
+	// It allows users to set a time range to an individual blend space in the database.
+	// This is effectively trimming the beginning and end of the animation in the database (not in the original blend space).
+	// If set to [0, 0] it will be the entire frame range of the original blend space.
+	UPROPERTY(EditAnywhere, Category = "Settings", meta = (ClampToMinMaxLimits, DisplayPriority = 2))
+	FFloatInterval SamplingRange = FFloatInterval(0.f, 0.f);
+
+	UE_API virtual UClass* GetAnimationAssetStaticClass() const override;
+	UE_API static bool IsLooping(const UObject* AnimAsset, const FFloatInterval& SamplingRange);
+	UE_API virtual bool IsLooping() const override;
+
+	UE_API static bool IsRootMotionEnabled(const UObject* AnimAsset);
+	UE_API virtual bool IsRootMotionEnabled() const override;
+
+	virtual FFloatInterval GetSamplingRange() const override { return SamplingRange; }
+	UE_API virtual FFloatInterval GetEffectiveSamplingRange(const FVector& BlendParameters) const override;
+	UE_API virtual void IterateOverSamplingParameter(const TFunction<void(const FVector& BlendParameters)>& ProcessSamplingParameter) const override;
+	virtual void SetSamplingRange(const FFloatInterval& NewRange) override { SamplingRange = NewRange; }
+
+	UE_DEPRECATED(5.6, "Use IterateOverSamplingParameter to iterate over all the sampling paramters for the blend space")
+	UE_API void GetBlendSpaceParameterSampleRanges(int32& HorizontalBlendNum, int32& VerticalBlendNum) const;
+	UE_DEPRECATED(5.6, "Use IterateOverSamplingParameter to iterate over all the sampling paramters for the blend space")
+	UE_API FVector BlendParameterForSampleRanges(int32 HorizontalBlendIndex, int32 VerticalBlendIndex) const;
+#endif // WITH_EDITORONLY_DATA
+
+#if WITH_EDITOR
+	UE_API USkeletalMesh* GetPreviewMeshForRole(const UE::PoseSearch::FRole& Role) const;
+#endif // WITH_EDITOR
+
+	UE_API virtual UObject* GetAnimationAsset() const override;
+	UE_API virtual float GetPlayLength(const FVector& BlendParameters) const override;
+
+	friend bool operator==(const FPoseSearchDatabaseAnimationAsset& A, const FPoseSearchDatabaseAnimationAsset& B)
+	{
+		return static_cast<const FPoseSearchDatabaseAnimationAssetBase&>(A) == static_cast<const FPoseSearchDatabaseAnimationAssetBase&>(B)
+			&& A.AnimAsset == B.AnimAsset
+#if WITH_EDITORONLY_DATA
+			&& A.bUseSingleSample == B.bUseSingleSample
+			&& A.bUseGridForSampling == B.bUseGridForSampling
+			&& A.NumberOfHorizontalSamples == B.NumberOfHorizontalSamples
+			&& A.NumberOfVerticalSamples == B.NumberOfVerticalSamples
+			&& A.BlendParamX == B.BlendParamX
+			&& A.BlendParamY == B.BlendParamY
+			&& A.SamplingRange == B.SamplingRange
+#endif // WITH_EDITORONLY_DATA
+			;
+	}
+};
+
+template<> struct TStructOpsTypeTraits<FPoseSearchDatabaseAnimationAsset> : public TStructOpsTypeTraitsBase2<FPoseSearchDatabaseAnimationAsset>
+{
+	enum { WithIdenticalViaEquality = true };
+};
+
+
+// A sequence entry in a UPoseSearchDatabase.
+USTRUCT(BlueprintType, Category = "Animation|Pose Search")
+struct UE_DEPRECATED(5.7, "Use FPoseSearchDatabaseAnimationAsset instead. To convert your struct to support the new APIs use ") FPoseSearchDatabaseSequence : public FPoseSearchDatabaseAnimationAssetBase
 {
 	GENERATED_BODY()
-	virtual ~FPoseSearchDatabaseSequence() = default;
 
 	UPROPERTY(EditAnywhere, Category = "Settings", meta = (DisplayPriority = 0))
 	TObjectPtr<UAnimSequence> Sequence;
@@ -181,17 +271,18 @@ struct FPoseSearchDatabaseSequence : public FPoseSearchDatabaseAnimationAssetBas
 	}
 };
 
+PRAGMA_DISABLE_DEPRECATION_WARNINGS
 template<> struct TStructOpsTypeTraits<FPoseSearchDatabaseSequence> : public TStructOpsTypeTraitsBase2<FPoseSearchDatabaseSequence>
 {
 	enum { WithIdenticalViaEquality = true };
 };
+PRAGMA_ENABLE_DEPRECATION_WARNINGS
 
-/** An blend space entry in a UPoseSearchDatabase. */
+// An blend space entry in a UPoseSearchDatabase.
 USTRUCT(BlueprintType, Category = "Animation|Pose Search")
-struct FPoseSearchDatabaseBlendSpace : public FPoseSearchDatabaseAnimationAssetBase
+struct UE_DEPRECATED(5.7, "Use FPoseSearchDatabaseAnimationAsset instead. To convert your struct to support the new APIs use ") FPoseSearchDatabaseBlendSpace : public FPoseSearchDatabaseAnimationAssetBase
 {
 	GENERATED_BODY()
-	virtual ~FPoseSearchDatabaseBlendSpace() = default;
 
 	UPROPERTY(EditAnywhere, Category = "Settings", meta = (DisplayPriority = 0))
 	TObjectPtr<UBlendSpace> BlendSpace;
@@ -262,17 +353,18 @@ struct FPoseSearchDatabaseBlendSpace : public FPoseSearchDatabaseAnimationAssetB
 	}
 };
 
+PRAGMA_DISABLE_DEPRECATION_WARNINGS
 template<> struct TStructOpsTypeTraits<FPoseSearchDatabaseBlendSpace> : public TStructOpsTypeTraitsBase2<FPoseSearchDatabaseBlendSpace>
 {
 	enum { WithIdenticalViaEquality = true };
 };
+PRAGMA_ENABLE_DEPRECATION_WARNINGS
 
-/** An entry in a UPoseSearchDatabase. */
+// An entry in a UPoseSearchDatabase.
 USTRUCT(BlueprintType, Category = "Animation|Pose Search")
-struct FPoseSearchDatabaseAnimComposite : public FPoseSearchDatabaseAnimationAssetBase
+struct UE_DEPRECATED(5.7, "Use FPoseSearchDatabaseAnimationAsset instead. To convert your struct to support the new APIs use ") FPoseSearchDatabaseAnimComposite : public FPoseSearchDatabaseAnimationAssetBase
 {
 	GENERATED_BODY()
-	virtual ~FPoseSearchDatabaseAnimComposite() = default;
 
 	UPROPERTY(EditAnywhere, Category = "Settings", meta = (DisplayPriority = 0))
 	TObjectPtr<UAnimComposite> AnimComposite;
@@ -306,17 +398,18 @@ struct FPoseSearchDatabaseAnimComposite : public FPoseSearchDatabaseAnimationAss
 	}
 };
 
+PRAGMA_DISABLE_DEPRECATION_WARNINGS
 template<> struct TStructOpsTypeTraits<FPoseSearchDatabaseAnimComposite> : public TStructOpsTypeTraitsBase2<FPoseSearchDatabaseAnimComposite>
 {
 	enum { WithIdenticalViaEquality = true };
 };
+PRAGMA_ENABLE_DEPRECATION_WARNINGS
 
-/** An anim montage entry in a UPoseSearchDatabase. */
+// An anim montage entry in a UPoseSearchDatabase.
 USTRUCT(BlueprintType, Category = "Animation|Pose Search")
-struct FPoseSearchDatabaseAnimMontage : public FPoseSearchDatabaseAnimationAssetBase
+struct UE_DEPRECATED(5.7, "Use FPoseSearchDatabaseAnimationAsset instead. To convert your struct to support the new APIs use ") FPoseSearchDatabaseAnimMontage : public FPoseSearchDatabaseAnimationAssetBase
 {
 	GENERATED_BODY()
-	virtual ~FPoseSearchDatabaseAnimMontage() = default;
 
 	UPROPERTY(EditAnywhere, Category = "Settings", meta = (DisplayPriority = 0))
 	TObjectPtr<UAnimMontage> AnimMontage;
@@ -350,16 +443,17 @@ struct FPoseSearchDatabaseAnimMontage : public FPoseSearchDatabaseAnimationAsset
 	}
 };
 
+PRAGMA_DISABLE_DEPRECATION_WARNINGS
 template<> struct TStructOpsTypeTraits<FPoseSearchDatabaseAnimMontage> : public TStructOpsTypeTraitsBase2<FPoseSearchDatabaseAnimMontage>
 {
 	enum { WithIdenticalViaEquality = true };
 };
+PRAGMA_ENABLE_DEPRECATION_WARNINGS
 
 USTRUCT(Experimental, BlueprintType, Category = "Animation|Pose Search")
-struct FPoseSearchDatabaseMultiAnimAsset : public FPoseSearchDatabaseAnimationAssetBase
+struct UE_DEPRECATED(5.7, "Use FPoseSearchDatabaseAnimationAsset instead. To convert your struct to support the new APIs use ") FPoseSearchDatabaseMultiAnimAsset : public FPoseSearchDatabaseAnimationAssetBase
 {
 	GENERATED_BODY()
-	virtual ~FPoseSearchDatabaseMultiAnimAsset() = default;
 
 	UPROPERTY(EditAnywhere, Category = "Settings", meta = (DisplayPriority = 0))
 	TObjectPtr<UMultiAnimAsset> MultiAnimAsset;
@@ -411,12 +505,14 @@ struct FPoseSearchDatabaseMultiAnimAsset : public FPoseSearchDatabaseAnimationAs
 	}
 };
 
+PRAGMA_DISABLE_DEPRECATION_WARNINGS
 template<> struct TStructOpsTypeTraits<FPoseSearchDatabaseMultiAnimAsset> : public TStructOpsTypeTraitsBase2<FPoseSearchDatabaseMultiAnimAsset>
 {
 	enum { WithIdenticalViaEquality = true };
 };
+PRAGMA_ENABLE_DEPRECATION_WARNINGS
 
-/** A data asset for indexing a collection of animation sequences. */
+// A data asset for indexing a collection of animation sequences.
 UCLASS(MinimalAPI, BlueprintType, Category = "Animation|Pose Search", meta = (DisplayName = "Pose Search Database"))
 class UPoseSearchDatabase : public UDataAsset
 {
@@ -464,8 +560,12 @@ public:
 #endif // WITH_EDITORONLY_DATA
 
 private:
+	UE_DEPRECATED(5.7, "Please reference DatabaseAnimationAssets instead")
 	UPROPERTY()
-	TArray<FInstancedStruct> AnimationAssets;
+	TArray<FInstancedStruct> AnimationAssets_DEPRECATED;
+
+	UPROPERTY()
+	TArray<FPoseSearchDatabaseAnimationAsset> DatabaseAnimationAssets;
 
 	// Experimental, this feature might be removed without warning, not for production use
 	// if Chooser is valid, all the AnimationAssets will come from it and AnimationAssets is not gonna be used AT ALL
@@ -559,15 +659,23 @@ public:
 	// It'll return INDEX_NONE if it cannot find a pose
 	UE_API int32 GetPoseIndex(const UObject* AnimationAsset, float AnimationAssetTime, bool bMirrored, const FVector& BlendParameters) const;
 
+	UE_DEPRECATED(5.7, "Use AddAnimationAsset(const FPoseSearchDatabaseAnimationAsset& AnimationAsset) instead")
 	UE_API void AddAnimationAsset(FInstancedStruct AnimationAsset);
+
+	UE_API void AddAnimationAsset(const FPoseSearchDatabaseAnimationAsset& AnimationAsset);
 	// Experimental, this feature might be removed without warning, not for production use
-	UE_API void InsertAnimationAssetAt(FInstancedStruct AnimationAsset, int32 AnimationAssetIndex);
+	UE_API void InsertAnimationAssetAt(const FPoseSearchDatabaseAnimationAsset& AnimationAsset, int32 AnimationAssetIndex);
 	UE_API void RemoveAnimationAssetAt(int32 AnimationAssetIndex);
 
-	template<typename TDatabaseAnimationAsset> const TDatabaseAnimationAsset* GetDatabaseAnimationAsset(int32 AnimationAssetIndex) const;
-	template<typename TDatabaseAnimationAsset> const TDatabaseAnimationAsset* GetDatabaseAnimationAsset(const UE::PoseSearch::FSearchIndexAsset& SearchIndexAsset) const;
-	template<typename TDatabaseAnimationAsset> TDatabaseAnimationAsset* GetMutableDatabaseAnimationAsset(int32 AnimationAssetIndex);
-	template<typename TDatabaseAnimationAsset> TDatabaseAnimationAsset* GetMutableDatabaseAnimationAsset(const UE::PoseSearch::FSearchIndexAsset& SearchIndexAsset);
+	UE_API const FPoseSearchDatabaseAnimationAsset* GetDatabaseAnimationAsset(int32 AnimationAssetIndex) const;
+	UE_API const FPoseSearchDatabaseAnimationAsset* GetDatabaseAnimationAsset(const UE::PoseSearch::FSearchIndexAsset& SearchIndexAsset) const;
+	UE_API FPoseSearchDatabaseAnimationAsset* GetMutableDatabaseAnimationAsset(int32 AnimationAssetIndex);
+	UE_API FPoseSearchDatabaseAnimationAsset* GetMutableDatabaseAnimationAsset(const UE::PoseSearch::FSearchIndexAsset& SearchIndexAsset);
+
+	template<typename TDatabaseAnimationAsset> UE_DEPRECATED(5.7, "Use the non templated version of this API") const TDatabaseAnimationAsset* GetDatabaseAnimationAsset(int32 AnimationAssetIndex) const;
+	template<typename TDatabaseAnimationAsset> UE_DEPRECATED(5.7, "Use the non templated version of this API") const TDatabaseAnimationAsset* GetDatabaseAnimationAsset(const UE::PoseSearch::FSearchIndexAsset& SearchIndexAsset) const;
+	template<typename TDatabaseAnimationAsset> UE_DEPRECATED(5.7, "Use the non templated version of this API") TDatabaseAnimationAsset* GetMutableDatabaseAnimationAsset(int32 AnimationAssetIndex);
+	template<typename TDatabaseAnimationAsset> UE_DEPRECATED(5.7, "Use the non templated version of this API") TDatabaseAnimationAsset* GetMutableDatabaseAnimationAsset(const UE::PoseSearch::FSearchIndexAsset& SearchIndexAsset);
 
 	UE_API float GetRealAssetTime(int32 PoseIdx) const;
 	UE_API float GetNormalizedAssetTime(int32 PoseIdx) const;
@@ -649,33 +757,49 @@ private:
 template<typename TDatabaseAnimationAsset>
 inline const TDatabaseAnimationAsset* UPoseSearchDatabase::GetDatabaseAnimationAsset(int32 AnimationAssetIndex) const
 {
-	if (AnimationAssets.IsValidIndex(AnimationAssetIndex))
+	if (std::is_same<TDatabaseAnimationAsset, FPoseSearchDatabaseAnimationAsset>::value ||
+		std::is_same<TDatabaseAnimationAsset, FPoseSearchDatabaseAnimationAssetBase>::value)
 	{
-		return AnimationAssets[AnimationAssetIndex].GetPtr<TDatabaseAnimationAsset>();
+		if (DatabaseAnimationAssets.IsValidIndex(AnimationAssetIndex))
+		{
+			// using reinterpret_cast to be able to compile this method even with other types (that will return null)
+			return reinterpret_cast<const TDatabaseAnimationAsset*>(&DatabaseAnimationAssets[AnimationAssetIndex]);
+		}
 	}
+	
 	return nullptr;
 }
 
 template<typename TDatabaseAnimationAsset>
 inline const TDatabaseAnimationAsset* UPoseSearchDatabase::GetDatabaseAnimationAsset(const UE::PoseSearch::FSearchIndexAsset& SearchIndexAsset) const
 {
+	PRAGMA_DISABLE_DEPRECATION_WARNINGS
 	return GetDatabaseAnimationAsset<TDatabaseAnimationAsset>(SearchIndexAsset.GetSourceAssetIdx());
+	PRAGMA_ENABLE_DEPRECATION_WARNINGS
 }
 
 template<typename TDatabaseAnimationAsset>
 inline TDatabaseAnimationAsset* UPoseSearchDatabase::GetMutableDatabaseAnimationAsset(int32 AnimationAssetIndex)
 {
-	if (AnimationAssets.IsValidIndex(AnimationAssetIndex))
+	if (std::is_same<TDatabaseAnimationAsset, FPoseSearchDatabaseAnimationAsset>::value ||
+		std::is_same<TDatabaseAnimationAsset, FPoseSearchDatabaseAnimationAssetBase>::value)
 	{
-		return AnimationAssets[AnimationAssetIndex].GetMutablePtr<TDatabaseAnimationAsset>();
+		if (DatabaseAnimationAssets.IsValidIndex(AnimationAssetIndex))
+		{
+			// using reinterpret_cast to be able to compile this method even with other types (that will return null)
+			return reinterpret_cast<TDatabaseAnimationAsset*>(&DatabaseAnimationAssets[AnimationAssetIndex]);
+		}
 	}
+
 	return nullptr;
 }
 
 template<typename TDatabaseAnimationAsset>
 inline TDatabaseAnimationAsset* UPoseSearchDatabase::GetMutableDatabaseAnimationAsset(const UE::PoseSearch::FSearchIndexAsset& SearchIndexAsset)
 {
+	PRAGMA_DISABLE_DEPRECATION_WARNINGS
 	return GetMutableDatabaseAnimationAsset<TDatabaseAnimationAsset>(SearchIndexAsset.GetSourceAssetIdx());
+	PRAGMA_ENABLE_DEPRECATION_WARNINGS
 }
 
 #undef UE_API
diff --git a/Engine/Plugins/Experimental/MovieScenePoseSearchTracks/Source/MovieScenePoseSearchTracks/Private/Sections/MovieSceneStitchAnimSection.cpp b/Engine/Plugins/Experimental/MovieScenePoseSearchTracks/Source/MovieScenePoseSearchTracks/Private/Sections/MovieSceneStitchAnimSection.cpp
--- a/Engine/Plugins/Experimental/MovieScenePoseSearchTracks/Source/MovieScenePoseSearchTracks/Private/Sections/MovieSceneStitchAnimSection.cpp
+++ b/Engine/Plugins/Experimental/MovieScenePoseSearchTracks/Source/MovieScenePoseSearchTracks/Private/Sections/MovieSceneStitchAnimSection.cpp
@@ -124,7 +124,7 @@ void FMovieSceneStitchAnimEvaluationTask::Execute(UE::UAF::FEvaluationVM& VM) co
 			if (const FSearchIndexAsset* SearchIndexAsset = SearchResult.GetSearchIndexAsset())
 			{
 				const UPoseSearchDatabase* Database = SearchResult.Database.Get();
-				if (const FPoseSearchDatabaseAnimationAssetBase* DatabaseAnimationAssetBase = Database->GetDatabaseAnimationAsset<FPoseSearchDatabaseAnimationAssetBase>(*SearchIndexAsset))
+				if (const FPoseSearchDatabaseAnimationAssetBase* DatabaseAnimationAssetBase = Database->GetDatabaseAnimationAsset(*SearchIndexAsset))
 				{
 					MatchedAsset = Cast<UAnimationAsset>(DatabaseAnimationAssetBase->GetAnimationAsset());
 					MatchedAssetTime = SearchResult.AssetTime;
