PoseSearch - fix for crash in FAsyncPoseSearchDatabasesManagement::CollectDatabasesToSynchronize with Objects with nullptr Object->GetClass()

diff --git a/Engine/Plugins/Animation/PoseSearch/Source/Runtime/Private/PoseSearchDerivedData.cpp b/Engine/Plugins/Animation/PoseSearch/Source/Runtime/Private/PoseSearchDerivedData.cpp
--- a/Engine/Plugins/Animation/PoseSearch/Source/Runtime/Private/PoseSearchDerivedData.cpp
+++ b/Engine/Plugins/Animation/PoseSearch/Source/Runtime/Private/PoseSearchDerivedData.cpp
@@ -2291,7 +2291,7 @@ FAsyncPoseSearchDatabasesManagement::~FAsyncPoseSearchDatabasesManagement()
 // given Object it figures out a map of databases to UAnimSequenceBase(s) containing UAnimNotifyState_PoseSearchBranchIn
 void FAsyncPoseSearchDatabasesManagement::CollectDatabasesToSynchronize(UObject* Object)
 {
-	if (!Object)
+	if (!Object || !Object->GetClass())
 	{
 		return;
 	}
