PoseSearch - temp fix for unresolved external symbol "__declspec(dllimport) public: void __cdecl UE::PoseSearch::FAnimationAssetSampler::`default constructor closure'(void)"

diff --git a/Engine/Plugins/Animation/PoseSearch/Source/Editor/Public/PoseSearchInteractionAssetEditor.h b/Engine/Plugins/Animation/PoseSearch/Source/Editor/Public/PoseSearchInteractionAssetEditor.h
--- a/Engine/Plugins/Animation/PoseSearch/Source/Editor/Public/PoseSearchInteractionAssetEditor.h
+++ b/Engine/Plugins/Animation/PoseSearch/Source/Editor/Public/PoseSearchInteractionAssetEditor.h
@@ -39,8 +39,8 @@ public:
 
 private:
 	enum { PreviewActor, DebugActor, NumActors };
-	TStaticArray<TWeakObjectPtr<AActor>, NumActors> ActorPtrs;
-	TStaticArray<UE::PoseSearch::FAnimationAssetSampler, NumActors> Samplers; // used to keep track of the animation position / orientation
+	TArray<TWeakObjectPtr<AActor>> ActorPtrs;
+	TArray<UE::PoseSearch::FAnimationAssetSampler> Samplers; // used to keep track of the animation position / orientation
 
 	float CurrentTime = 0.f;
 	UE::PoseSearch::FRole ActorRole = UE::PoseSearch::DefaultRole;
