PoseSearch - temporary back out of FAsyncPoseSearchDatabasesManagement::OnObjectTransacted from using FCoreUObjectDelegates::ObjectsModifiedThisFrame since it could potentially reference uninitialized memory

diff --git a/Engine/Plugins/Animation/PoseSearch/Source/Runtime/Private/PoseSearchDerivedData.cpp b/Engine/Plugins/Animation/PoseSearch/Source/Runtime/Private/PoseSearchDerivedData.cpp
--- a/Engine/Plugins/Animation/PoseSearch/Source/Runtime/Private/PoseSearchDerivedData.cpp
+++ b/Engine/Plugins/Animation/PoseSearch/Source/Runtime/Private/PoseSearchDerivedData.cpp
@@ -2409,26 +2409,9 @@ void FAsyncPoseSearchDatabasesManagement::PostModified(UObject* Object)
 	ClearPreCancelled();
 }
 
-void FAsyncPoseSearchDatabasesManagement::OnObjectTransacted(UObject*, const FTransactionObjectEvent&)
+void FAsyncPoseSearchDatabasesManagement::OnObjectTransacted(UObject* Object, const FTransactionObjectEvent& TransactionObjectEvent)
 {
-	check(IsInGameThread());
-
-	FScopeLock Lock(&Mutex);
-
-	// collecting databases to synchronize, and merging the results with the PreModified collection
-	for (UObject* ObjectModifiedThisFrame : FCoreUObjectDelegates::ObjectsModifiedThisFrame)
-	{
-		CollectDatabasesToSynchronize(ObjectModifiedThisFrame);
-	}
-
-	SynchronizeDatabases();
-
-	for (UObject* ObjectModifiedThisFrame : FCoreUObjectDelegates::ObjectsModifiedThisFrame)
-	{
-		SynchronizeDatabaseChooser(ObjectModifiedThisFrame);
-	}
-
-	ClearPreCancelled();
+	PostModified(Object);
 }
 
 void FAsyncPoseSearchDatabasesManagement::OnPackageReloaded(const EPackageReloadPhase InPackageReloadPhase, FPackageReloadedEvent* InPackageReloadedEvent)
