PoseSearch - removing PSD internal chooser (PoseSearchColumn is a better approach)

diff --git a/Engine/Plugins/Animation/PoseSearch/Source/Runtime/Private/PoseSearchDatabase.cpp b/Engine/Plugins/Animation/PoseSearch/Source/Runtime/Private/PoseSearchDatabase.cpp
--- a/Engine/Plugins/Animation/PoseSearch/Source/Runtime/Private/PoseSearchDatabase.cpp
+++ b/Engine/Plugins/Animation/PoseSearch/Source/Runtime/Private/PoseSearchDatabase.cpp
@@ -11,8 +11,6 @@
 #include "Animation/AnimSequence.h"
 #include "Animation/BlendSpace.h"
 #include "Animation/BlendSpace1D.h"
-#include "Chooser/Internal/Chooser.h"
-#include "IObjectChooser.h"
 #include "PoseSearch/MultiAnimAsset.h"
 #include "PoseSearch/PoseSearchAnimNotifies.h"
 #include "PoseSearch/PoseSearchContext.h"
@@ -1247,8 +1245,6 @@ void UPoseSearchDatabase::PostLoad()
 
 #if WITH_EDITOR
 
-	// todo: should we SynchronizeWithExternalDependencies() here?
-	SynchronizeChooser();
 
 	// Delay indexing until BeginCacheForCookedPlatformData when running a CookCommandlet
 	if (!IsRunningCookCommandlet())
@@ -1433,71 +1429,6 @@ void UPoseSearchDatabase::SynchronizeWithExternalDependencies(TConstArrayView<UA
 	}
 }
 
-const UChooserTable* UPoseSearchDatabase::GetChooser() const
-{
-	return Chooser;
-}
-
-void UPoseSearchDatabase::SynchronizeChooser()
-{
-	if (Chooser)
-	{
-		TArray<UObject*, TInlineAllocator<128>> ChooserAssets;
-		UChooserTable::IterateChooser(Chooser, FObjectChooserBase::FObjectChooserIteratorCallback::CreateLambda([&ChooserAssets](UObject* Object)
-			{
-				if (Object)
-				{
-					if (Object->GetClass()->IsChildOf(UAnimSequence::StaticClass()) ||
-						Object->GetClass()->IsChildOf(UAnimComposite::StaticClass()) ||
-						Object->GetClass()->IsChildOf(UAnimMontage::StaticClass()) ||
-						Object->GetClass()->IsChildOf(UBlendSpace::StaticClass()) ||
-						Object->GetClass()->IsChildOf(UMultiAnimAsset::StaticClass()))
-					{
-						ChooserAssets.Add(Object);
-					}
-				}
-				return FObjectChooserBase::EIteratorStatus::Continue;
-			}));
-
-		TArray<UObject*, TInlineAllocator<128>> DatabaseAssets;
-		for (FPoseSearchDatabaseAnimationAsset& DatabaseAnimationAsset : DatabaseAnimationAssets)
-		{
-			DatabaseAssets.Add(DatabaseAnimationAsset.GetAnimationAsset());
-		}
-
-		if (ChooserAssets != DatabaseAssets)
-		{
-			TArray<FPoseSearchDatabaseAnimationAsset> DatabaseAnimationAssetsBackup = MoveTemp(DatabaseAnimationAssets);
-		
-			for (UObject* ChooserAsset : ChooserAssets)
-			{
-				// searching ChooserAsset in AnimationAssetsBackup
-				bool bFoundInAnimationAssetsBackup = false;
-				for (int32 DatabaseAnimationAssetBackupIndex = 0; DatabaseAnimationAssetBackupIndex < DatabaseAnimationAssetsBackup.Num(); ++DatabaseAnimationAssetBackupIndex)
-				{
-					if (DatabaseAnimationAssetsBackup[DatabaseAnimationAssetBackupIndex].GetAnimationAsset() == ChooserAsset)
-					{
-						bFoundInAnimationAssetsBackup = true;
-						DatabaseAnimationAssets.Add(DatabaseAnimationAssetsBackup[DatabaseAnimationAssetBackupIndex]);
-						DatabaseAnimationAssetsBackup.RemoveAt(DatabaseAnimationAssetBackupIndex);
-						break;
-					}
-				}
-
-				if (!bFoundInAnimationAssetsBackup)
-				{
-					FPoseSearchDatabaseAnimationAsset NewDatabaseAnimationAsset;
-					NewDatabaseAnimationAsset.AnimAsset = ChooserAsset;
-					DatabaseAnimationAssets.Add(NewDatabaseAnimationAsset);
-				}
-			}
-
-			Modify();
-			NotifySynchronizeWithExternalDependencies();
-		}
-	}
-}
-
 void UPoseSearchDatabase::BeginCacheForCookedPlatformData(const ITargetPlatform* TargetPlatform)
 {
 	using namespace UE::PoseSearch;
@@ -1604,27 +1535,6 @@ UE::PoseSearch::FSearchResult UPoseSearchDatabase::Search(UE::PoseSearch::FSearc
 	}
 #endif // WITH_EDITOR
 
-	// updating the SearchContext::AssetsToConsider from the UPoseSearchDatabase::Chooser evaluation 
-	// (and restore them with the original one CurrentAssetToConsider at the end of the method)
-	TConstArrayView<const UObject*> CurrentAssetToConsider = SearchContext.GetAssetsToConsider();
-	TArray<const UObject*, TInlineAllocator<128>> ChooserAssetToConsider;
-	if (Chooser && !SearchContext.GetContexts().IsEmpty())
-	{
-		if (!CurrentAssetToConsider.IsEmpty())
-		{
-			UE_LOG(LogPoseSearch, Warning, TEXT("UPoseSearchDatabase::Search - database (%s) has an associated Chooser filtering the search down, so previous SearchContext.GetAssetsToConsider() context will be skipped"), *GetName());
-		}
-
-		// @todo: add MM interaction support. Maybe iterate over all the SearchContext.GetAnimContexts and use the union of the ChosenAssets?
-		UChooserTable::EvaluateChooser(*SearchContext.GetContexts()[0], Chooser, FObjectChooserBase::FObjectChooserIteratorCallback::CreateLambda([this, &ChooserAssetToConsider](UObject* InResult)
-			{
-				ChooserAssetToConsider.Add(InResult);
-				return FObjectChooserBase::EIteratorStatus::Continue;
-			}));
-
-		SearchContext.SetAssetsToConsider(ChooserAssetToConsider);
-	}
-
 	if (SearchContext.GetEventToSearch().IsValid())
 	{
 		Result = SearchEvent(SearchContext);
@@ -1666,8 +1576,6 @@ UE::PoseSearch::FSearchResult UPoseSearchDatabase::Search(UE::PoseSearch::FSearc
 	SearchContext.Track(this);
 #endif // UE_POSE_SEARCH_TRACE_ENABLED
 
-	SearchContext.SetAssetsToConsider(CurrentAssetToConsider);
-
 	return Result;
 }
 
diff --git a/Engine/Plugins/Animation/PoseSearch/Source/Runtime/Private/PoseSearchDerivedData.cpp b/Engine/Plugins/Animation/PoseSearch/Source/Runtime/Private/PoseSearchDerivedData.cpp
--- a/Engine/Plugins/Animation/PoseSearch/Source/Runtime/Private/PoseSearchDerivedData.cpp
+++ b/Engine/Plugins/Animation/PoseSearch/Source/Runtime/Private/PoseSearchDerivedData.cpp
@@ -9,7 +9,6 @@
 #include "Animation/BlendSpace.h"
 #include "AssetRegistry/ARFilter.h"
 #include "AssetRegistry/AssetRegistryModule.h"
-#include "Chooser/Internal/Chooser.h"
 #include "DerivedDataCache.h"
 #include "DerivedDataRequestOwner.h"
 #include "Interfaces/ITargetPlatformManagerModule.h"
@@ -1253,41 +1252,6 @@ static bool ValidateSearchInfoAgainstDatabase(const FSearchIndex& SearchIndex, c
 	return true;
 }
 
-static void SynchronizeDatabaseChooser(UObject* Object)
-{
-	if (UPoseSearchDatabase* Database = Cast<UPoseSearchDatabase>(Object))
-	{
-		Database->SynchronizeChooser();
-	}
-	else if (const UChooserTable* ChooserTable = Cast<UChooserTable>(Object))
-	{
-		const UChooserTable* RootChooser = ChooserTable->GetRootChooser();
-
-		IAssetRegistry& AssetRegistry = FModuleManager::LoadModuleChecked<FAssetRegistryModule>("AssetRegistry").Get();
-		TArray<FAssetIdentifier> Referencers;
-		AssetRegistry.GetReferencers(RootChooser->GetPackage()->GetFName(), Referencers);
-
-		TArray<FAssetData> Assets;
-		Assets.Reserve(256);
-		for (const FAssetIdentifier& Referencer : Referencers)
-		{
-			Assets.Reset();
-			AssetRegistry.GetAssetsByPackageName(Referencer.PackageName, Assets);
-
-			for (const FAssetData& Asset : Assets)
-			{
-				if (Asset.IsInstanceOf(UPoseSearchDatabase::StaticClass()))
-				{
-					if (UPoseSearchDatabase* ReferencedDatabase = CastChecked<UPoseSearchDatabase>(Asset.FastGetAsset(true)))
-					{
-						ReferencedDatabase->SynchronizeChooser();
-					}
-				}
-			}
-		}
-	}
-}
-
 #if ENABLE_ANIM_DEBUG
 
 static void CompareChannelValues(int32 RecursionIndex, int32 PoseIndex, const TConstArrayView<float>& PoseA, const TConstArrayView<float>& PoseB, const TConstArrayView<TObjectPtr<UPoseSearchFeatureChannel>>& Channels, FStringBuilderBase& StringBuilder)
@@ -2404,8 +2368,6 @@ void FAsyncPoseSearchDatabasesManagement::PostModified(UObject* Object)
 
 	SynchronizeDatabases();
 
-	SynchronizeDatabaseChooser(Object);
-
 	ClearPreCancelled();
 }
 
diff --git a/Engine/Plugins/Animation/PoseSearch/Source/Runtime/Public/PoseSearch/PoseSearchDatabase.h b/Engine/Plugins/Animation/PoseSearch/Source/Runtime/Public/PoseSearch/PoseSearchDatabase.h
--- a/Engine/Plugins/Animation/PoseSearch/Source/Runtime/Public/PoseSearch/PoseSearchDatabase.h
+++ b/Engine/Plugins/Animation/PoseSearch/Source/Runtime/Public/PoseSearch/PoseSearchDatabase.h
@@ -17,7 +17,6 @@ class UAnimationAsset;
 class UAnimComposite;
 class UAnimMontage;
 class UBlendSpace;
-class UChooserTable;
 class UMultiAnimAsset;
 
 #if WITH_EDITORONLY_DATA
@@ -567,10 +566,6 @@ private:
 	UPROPERTY()
 	TArray<FPoseSearchDatabaseAnimationAsset> DatabaseAnimationAssets;
 
-	// Experimental, this feature might be removed without warning, not for production use
-	// if Chooser is valid, all the AnimationAssets will come from it and AnimationAssets is not gonna be used AT ALL
-	UPROPERTY(EditAnywhere, Category = "Database|Experimental")
-	TObjectPtr<UChooserTable> Chooser;
 
 public:
 	/** Array of tags that can be used as metadata. */
@@ -716,11 +711,6 @@ public:
 
 	UE_API void SynchronizeWithExternalDependencies();
 	UE_API void SynchronizeWithExternalDependencies(TConstArrayView<UAnimSequenceBase*> SequencesBase);
-	
-	// Experimental, this feature might be removed without warning, not for production use
-	UE_API void SynchronizeChooser();
-	// Experimental, this feature might be removed without warning, not for production use
-	UE_API const UChooserTable* GetChooser() const;
 #endif // WITH_EDITOR
 
 #if WITH_EDITOR && ENABLE_ANIM_DEBUG
