PoseSearch - PSD synchronization fix

diff --git a/Engine/Plugins/Animation/PoseSearch/Source/Runtime/Private/Chooser/PoseSearchChooserColumn.cpp b/Engine/Plugins/Animation/PoseSearch/Source/Runtime/Private/Chooser/PoseSearchChooserColumn.cpp
--- a/Engine/Plugins/Animation/PoseSearch/Source/Runtime/Private/Chooser/PoseSearchChooserColumn.cpp
+++ b/Engine/Plugins/Animation/PoseSearch/Source/Runtime/Private/Chooser/PoseSearchChooserColumn.cpp
@@ -106,6 +106,8 @@ void FPoseSearchColumn::SetNumRows(int32 NumRows)
 		}
 	}
 
+	InternalDatabase->NotifySynchronizeWithExternalDependencies();
+
 	// NumRows == 0 has a special meaning of resetting the column and usually make the column inconsistent with the chooser asset, so we don't CheckConsistency
 	if (NumRows != 0)
 	{
@@ -128,6 +130,7 @@ void FPoseSearchColumn::InsertRows(int Index, int Count)
 		}
 	}
 
+	InternalDatabase->NotifySynchronizeWithExternalDependencies();
 	CheckConsistency();
 }
 
@@ -140,7 +143,8 @@ void FPoseSearchColumn::DeleteRows(const TArray<uint32>& RowIndices)
 	{
 		InternalDatabase->RemoveAnimationAssetAt(Index);
 	}
-	
+
+	InternalDatabase->NotifySynchronizeWithExternalDependencies();
 	CheckConsistency();
 }
 
@@ -160,6 +164,8 @@ void FPoseSearchColumn::MoveRow(int SourceRowIndex, int TargetRowIndex)
 			TargetRowIndex--;
 		}
 		InternalDatabase->InsertAnimationAssetAt(CopySourceAsset, TargetRowIndex);
+
+		InternalDatabase->NotifySynchronizeWithExternalDependencies();
 	}
 
 	CheckConsistency();
@@ -177,6 +183,8 @@ void FPoseSearchColumn::CopyRow(FChooserColumnBase& SourceColumn, int SourceRowI
 				FScopedTransaction ScopedTransaction(LOCTEXT("CopyRow","CopyRow FPoseSearchColumn"));
 				InternalDatabase->Modify();
 				*TargetAsset = *SourceAsset;
+
+				InternalDatabase->NotifySynchronizeWithExternalDependencies();
 			}
 		}
 	}
@@ -374,7 +382,9 @@ void FPoseSearchColumn::SetFromDetails(FInstancedPropertyBag& PropertyBag, int32
 	{
 		if (FPoseSearchDatabaseAnimationAsset* Asset = InternalDatabase->GetMutableDatabaseAnimationAsset(RowIndex))
 		{
+			InternalDatabase->Modify();
 			*Asset = StructView->Get<FPoseSearchDatabaseAnimationAsset>();
+			InternalDatabase->NotifySynchronizeWithExternalDependencies();
 		}
 	}
 
diff --git a/Engine/Plugins/Animation/PoseSearch/Source/Runtime/Private/PoseSearchDatabase.cpp b/Engine/Plugins/Animation/PoseSearch/Source/Runtime/Private/PoseSearchDatabase.cpp
--- a/Engine/Plugins/Animation/PoseSearch/Source/Runtime/Private/PoseSearchDatabase.cpp
+++ b/Engine/Plugins/Animation/PoseSearch/Source/Runtime/Private/PoseSearchDatabase.cpp
@@ -1410,21 +1410,33 @@ void UPoseSearchDatabase::SynchronizeWithExternalDependencies(TConstArrayView<UA
 		}
 		else
 		{
+			if (!bModified)
+			{
+				Modify();
+				bModified = true;
+			}
+
 			DatabaseAnimationAssets.RemoveAt(AnimationAssetIndex);
-			bModified = true;
 		}
 	}
 
 	// adding the remaining NewDatabaseAnimationAsset(s) from NewDatabaseAnimationAssets
-	for (const FPoseSearchDatabaseAnimationAsset& NewDatabaseAnimationAsset : NewDatabaseAnimationAssets)
+	if (!NewDatabaseAnimationAssets.IsEmpty())
 	{
-		DatabaseAnimationAssets.Add(NewDatabaseAnimationAsset);
-		bModified = true;
+		if (!bModified)
+		{
+			Modify();
+			bModified = true;
+		}
+
+		for (const FPoseSearchDatabaseAnimationAsset& NewDatabaseAnimationAsset : NewDatabaseAnimationAssets)
+		{
+			DatabaseAnimationAssets.Add(NewDatabaseAnimationAsset);
+		}
 	}
 
 	if (bModified)
 	{
-		Modify();
 		NotifySynchronizeWithExternalDependencies();
 	}
 }
diff --git a/Engine/Plugins/Animation/PoseSearch/Source/Runtime/Private/PoseSearchDerivedData.cpp b/Engine/Plugins/Animation/PoseSearch/Source/Runtime/Private/PoseSearchDerivedData.cpp
--- a/Engine/Plugins/Animation/PoseSearch/Source/Runtime/Private/PoseSearchDerivedData.cpp
+++ b/Engine/Plugins/Animation/PoseSearch/Source/Runtime/Private/PoseSearchDerivedData.cpp
@@ -2327,18 +2327,6 @@ void FAsyncPoseSearchDatabasesManagement::OnObjectModified(UObject* Object)
 	PreModified(Object);
 }
 
-void FAsyncPoseSearchDatabasesManagement::ClearPreCancelled()
-{
-	// iterating backwards because of the possible RemoveAtSwap 
-	for (int32 TaskIndex = Tasks.Num() - 1; TaskIndex >= 0; --TaskIndex)
-	{
-		if (Tasks[TaskIndex]->GetState() == FPoseSearchDatabaseAsyncCacheTask::EState::PreCancelled)
-		{
-			Tasks[TaskIndex]->Cancel();
-		}
-	}
-}
-
 void FAsyncPoseSearchDatabasesManagement::PreModified(UObject* Object)
 {
 	check(IsInGameThread());
@@ -2347,10 +2335,9 @@ void FAsyncPoseSearchDatabasesManagement::PreModified(UObject* Object)
 
 	PartialKeyHashes.Remove(Object);
 
-	// iterating backwards because of the possible RemoveAtSwap
-	for (int32 TaskIndex = Tasks.Num() - 1; TaskIndex >= 0; --TaskIndex)
+	for (TUniquePtr<FPoseSearchDatabaseAsyncCacheTask>& TaskPtr : Tasks)
 	{
-		Tasks[TaskIndex]->PreCancelIfDependsOn(Object);
+		TaskPtr->PreCancelIfDependsOn(Object);
 	}
 
 	// collecting databases to synchronize prior modifying the Object
@@ -2368,7 +2355,14 @@ void FAsyncPoseSearchDatabasesManagement::PostModified(UObject* Object)
 
 	SynchronizeDatabases();
 
-	ClearPreCancelled();
+	// promoting the EState::PreCancelled into EState::Cancelled
+	for (TUniquePtr<FPoseSearchDatabaseAsyncCacheTask>& TaskPtr : Tasks)
+	{
+		if (TaskPtr->GetState() == FPoseSearchDatabaseAsyncCacheTask::EState::PreCancelled)
+		{
+			TaskPtr->Cancel();
+		}
+	}
 }
 
 void FAsyncPoseSearchDatabasesManagement::OnObjectTransacted(UObject* Object, const FTransactionObjectEvent& TransactionObjectEvent)
@@ -2380,22 +2374,24 @@ void FAsyncPoseSearchDatabasesManagement::OnPackageReloaded(const EPackageReload
 {
 	check(IsInGameThread());
 
-	if (InPackageReloadPhase == EPackageReloadPhase::PostPackageFixup && InPackageReloadedEvent)
+	if (!InPackageReloadedEvent)
 	{
-		FScopeLock Lock(&Mutex);
-
-		// @todo: figure out why we don't find the correct dependency into InPackageReloadedEvent->GetRepointedObjects()
-		//		  for now we invalidate all the DDC cache to be on the safe side
-		//for (const TPair<UObject*, UObject*>& Pair : InPackageReloadedEvent->GetRepointedObjects())
-		//{
-		//	OnObjectModified(Pair.Key);
-		//}
-
-		for (TUniquePtr<FPoseSearchDatabaseAsyncCacheTask>& TaskPtr : Tasks)
+	}
+	else if (InPackageReloadPhase == EPackageReloadPhase::PrePackageFixup)
+	{
+		const TMap<UObject*, UObject*>& RepointedObjects = InPackageReloadedEvent->GetRepointedObjects();
+		for (const TPair<UObject*, UObject*>& RepointPair : RepointedObjects)
+		{
+			PreModified(RepointPair.Key);
+		}
+	}
+	else if (InPackageReloadPhase == EPackageReloadPhase::PostPackageFixup)
+	{
+		const TMap<UObject*, UObject*>& RepointedObjects = InPackageReloadedEvent->GetRepointedObjects();
+		for (const TPair<UObject*, UObject*>& RepointPair : RepointedObjects)
 		{
-			UE_LOG(LogPoseSearch, Log, TEXT("%s - %s Cancelled because of OnPackageReloaded"), *LexToString(TaskPtr->GetDerivedDataKey()), *TaskPtr->GetDatabase()->GetName());
+			PostModified(RepointPair.Key);
 		}
-		Tasks.Reset();
 	}
 }
 
diff --git a/Engine/Plugins/Animation/PoseSearch/Source/Runtime/Public/PoseSearch/PoseSearchDerivedData.h b/Engine/Plugins/Animation/PoseSearch/Source/Runtime/Public/PoseSearch/PoseSearchDerivedData.h
--- a/Engine/Plugins/Animation/PoseSearch/Source/Runtime/Public/PoseSearch/PoseSearchDerivedData.h
+++ b/Engine/Plugins/Animation/PoseSearch/Source/Runtime/Public/PoseSearch/PoseSearchDerivedData.h
@@ -37,43 +37,42 @@ namespace UE::PoseSearch
 		static POSESEARCH_API EAsyncBuildIndexResult RequestAsyncBuildIndex(const UPoseSearchDatabase* Database, ERequestAsyncBuildFlag Flag);
 
 	private:
-		POSESEARCH_API FAsyncPoseSearchDatabasesManagement();
-		POSESEARCH_API ~FAsyncPoseSearchDatabasesManagement();
+		FAsyncPoseSearchDatabasesManagement();
+		~FAsyncPoseSearchDatabasesManagement();
 
-		static POSESEARCH_API FAsyncPoseSearchDatabasesManagement& Get();
-		static POSESEARCH_API EAsyncBuildIndexResult RequestAsyncBuildIndexInternal(const UPoseSearchDatabase* Database, ERequestAsyncBuildFlag Flag);
+		static FAsyncPoseSearchDatabasesManagement& Get();
+		static EAsyncBuildIndexResult RequestAsyncBuildIndexInternal(const UPoseSearchDatabase* Database, ERequestAsyncBuildFlag Flag);
 
-		POSESEARCH_API void OnObjectModified(UObject* Object);
-		POSESEARCH_API void OnObjectTransacted(UObject* Object, const FTransactionObjectEvent& TransactionObjectEvent);
-		POSESEARCH_API void OnPackageReloaded(const EPackageReloadPhase InPackageReloadPhase, FPackageReloadedEvent* InPackageReloadedEvent);
-		POSESEARCH_API void OnPreObjectPropertyChanged(UObject* InObject, const FEditPropertyChain& InPropertyChain);
-		POSESEARCH_API void OnObjectPropertyChanged(UObject* InObject, FPropertyChangedEvent& InPropertyChangedEvent);
+		void OnObjectModified(UObject* Object);
+		void OnObjectTransacted(UObject* Object, const FTransactionObjectEvent& TransactionObjectEvent);
+		void OnPackageReloaded(const EPackageReloadPhase InPackageReloadPhase, FPackageReloadedEvent* InPackageReloadedEvent);
+		void OnPreObjectPropertyChanged(UObject* InObject, const FEditPropertyChain& InPropertyChain);
+		void OnObjectPropertyChanged(UObject* InObject, FPropertyChangedEvent& InPropertyChangedEvent);
 
-		POSESEARCH_API void Shutdown();
-		POSESEARCH_API void StartQueuedTasks(int32 MaxActiveTasks);
-		POSESEARCH_API void PreModified(UObject* Object);
-		POSESEARCH_API void PostModified(UObject* Object);
-		POSESEARCH_API void ClearPreCancelled();
+		void Shutdown();
+		void StartQueuedTasks(int32 MaxActiveTasks);
+		void PreModified(UObject* Object);
+		void PostModified(UObject* Object);
 
 		// Begin FTickableGameObject
-		POSESEARCH_API virtual void Tick(float DeltaTime) override;
+		virtual void Tick(float DeltaTime) override;
 		virtual ETickableTickType GetTickableTickType() const override { return ETickableTickType::Always; }
-		POSESEARCH_API virtual TStatId GetStatId() const override;
+		virtual TStatId GetStatId() const override;
 		virtual bool IsTickableWhenPaused() const override { return true; }
 		virtual bool IsTickableInEditor() const override { return true; }
 		// End FTickableGameObject
 
 		// Begin FTickableCookObject
-		POSESEARCH_API virtual void TickCook(float DeltaTime, bool bCookCompete) override;
+		virtual void TickCook(float DeltaTime, bool bCookCompete) override;
 		// End FTickableCookObject
 
 		// Begin FGCObject
-		POSESEARCH_API void AddReferencedObjects(FReferenceCollector& Collector) override;
+		void AddReferencedObjects(FReferenceCollector& Collector) override;
 		virtual FString GetReferencerName() const override { return TEXT("FAsyncPoseSearchDatabaseManagement"); }
 		// End FGCObject
 		
-		POSESEARCH_API void CollectDatabasesToSynchronize(UObject* Object);
-		POSESEARCH_API void SynchronizeDatabases();
+		void CollectDatabasesToSynchronize(UObject* Object);
+		void SynchronizeDatabases();
 		
 		// map of all those databases UAnimSequenceBase(s) that were containing or are containing UAnimNotifyState_PoseSearchBranchIn that requires synchronization
 		typedef TMap<TWeakObjectPtr<UPoseSearchDatabase>, TArray<TWeakObjectPtr<UAnimSequenceBase>>> TDatabasesToSynchronize;
@@ -90,7 +89,7 @@ namespace UE::PoseSearch
 		// Experimental, this feature might be removed without warning, not for production use
 		FPartialKeyHashes PartialKeyHashes;
 
-		static POSESEARCH_API FCriticalSection Mutex;
+		static FCriticalSection Mutex;
 	};
 } // namespace UE::PoseSearch
 
