PoseSearch - support for PSD search returning multiple results

[RN] support for pose search database search returning multiple results

diff --git a/Engine/Plugins/Animation/PoseSearch/Source/Runtime/Private/AnimNode_MotionMatching.cpp b/Engine/Plugins/Animation/PoseSearch/Source/Runtime/Private/AnimNode_MotionMatching.cpp
--- a/Engine/Plugins/Animation/PoseSearch/Source/Runtime/Private/AnimNode_MotionMatching.cpp
+++ b/Engine/Plugins/Animation/PoseSearch/Source/Runtime/Private/AnimNode_MotionMatching.cpp
@@ -177,9 +177,11 @@ void FAnimNode_MotionMatching::UpdateAssetPlayer(const FAnimationUpdateContext&
 			AssetsToSearch[DatabaseToSearchIndex] = DatabasesToSearch[DatabaseToSearchIndex].Get();
 		}
 
+		FSearchResults_Single SearchResults;
 		FPoseSearchContinuingProperties ContinuingProperties;
 		ContinuingProperties.InitFrom(MotionMatchingState.SearchResult, NextUpdateInterruptMode);
-		SearchResult = UPoseSearchLibrary::MotionMatch(SearchContext, AssetsToSearch, ContinuingProperties);
+		UPoseSearchLibrary::MotionMatch(SearchContext, AssetsToSearch, ContinuingProperties, SearchResults);
+		SearchResult = SearchResults.GetBestResult();
 	}
 	else
 	{
diff --git a/Engine/Plugins/Animation/PoseSearch/Source/Runtime/Private/Chooser/PoseSearchChooserColumn.cpp b/Engine/Plugins/Animation/PoseSearch/Source/Runtime/Private/Chooser/PoseSearchChooserColumn.cpp
--- a/Engine/Plugins/Animation/PoseSearch/Source/Runtime/Private/Chooser/PoseSearchChooserColumn.cpp
+++ b/Engine/Plugins/Animation/PoseSearch/Source/Runtime/Private/Chooser/PoseSearchChooserColumn.cpp
@@ -293,8 +293,21 @@ void FPoseSearchColumn::Filter(FChooserEvaluationContext& Context, const FChoose
 			ContinuingProperties.bIsPlayingAssetMirrored = ChooserPlayerSettings->bIsPlayingAssetMirrored;
 			ContinuingProperties.PlayingAssetBlendParameters = ChooserPlayerSettings->PlayingAssetBlendParameters;
 		}
-		const UE::PoseSearch::FSearchResult SearchResult = UPoseSearchLibrary::MotionMatch(SearchContext, DatabasesToSearch, ContinuingProperties);
-
+		FSearchResults_AssetBests SearchResults;
+		UPoseSearchLibrary::MotionMatch(SearchContext, DatabasesToSearch, ContinuingProperties, SearchResults);
+
+		// @todo: WIP
+		// SearchResults.GetSourceAssetIdxToBestSearchResult() contains the mapping between asset indes in the PSD and the frame with the best cost (SearchResult)
+		//for (const FChooserIndexArray::FIndexData& IndexData : IndexListIn)
+		//{
+		//	if (const FSearchResult* SearchResult = SearchResults.GetSourceAssetIdxToBestSearchResult().Find(IndexData.Index))
+		//	{
+		//		IndexListOut.Push({ IndexData.Index, SearchResult->PoseCost });
+		//	}
+		//}
+
+		// @todo: shouldn't OutputStartTime, OutputCost, OutputMirror be arrays?
+		const FSearchResult SearchResult = SearchResults.GetBestResult();
 		if (const UE::PoseSearch::FSearchIndexAsset* SearchIndexAsset = SearchResult.GetSearchIndexAsset())
 		{
 			if (const FPoseSearchDatabaseAnimationAsset* DatabaseAnimationAsset = SearchResult.Database->GetDatabaseAnimationAsset(*SearchIndexAsset))
@@ -307,9 +320,9 @@ void FPoseSearchColumn::Filter(FChooserEvaluationContext& Context, const FChoose
 
 				 	if (const FChooserParameterFloatBase* StartTime = OutputStartTime.GetPtr<FChooserParameterFloatBase>())
 				 	{
-				 		StartTime->SetValue(Context, SearchResult.AssetTime);
+				 		StartTime->SetValue(Context, SearchResult.GetAssetTime());
 				 	}
-				 			
+
 				 	if (const FChooserParameterFloatBase* Cost = OutputCost.GetPtr<FChooserParameterFloatBase>())
                     {
 						Cost->SetValue(Context, SearchResult.PoseCost);
diff --git a/Engine/Plugins/Animation/PoseSearch/Source/Runtime/Private/PoseSearchContext.cpp b/Engine/Plugins/Animation/PoseSearch/Source/Runtime/Private/PoseSearchContext.cpp
--- a/Engine/Plugins/Animation/PoseSearch/Source/Runtime/Private/PoseSearchContext.cpp
+++ b/Engine/Plugins/Animation/PoseSearch/Source/Runtime/Private/PoseSearchContext.cpp
@@ -1178,15 +1178,10 @@ const UPoseSearchFeatureChannel* FSearchContext::GetCachedChannelData(uint32 Cha
 
 void FSearchContext::ResetCurrentBestCost()
 {
-	CurrentBestTotalCost = MAX_flt;
 }
 
 void FSearchContext::UpdateCurrentBestCost(const FPoseSearchCost& PoseSearchCost)
 {
-	if (PoseSearchCost < CurrentBestTotalCost)
-	{
-		CurrentBestTotalCost = PoseSearchCost;
-	};
 }
 
 TConstArrayView<float> FSearchContext::GetCachedQuery(const UPoseSearchSchema* Schema) const
diff --git a/Engine/Plugins/Animation/PoseSearch/Source/Runtime/Private/PoseSearchDatabase.cpp b/Engine/Plugins/Animation/PoseSearch/Source/Runtime/Private/PoseSearchDatabase.cpp
--- a/Engine/Plugins/Animation/PoseSearch/Source/Runtime/Private/PoseSearchDatabase.cpp
+++ b/Engine/Plugins/Animation/PoseSearch/Source/Runtime/Private/PoseSearchDatabase.cpp
@@ -277,9 +277,16 @@ private:
 };
 
 template<bool bReconstructPoseValues, bool bAlignedAndPadded>
-static inline void EvaluatePoseKernel(UE::PoseSearch::FSearchResult& Result, const UE::PoseSearch::FSearchIndex& SearchIndex, TConstArrayView<float> QueryValues, TArrayView<float> ReconstructedPoseValuesBuffer,
-	int32 PoseIdx, int32 EventPoseIdx, const UE::PoseSearch::FSearchFilters& SearchFilters, float ContinuingPoseCostAddend, float ContinuingInteractionCostAddend,
-	UE::PoseSearch::FSearchContext& SearchContext, const UPoseSearchDatabase* Database, TConstArrayView<float> DynamicWeightsSqrt, bool bUpdateBestCandidates, int32 ResultIndex = -1)
+static inline void EvaluatePoseKernel(ISearchResults& SearchResults, const FSearchIndex& SearchIndex, TConstArrayView<float> QueryValues, TArrayView<float> ReconstructedPoseValuesBuffer,
+	int32 PoseIdx, int32 EventPoseIdx, const FSearchFilters& SearchFilters, float ContinuingPoseCostAddend, float ContinuingInteractionCostAddend,
+	FSearchContext& SearchContext, const UPoseSearchDatabase* Database, TConstArrayView<float> DynamicWeightsSqrt
+#if UE_POSE_SEARCH_TRACE_ENABLED
+	, bool bTrackResult = true
+#endif // UE_POSE_SEARCH_TRACE_ENABLED
+#if WITH_EDITOR && ENABLE_ANIM_DEBUG && UE_POSE_SEARCH_TRACE_ENABLED
+	, const FSearchResult& BruteForceSearchResult = FSearchResult(), int32 ResultIndex = INDEX_NONE
+#endif // WITH_EDITOR && ENABLE_ANIM_DEBUG && UE_POSE_SEARCH_TRACE_ENABLED
+)
 {
 	using namespace UE::PoseSearch;
 
@@ -291,26 +298,24 @@ static inline void EvaluatePoseKernel(UE::PoseSearch::FSearchResult& Result, con
 #endif // UE_POSE_SEARCH_TRACE_ENABLED
 	))
 	{
-		const FPoseSearchCost PoseCost(CompareFeatureVectors<bAlignedAndPadded>(PoseValues, QueryValues, DynamicWeightsSqrt), SearchIndex.PoseMetadata[PoseIdx].GetCostAddend(), ContinuingPoseCostAddend, ContinuingInteractionCostAddend);
-		if (PoseCost < Result.PoseCost)
-		{
-			Result.PoseCost = PoseCost;
-			Result.PoseIdx = PoseIdx;
-			Result.EventPoseIdx = EventPoseIdx;
-
+		FSearchResult Result;
+		Result.PoseCost = FPoseSearchCost(CompareFeatureVectors<bAlignedAndPadded>(PoseValues, QueryValues, DynamicWeightsSqrt), SearchIndex.PoseMetadata[PoseIdx].GetCostAddend(), ContinuingPoseCostAddend, ContinuingInteractionCostAddend);
+		Result.PoseIdx = PoseIdx;
+		Result.EventPoseIdx = EventPoseIdx;
+		Result.Database = Database;
+		
 #if WITH_EDITOR && ENABLE_ANIM_DEBUG && UE_POSE_SEARCH_TRACE_ENABLED
-			if (bUpdateBestCandidates)
-			{
-				Result.BestPosePos = ResultIndex;
-			}
+		Result.BestPosePos = ResultIndex;
+		Result.BruteForcePoseCost = BruteForceSearchResult.PoseCost;
 #endif // WITH_EDITOR && ENABLE_ANIM_DEBUG && UE_POSE_SEARCH_TRACE_ENABLED
-		}
+
+		SearchResults.UpdateWith(Result);
 
 #if UE_POSE_SEARCH_TRACE_ENABLED
-		if (bUpdateBestCandidates)
+		if (bTrackResult)
 		{
 			// @todo: add tracking for EventPoseIdx
-			SearchContext.Track(Database, PoseIdx, EPoseCandidateFlags::Valid_Pose, PoseCost);
+			SearchContext.Track(Database, PoseIdx, EPoseCandidateFlags::Valid_Pose, Result.PoseCost);
 		}
 #endif // UE_POSE_SEARCH_TRACE_ENABLED
 	}
@@ -1528,67 +1533,44 @@ float UPoseSearchDatabase::GetNormalizedAssetTime(int32 PoseIdx) const
 	return NormalizedAssetTime;
 }
 
-UE::PoseSearch::FSearchResult UPoseSearchDatabase::Search(UE::PoseSearch::FSearchContext& SearchContext) const
+void UPoseSearchDatabase::Search(UE::PoseSearch::FSearchContext& SearchContext, UE::PoseSearch::ISearchResults& SearchResults) const
 {
 	using namespace UE::PoseSearch;
 
-	FSearchResult Result;
-
 #if WITH_EDITOR
 	if (EAsyncBuildIndexResult::Success != FAsyncPoseSearchDatabasesManagement::RequestAsyncBuildIndex(this, ERequestAsyncBuildFlag::ContinueRequest))
 	{
 		SearchContext.SetAsyncBuildIndexInProgress();
-		return Result;
+		return;
 	}
 #else
 	if (SearchIndexPrivate.IsEmpty())
 	{
-		return Result;
+		return;
 	}
 #endif // WITH_EDITOR
 
 	if (SearchContext.GetEventToSearch().IsValid())
 	{
-		Result = SearchEvent(SearchContext);
+		SearchEvent(SearchContext, SearchResults);
 	}
 	else
 	{
 		switch (PoseSearchMode)
 		{
 		case EPoseSearchMode::BruteForce:
-			Result = SearchBruteForce(SearchContext);
+			SearchBruteForce(SearchContext, SearchResults);
 			break;
 		case EPoseSearchMode::VPTree:
-			Result = SearchVPTree(SearchContext);
+			SearchVPTree(SearchContext, SearchResults);
 			break;
 		case EPoseSearchMode::PCAKDTree:
-			Result = SearchPCAKDTree(SearchContext);
+			SearchPCAKDTree(SearchContext, SearchResults);
 			break;
 		case EPoseSearchMode::EventOnly:
 			break;
 		}
 	}
-
-#if WITH_EDITOR && ENABLE_ANIM_DEBUG && UE_POSE_SEARCH_TRACE_ENABLED
-	if (GVarMotionMatchCompareAgainstBruteForce)
-	{
-		if (PoseSearchMode == EPoseSearchMode::PCAKDTree || PoseSearchMode == EPoseSearchMode::VPTree)
-		{
-			Result.BruteForcePoseCost = SearchBruteForce(SearchContext).PoseCost;
-		}
-		else
-		{
-			Result.BruteForcePoseCost = Result.PoseCost;
-		}
-	}
-#endif // WITH_EDITOR && ENABLE_ANIM_DEBUG && UE_POSE_SEARCH_TRACE_ENABLED
-
-#if UE_POSE_SEARCH_TRACE_ENABLED
-	// in case we skipped the search, or we didn't find any candidates we still have to track we requested to evaluate this database, so we keep track of this
-	SearchContext.Track(this);
-#endif // UE_POSE_SEARCH_TRACE_ENABLED
-
-	return Result;
 }
 
 void UPoseSearchDatabase::PopulateNonSelectableIdx(FNonSelectableIdx& NonSelectableIdx, UE::PoseSearch::FSearchContext& SearchContext
@@ -1733,7 +1715,23 @@ void UPoseSearchDatabase::PopulateSelectableAssetIdx(FSelectableAssetIdx& Select
 	}
 }
 
+UE::PoseSearch::FSearchResult UPoseSearchDatabase::Search(UE::PoseSearch::FSearchContext& SearchContext) const
+{
+	using namespace UE::PoseSearch;
+	FSearchResults_Single SearchResults;
+	Search(SearchContext, SearchResults);
+	return SearchResults.GetBestResult();
+}
+
 UE::PoseSearch::FSearchResult UPoseSearchDatabase::SearchContinuingPose(UE::PoseSearch::FSearchContext& SearchContext) const
+{
+	using namespace UE::PoseSearch;
+	FSearchResults_Single SearchResults;
+	SearchContinuingPose(SearchContext, SearchResults);
+	return SearchResults.GetBestResult();
+}
+
+void UPoseSearchDatabase::SearchContinuingPose(UE::PoseSearch::FSearchContext& SearchContext, UE::PoseSearch::ISearchResults& SearchResults) const
 {
 	QUICK_SCOPE_CYCLE_COUNTER(STAT_PoseSearch_ContinuingPose);
 
@@ -1741,14 +1739,11 @@ UE::PoseSearch::FSearchResult UPoseSearchDatabase::SearchContinuingPose(UE::Pose
 
 	check(SearchContext.GetContinuingPoseSearchResult().Database.Get() == this);
 
-	FSearchResult Result;
-	Result.bIsContinuingPoseSearch = true;
-
 #if WITH_EDITOR
 	if (EAsyncBuildIndexResult::Success != FAsyncPoseSearchDatabasesManagement::RequestAsyncBuildIndex(this, ERequestAsyncBuildFlag::ContinueRequest))
 	{
 		SearchContext.SetAsyncBuildIndexInProgress();
-		return Result;
+		return;
 	}
 #endif // WITH_EDITOR
 
@@ -1791,8 +1786,8 @@ UE::PoseSearch::FSearchResult UPoseSearchDatabase::SearchContinuingPose(UE::Pose
 	const float ContinuingInteractionCostAddend = SearchContext.IsContinuingInteraction() ? ContinuingInteractionCostBias : 0.f;
 
 	// since any PoseCost calculated here is at least SearchIndex.MinCostAddend + ContinuingPoseCostAddend + ContinuingInteractionCostAddend,
-	// there's no point in performing the search if CurrentBestTotalCost is already better than that
-	if (!GetSkipSearchIfPossible() || SearchContext.GetCurrentBestTotalCost() > (SearchIndex.MinCostAddend + ContinuingPoseCostAddend + ContinuingInteractionCostAddend))
+	// there's no point in performing the search if the worst search result is already better than that
+	if (!GetSkipSearchIfPossible() || SearchResults.ShouldPerformSearch(SearchIndex.MinCostAddend + ContinuingPoseCostAddend + ContinuingInteractionCostAddend))
 	{
 		const int32 NumDimensions = Schema->SchemaCardinality;
 		// FMemory_Alloca is forced 16 bytes aligned
@@ -1804,6 +1799,9 @@ UE::PoseSearch::FSearchResult UPoseSearchDatabase::SearchContinuingPose(UE::Pose
 		TArrayView<float> DynamicWeightsSqrtBuffer((float*)FMemory_Alloca(NumDimensions * sizeof(float)), NumDimensions);
 		const TConstArrayView<float> DynamicWeightsSqrt = CalculateDynamicWeightsSqrt(DynamicWeightsSqrtBuffer);
 
+		FSearchResult Result;
+		Result.bIsContinuingPoseSearch = true;
+
 		// is the data padded at 16 bytes (and 16 bytes aligned by construction)?
 		if (NumDimensions % 4 == 0)
 		{
@@ -1815,13 +1813,23 @@ UE::PoseSearch::FSearchResult UPoseSearchDatabase::SearchContinuingPose(UE::Pose
 			Result.PoseCost = FPoseSearchCost(CompareFeatureVectors<false>(PoseValues, SearchContext.GetOrBuildQuery(Schema), DynamicWeightsSqrt), SearchIndex.PoseMetadata[ContinuingPoseIdx].GetCostAddend(), ContinuingPoseCostAddend, ContinuingInteractionCostAddend);
 		}
 
-		Result.AssetTime = SearchContext.GetContinuingPoseSearchResult().AssetTime;
+		Result.SetAssetTime(SearchContext.GetContinuingPoseSearchResult().GetAssetTime());
 		Result.PoseIdx = ContinuingPoseIdx;
 		Result.Database = this;
 
 #if UE_POSE_SEARCH_TRACE_ENABLED
 		SearchContext.Track(this, ContinuingPoseIdx, EPoseCandidateFlags::Valid_ContinuingPose, Result.PoseCost);
 #endif // UE_POSE_SEARCH_TRACE_ENABLED
+
+#if WITH_EDITOR && ENABLE_ANIM_DEBUG && UE_POSE_SEARCH_TRACE_ENABLED
+		if (GVarMotionMatchCompareAgainstBruteForce)
+		{
+			Result.BruteForcePoseCost = Result.PoseCost;
+		}
+#endif // WITH_EDITOR && ENABLE_ANIM_DEBUG && UE_POSE_SEARCH_TRACE_ENABLED
+
+		check(Result.DebugValidate());
+		SearchResults.UpdateWith(Result);
 	}
 	else
 	{
@@ -1831,24 +1839,25 @@ UE::PoseSearch::FSearchResult UPoseSearchDatabase::SearchContinuingPose(UE::Pose
 #endif // UE_POSE_SEARCH_TRACE_ENABLED
 	}
 
-#if WITH_EDITOR && ENABLE_ANIM_DEBUG && UE_POSE_SEARCH_TRACE_ENABLED
-	if (GVarMotionMatchCompareAgainstBruteForce)
-	{
-		Result.BruteForcePoseCost = Result.PoseCost;
-	}
-#endif // WITH_EDITOR && ENABLE_ANIM_DEBUG && UE_POSE_SEARCH_TRACE_ENABLED
-
-	check(Result.DebugValidate());
-	return Result;
+	// NoTe: no need to SearchResults.FinalizeResults(), since SearchContinuingPose already calls SearchResults.UpdateWith with a Result with AssetTime already initialized
 }
 
-UE::PoseSearch::FSearchResult UPoseSearchDatabase::SearchPCAKDTree(UE::PoseSearch::FSearchContext& SearchContext) const
+void UPoseSearchDatabase::SearchPCAKDTree(UE::PoseSearch::FSearchContext& SearchContext, UE::PoseSearch::ISearchResults& SearchResults) const
 {
 	SCOPE_CYCLE_COUNTER(STAT_PoseSearch_PCAKNN);
 
 	using namespace UE::PoseSearch;
 
-	FSearchResult Result;
+#if WITH_EDITOR && ENABLE_ANIM_DEBUG && UE_POSE_SEARCH_TRACE_ENABLED
+	FSearchResult BruteForceSearchResult;
+	if (GVarMotionMatchCompareAgainstBruteForce)
+	{
+		FSearchResults_Single BruteForceSearchResults;
+		SearchBruteForce(SearchContext, BruteForceSearchResults);
+		BruteForceSearchResult = BruteForceSearchResults.GetBestResult();
+	}
+#endif // WITH_EDITOR && ENABLE_ANIM_DEBUG && UE_POSE_SEARCH_TRACE_ENABLED
+
 
 	const int32 NumDimensions = Schema->SchemaCardinality;
 	const FSearchIndex& SearchIndex = GetSearchIndex();
@@ -1856,8 +1865,8 @@ UE::PoseSearch::FSearchResult UPoseSearchDatabase::SearchPCAKDTree(UE::PoseSearc
 	const float ContinuingInteractionCostAddend = SearchContext.IsContinuingInteraction() ? ContinuingInteractionCostBias : 0.f;
 
 	// since any PoseCost calculated here is at least SearchIndex.MinCostAddend,
-	// there's no point in performing the search if CurrentBestTotalCost is already better than that
-	if (!GetSkipSearchIfPossible() || SearchContext.GetCurrentBestTotalCost() > SearchIndex.MinCostAddend)
+	// there's no point in performing the search if the worst search result is already better than that
+	if (!GetSkipSearchIfPossible() || SearchResults.ShouldPerformSearch(SearchIndex.MinCostAddend))
 	{
 		const uint32 ClampedNumberOfPrincipalComponents = SearchIndex.GetNumberOfPrincipalComponents();
 		const uint32 ClampedKDTreeQueryNumNeighbors = FMath::Clamp<uint32>(KDTreeQueryNumNeighbors, 1, SearchIndex.GetNumPoses());
@@ -2027,10 +2036,13 @@ UE::PoseSearch::FSearchResult UPoseSearchDatabase::SearchPCAKDTree(UE::PoseSearc
 					const TConstArrayView<int32> PoseIndexes = SearchIndex.PCAValuesVectorToPoseIndexes[Results[ResultIndex].Index];
 					for (int32 Index = 0; Index < PoseIndexes.Num() && NumEvaluatePoseKernelCalls < MaxNumEvaluatePoseKernelCalls; ++Index, ++NumEvaluatePoseKernelCalls)
 					{
-						EvaluatePoseKernel<false, true>(Result, SearchIndex, QueryValues, TArrayView<float>(), PoseIndexes[Index], INDEX_NONE, SearchFilters, ContinuingPoseCostAddend, ContinuingInteractionCostAddend, SearchContext, this, DynamicWeightsSqrt, true
-#if WITH_EDITOR && ENABLE_ANIM_DEBUG
-							, SortedResultsIndexes[ResultIndex]
-#endif // WITH_EDITOR && ENABLE_ANIM_DEBUG
+						EvaluatePoseKernel<false, true>(SearchResults, SearchIndex, QueryValues, TArrayView<float>(), PoseIndexes[Index], INDEX_NONE, SearchFilters, ContinuingPoseCostAddend, ContinuingInteractionCostAddend, SearchContext, this, DynamicWeightsSqrt
+#if UE_POSE_SEARCH_TRACE_ENABLED
+							, true
+#endif // UE_POSE_SEARCH_TRACE_ENABLED
+#if WITH_EDITOR && ENABLE_ANIM_DEBUG && UE_POSE_SEARCH_TRACE_ENABLED
+							, BruteForceSearchResult, SortedResultsIndexes[ResultIndex]
+#endif // WITH_EDITOR && ENABLE_ANIM_DEBUG && UE_POSE_SEARCH_TRACE_ENABLED
 						);
 					}
 				}
@@ -2043,10 +2055,13 @@ UE::PoseSearch::FSearchResult UPoseSearchDatabase::SearchPCAKDTree(UE::PoseSearc
 					const TConstArrayView<int32> PoseIndexes = SearchIndex.PCAValuesVectorToPoseIndexes[Results[ResultIndex].Index];
 					for (int32 Index = 0; Index < PoseIndexes.Num() && NumEvaluatePoseKernelCalls < MaxNumEvaluatePoseKernelCalls; ++Index, ++NumEvaluatePoseKernelCalls)
 					{
-						EvaluatePoseKernel<false, false>(Result, SearchIndex, QueryValues, TArrayView<float>(), PoseIndexes[Index], INDEX_NONE, SearchFilters, ContinuingPoseCostAddend, ContinuingInteractionCostAddend, SearchContext, this, DynamicWeightsSqrt, true
-#if WITH_EDITOR && ENABLE_ANIM_DEBUG
-							, SortedResultsIndexes[ResultIndex]
-#endif // WITH_EDITOR && ENABLE_ANIM_DEBUG
+						EvaluatePoseKernel<false, false>(SearchResults, SearchIndex, QueryValues, TArrayView<float>(), PoseIndexes[Index], INDEX_NONE, SearchFilters, ContinuingPoseCostAddend, ContinuingInteractionCostAddend, SearchContext, this, DynamicWeightsSqrt
+#if UE_POSE_SEARCH_TRACE_ENABLED
+							, true
+#endif // UE_POSE_SEARCH_TRACE_ENABLED
+#if WITH_EDITOR && ENABLE_ANIM_DEBUG && UE_POSE_SEARCH_TRACE_ENABLED
+						, BruteForceSearchResult, SortedResultsIndexes[ResultIndex]
+#endif // WITH_EDITOR && ENABLE_ANIM_DEBUG && UE_POSE_SEARCH_TRACE_ENABLED
 						);
 					}
 				}
@@ -2060,10 +2075,13 @@ UE::PoseSearch::FSearchResult UPoseSearchDatabase::SearchPCAKDTree(UE::PoseSearc
 			check(IsAligned(ReconstructedPoseValuesBuffer.GetData(), alignof(VectorRegister4Float)));
 			for (int32 ResultIndex = 0; ResultIndex < NumResults; ++ResultIndex)
 			{
-				EvaluatePoseKernel<true, false>(Result, SearchIndex, QueryValues, ReconstructedPoseValuesBuffer, Results[ResultIndex].Index, INDEX_NONE, SearchFilters, ContinuingPoseCostAddend, ContinuingInteractionCostAddend, SearchContext, this, DynamicWeightsSqrt, true
-#if WITH_EDITOR && ENABLE_ANIM_DEBUG
-					, SortedResultsIndexes[ResultIndex]
-#endif // WITH_EDITOR && ENABLE_ANIM_DEBUG
+				EvaluatePoseKernel<true, false>(SearchResults, SearchIndex, QueryValues, ReconstructedPoseValuesBuffer, Results[ResultIndex].Index, INDEX_NONE, SearchFilters, ContinuingPoseCostAddend, ContinuingInteractionCostAddend, SearchContext, this, DynamicWeightsSqrt
+#if UE_POSE_SEARCH_TRACE_ENABLED
+					, true
+#endif // UE_POSE_SEARCH_TRACE_ENABLED
+#if WITH_EDITOR && ENABLE_ANIM_DEBUG && UE_POSE_SEARCH_TRACE_ENABLED
+					, BruteForceSearchResult, SortedResultsIndexes[ResultIndex]
+#endif // WITH_EDITOR && ENABLE_ANIM_DEBUG && UE_POSE_SEARCH_TRACE_ENABLED
 				);
 			}
 		}
@@ -2072,10 +2090,13 @@ UE::PoseSearch::FSearchResult UPoseSearchDatabase::SearchPCAKDTree(UE::PoseSearc
 		{
 			for (int32 ResultIndex = 0; ResultIndex < NumResults; ++ResultIndex)
 			{
-				EvaluatePoseKernel<false, true>(Result, SearchIndex, QueryValues, TArrayView<float>(), Results[ResultIndex].Index, INDEX_NONE, SearchFilters, ContinuingPoseCostAddend, ContinuingInteractionCostAddend, SearchContext, this, DynamicWeightsSqrt, true
-#if WITH_EDITOR && ENABLE_ANIM_DEBUG
-					, SortedResultsIndexes[ResultIndex]
-#endif // WITH_EDITOR && ENABLE_ANIM_DEBUG
+				EvaluatePoseKernel<false, true>(SearchResults, SearchIndex, QueryValues, TArrayView<float>(), Results[ResultIndex].Index, INDEX_NONE, SearchFilters, ContinuingPoseCostAddend, ContinuingInteractionCostAddend, SearchContext, this, DynamicWeightsSqrt
+#if UE_POSE_SEARCH_TRACE_ENABLED
+					, true
+#endif // UE_POSE_SEARCH_TRACE_ENABLED
+#if WITH_EDITOR && ENABLE_ANIM_DEBUG && UE_POSE_SEARCH_TRACE_ENABLED
+					, BruteForceSearchResult, SortedResultsIndexes[ResultIndex]
+#endif // WITH_EDITOR && ENABLE_ANIM_DEBUG && UE_POSE_SEARCH_TRACE_ENABLED
 				);
 			}
 		}
@@ -2084,11 +2105,14 @@ UE::PoseSearch::FSearchResult UPoseSearchDatabase::SearchPCAKDTree(UE::PoseSearc
 		{
 			for (int32 ResultIndex = 0; ResultIndex < NumResults; ++ResultIndex)
 			{
-				EvaluatePoseKernel<false, false>(Result, SearchIndex, QueryValues, TArrayView<float>(), Results[ResultIndex].Index, INDEX_NONE, SearchFilters, ContinuingPoseCostAddend, ContinuingInteractionCostAddend, SearchContext, this, DynamicWeightsSqrt, true
-#if WITH_EDITOR && ENABLE_ANIM_DEBUG
-					, SortedResultsIndexes[ResultIndex]
-#endif // WITH_EDITOR && ENABLE_ANIM_DEBUG
-						);
+				EvaluatePoseKernel<false, false>(SearchResults, SearchIndex, QueryValues, TArrayView<float>(), Results[ResultIndex].Index, INDEX_NONE, SearchFilters, ContinuingPoseCostAddend, ContinuingInteractionCostAddend, SearchContext, this, DynamicWeightsSqrt
+#if UE_POSE_SEARCH_TRACE_ENABLED
+					, true
+#endif // UE_POSE_SEARCH_TRACE_ENABLED
+#if WITH_EDITOR && ENABLE_ANIM_DEBUG && UE_POSE_SEARCH_TRACE_ENABLED
+					, BruteForceSearchResult, SortedResultsIndexes[ResultIndex]
+#endif // WITH_EDITOR && ENABLE_ANIM_DEBUG && UE_POSE_SEARCH_TRACE_ENABLED
+				);
 			}
 		}
 	}
@@ -2106,32 +2130,37 @@ UE::PoseSearch::FSearchResult UPoseSearchDatabase::SearchPCAKDTree(UE::PoseSearc
 #endif // UE_POSE_SEARCH_TRACE_ENABLED
 	}
 
-	// finalizing Result properties
-	if (Result.PoseIdx != INDEX_NONE)
-	{
-		Result.AssetTime = GetNormalizedAssetTime(Result.PoseIdx);
-		Result.Database = this;
-	}
+	SearchResults.FinalizeResults();
 
-	check(Result.DebugValidate());
-	return Result;
+#if UE_POSE_SEARCH_TRACE_ENABLED
+	// in case we skipped the search, or we didn't find any candidates we still have to track we requested to evaluate this database, so we keep track of this
+	SearchContext.Track(this);
+#endif // UE_POSE_SEARCH_TRACE_ENABLED
 }
 
-UE::PoseSearch::FSearchResult UPoseSearchDatabase::SearchVPTree(UE::PoseSearch::FSearchContext& SearchContext) const
+void UPoseSearchDatabase::SearchVPTree(UE::PoseSearch::FSearchContext& SearchContext, UE::PoseSearch::ISearchResults& SearchResults) const
 {
 	SCOPE_CYCLE_COUNTER(STAT_PoseSearch_VPTree);
 
 	using namespace UE::PoseSearch;
 
-	FSearchResult Result;
+#if WITH_EDITOR && ENABLE_ANIM_DEBUG && UE_POSE_SEARCH_TRACE_ENABLED
+	FSearchResult BruteForceSearchResult;
+	if (GVarMotionMatchCompareAgainstBruteForce)
+	{
+		FSearchResults_Single BruteForceSearchResults;
+		SearchBruteForce(SearchContext, BruteForceSearchResults);
+		BruteForceSearchResult = BruteForceSearchResults.GetBestResult();
+	}
+#endif // WITH_EDITOR && ENABLE_ANIM_DEBUG && UE_POSE_SEARCH_TRACE_ENABLED
 
 	const FSearchIndex& SearchIndex = GetSearchIndex();
 	const float ContinuingPoseCostAddend = 0.f;
 	const float ContinuingInteractionCostAddend = SearchContext.IsContinuingInteraction() ? ContinuingInteractionCostBias : 0.f;
 
 	// since any PoseCost calculated here is at least SearchIndex.MinCostAddend,
-	// there's no point in performing the search if CurrentBestTotalCost is already better than that
-	if (!GetSkipSearchIfPossible() || SearchContext.GetCurrentBestTotalCost() > SearchIndex.MinCostAddend)
+	// there's no point in performing the search if the worst search result is already better than that
+	if (!GetSkipSearchIfPossible() || SearchResults.ShouldPerformSearch(SearchIndex.MinCostAddend))
 	{
 		// @todo: perhaps store the DynamicWeightsSqrt into the GetOrBuildQuery to share between the continuing pose and the full search
 		const int32 NumDimensions = Schema->SchemaCardinality;
@@ -2174,11 +2203,14 @@ UE::PoseSearch::FSearchResult UPoseSearchDatabase::SearchVPTree(UE::PoseSearch::
 				const TConstArrayView<int32> PoseIndexes = SearchIndex.ValuesVectorToPoseIndexes[IndexDistance.Index];
 				for (int32 Index = 0; Index < PoseIndexes.Num() && NumEvaluatePoseKernelCalls < MaxNumEvaluatePoseKernelCalls; ++Index, ++NumEvaluatePoseKernelCalls)
 				{
-					EvaluatePoseKernel<false, false>(Result, SearchIndex, QueryValues, TArrayView<float>(), PoseIndexes[Index], INDEX_NONE, SearchFilters, ContinuingPoseCostAddend, ContinuingInteractionCostAddend, SearchContext, this, DynamicWeightsSqrt, true
-#if WITH_EDITOR && ENABLE_ANIM_DEBUG
-						, ResultIndex
-#endif // WITH_EDITOR && ENABLE_ANIM_DEBUG
-						);
+					EvaluatePoseKernel<false, false>(SearchResults, SearchIndex, QueryValues, TArrayView<float>(), PoseIndexes[Index], INDEX_NONE, SearchFilters, ContinuingPoseCostAddend, ContinuingInteractionCostAddend, SearchContext, this, DynamicWeightsSqrt
+#if UE_POSE_SEARCH_TRACE_ENABLED
+						, true
+#endif // UE_POSE_SEARCH_TRACE_ENABLED
+#if WITH_EDITOR && ENABLE_ANIM_DEBUG && UE_POSE_SEARCH_TRACE_ENABLED
+						, BruteForceSearchResult, ResultIndex
+#endif // WITH_EDITOR && ENABLE_ANIM_DEBUG && UE_POSE_SEARCH_TRACE_ENABLED
+					);
 				}
 			}
 		}
@@ -2189,10 +2221,13 @@ UE::PoseSearch::FSearchResult UPoseSearchDatabase::SearchVPTree(UE::PoseSearch::
 				const FIndexDistance& IndexDistance = UnsortedResults[ResultIndex];
 
 				// @todo: IndexDistance.Distance is the Sqrt(DissimilarityCost), so there's no need to calculate it again in SearchIndex.ComparePoses
-				EvaluatePoseKernel<false, false>(Result, SearchIndex, QueryValues, TArrayView<float>(), IndexDistance.Index, INDEX_NONE, SearchFilters, ContinuingPoseCostAddend, ContinuingInteractionCostAddend, SearchContext, this, DynamicWeightsSqrt, true
-#if WITH_EDITOR && ENABLE_ANIM_DEBUG
-					, ResultIndex
-#endif // WITH_EDITOR && ENABLE_ANIM_DEBUG
+				EvaluatePoseKernel<false, false>(SearchResults, SearchIndex, QueryValues, TArrayView<float>(), IndexDistance.Index, INDEX_NONE, SearchFilters, ContinuingPoseCostAddend, ContinuingInteractionCostAddend, SearchContext, this, DynamicWeightsSqrt
+#if UE_POSE_SEARCH_TRACE_ENABLED
+					, true
+#endif // UE_POSE_SEARCH_TRACE_ENABLED
+#if WITH_EDITOR && ENABLE_ANIM_DEBUG && UE_POSE_SEARCH_TRACE_ENABLED
+					, BruteForceSearchResult, ResultIndex
+#endif // WITH_EDITOR && ENABLE_ANIM_DEBUG && UE_POSE_SEARCH_TRACE_ENABLED
 				);
 			}
 		}
@@ -2212,32 +2247,28 @@ UE::PoseSearch::FSearchResult UPoseSearchDatabase::SearchVPTree(UE::PoseSearch::
 #endif // UE_POSE_SEARCH_TRACE_ENABLED
 	}
 
-	// finalizing Result properties
-	if (Result.PoseIdx != INDEX_NONE)
-	{
-		Result.AssetTime = GetNormalizedAssetTime(Result.PoseIdx);
-		Result.Database = this;
-	}
+	SearchResults.FinalizeResults();
 
-	check(Result.DebugValidate());
-	return Result;
+#if UE_POSE_SEARCH_TRACE_ENABLED
+	// in case we skipped the search, or we didn't find any candidates we still have to track we requested to evaluate this database, so we keep track of this
+	SearchContext.Track(this);
+#endif // UE_POSE_SEARCH_TRACE_ENABLED
 }
 
-UE::PoseSearch::FSearchResult UPoseSearchDatabase::SearchBruteForce(UE::PoseSearch::FSearchContext& SearchContext) const
+void UPoseSearchDatabase::SearchBruteForce(UE::PoseSearch::FSearchContext& SearchContext, UE::PoseSearch::ISearchResults& SearchResults) const
 {
 	SCOPE_CYCLE_COUNTER(STAT_PoseSearch_BruteForce);
 	
 	using namespace UE::PoseSearch;
 	
-	FSearchResult Result;
-
 	const FSearchIndex& SearchIndex = GetSearchIndex();
 	const float ContinuingPoseCostAddend = 0.f;
 	const float ContinuingInteractionCostAddend = SearchContext.IsContinuingInteraction() ? ContinuingInteractionCostBias : 0.f;
+	const bool bTrackResults = PoseSearchMode == EPoseSearchMode::BruteForce;
 
 	// since any PoseCost calculated here is at least SearchIndex.MinCostAddend,
-	// there's no point in performing the search if CurrentBestTotalCost is already better than that
-	if (!GetSkipSearchIfPossible() || SearchContext.GetCurrentBestTotalCost() > SearchIndex.MinCostAddend)
+	// there's no point in performing the search if the worst search result is already better than that
+	if (!GetSkipSearchIfPossible() || SearchResults.ShouldPerformSearch(SearchIndex.MinCostAddend))
 	{
 		// @todo: perhaps store the DynamicWeightsSqrt into the GetOrBuildQuery to share between the continuing pose and the full search
 		const int32 NumDimensions = Schema->SchemaCardinality;
@@ -2256,7 +2287,6 @@ UE::PoseSearch::FSearchResult UPoseSearchDatabase::SearchBruteForce(UE::PoseSear
 #endif // UE_POSE_SEARCH_TRACE_ENABLED
 		);
 
-		const bool bUpdateBestCandidates = PoseSearchMode == EPoseSearchMode::BruteForce;
 		const FSearchFilters SearchFilters(Schema, NonSelectableIdx, FSelectableAssetIdx(), SearchIndex.bAnyBlockTransition);
 
 		if (SelectableAssetIdx.IsEmpty())
@@ -2269,10 +2299,13 @@ UE::PoseSearch::FSearchResult UPoseSearchDatabase::SearchBruteForce(UE::PoseSear
 				check(IsAligned(ReconstructedPoseValuesBuffer.GetData(), alignof(VectorRegister4Float)));
 				for (int32 PoseIdx = 0; PoseIdx < SearchIndex.GetNumPoses(); ++PoseIdx)
 				{
-					EvaluatePoseKernel<true, false>(Result, SearchIndex, QueryValues, ReconstructedPoseValuesBuffer, PoseIdx, INDEX_NONE, SearchFilters, ContinuingPoseCostAddend, ContinuingInteractionCostAddend, SearchContext, this, DynamicWeightsSqrt, bUpdateBestCandidates
-#if WITH_EDITOR && ENABLE_ANIM_DEBUG
-						, PoseIdx
-#endif // WITH_EDITOR && ENABLE_ANIM_DEBUG
+					EvaluatePoseKernel<true, false>(SearchResults, SearchIndex, QueryValues, ReconstructedPoseValuesBuffer, PoseIdx, INDEX_NONE, SearchFilters, ContinuingPoseCostAddend, ContinuingInteractionCostAddend, SearchContext, this, DynamicWeightsSqrt
+#if UE_POSE_SEARCH_TRACE_ENABLED
+						, bTrackResults
+#endif // UE_POSE_SEARCH_TRACE_ENABLED
+#if WITH_EDITOR && ENABLE_ANIM_DEBUG && UE_POSE_SEARCH_TRACE_ENABLED
+						, FSearchResult(), PoseIdx
+#endif // WITH_EDITOR && ENABLE_ANIM_DEBUG && UE_POSE_SEARCH_TRACE_ENABLED
 					);
 				}
 			}
@@ -2281,11 +2314,14 @@ UE::PoseSearch::FSearchResult UPoseSearchDatabase::SearchBruteForce(UE::PoseSear
 			{
 				for (int32 PoseIdx = 0; PoseIdx < SearchIndex.GetNumPoses(); ++PoseIdx)
 				{
-					EvaluatePoseKernel<false, true>(Result, SearchIndex, QueryValues, TArrayView<float>(), PoseIdx, INDEX_NONE, SearchFilters, ContinuingPoseCostAddend, ContinuingInteractionCostAddend, SearchContext, this, DynamicWeightsSqrt, bUpdateBestCandidates
-#if WITH_EDITOR && ENABLE_ANIM_DEBUG
-						, PoseIdx
-#endif // WITH_EDITOR && ENABLE_ANIM_DEBUG
-						);
+					EvaluatePoseKernel<false, true>(SearchResults, SearchIndex, QueryValues, TArrayView<float>(), PoseIdx, INDEX_NONE, SearchFilters, ContinuingPoseCostAddend, ContinuingInteractionCostAddend, SearchContext, this, DynamicWeightsSqrt
+#if UE_POSE_SEARCH_TRACE_ENABLED
+						, bTrackResults
+#endif // UE_POSE_SEARCH_TRACE_ENABLED
+#if WITH_EDITOR && ENABLE_ANIM_DEBUG && UE_POSE_SEARCH_TRACE_ENABLED
+						, FSearchResult(), PoseIdx
+#endif // WITH_EDITOR && ENABLE_ANIM_DEBUG && UE_POSE_SEARCH_TRACE_ENABLED
+					);
 				}
 			}
 			// no reconstruction, but data is not 16 bytes padded
@@ -2293,11 +2329,14 @@ UE::PoseSearch::FSearchResult UPoseSearchDatabase::SearchBruteForce(UE::PoseSear
 			{
 				for (int32 PoseIdx = 0; PoseIdx < SearchIndex.GetNumPoses(); ++PoseIdx)
 				{
-					EvaluatePoseKernel<false, false>(Result, SearchIndex, QueryValues, TArrayView<float>(), PoseIdx, INDEX_NONE, SearchFilters, ContinuingPoseCostAddend, ContinuingInteractionCostAddend, SearchContext, this, DynamicWeightsSqrt, bUpdateBestCandidates
-#if WITH_EDITOR && ENABLE_ANIM_DEBUG
-						, PoseIdx
-#endif // WITH_EDITOR && ENABLE_ANIM_DEBUG
-						);
+					EvaluatePoseKernel<false, false>(SearchResults, SearchIndex, QueryValues, TArrayView<float>(), PoseIdx, INDEX_NONE, SearchFilters, ContinuingPoseCostAddend, ContinuingInteractionCostAddend, SearchContext, this, DynamicWeightsSqrt
+#if UE_POSE_SEARCH_TRACE_ENABLED
+						, bTrackResults
+#endif // UE_POSE_SEARCH_TRACE_ENABLED
+#if WITH_EDITOR && ENABLE_ANIM_DEBUG && UE_POSE_SEARCH_TRACE_ENABLED
+						, FSearchResult(), PoseIdx
+#endif // WITH_EDITOR && ENABLE_ANIM_DEBUG && UE_POSE_SEARCH_TRACE_ENABLED
+					);
 				}
 			}
 		}
@@ -2319,10 +2358,13 @@ UE::PoseSearch::FSearchResult UPoseSearchDatabase::SearchBruteForce(UE::PoseSear
 					const int32 LastPoseIdx = FirstPoseIdx + SearchIndexAsset.GetNumPoses();
 					for (int32 PoseIdx = FirstPoseIdx; PoseIdx < LastPoseIdx; ++PoseIdx)
 					{
-						EvaluatePoseKernel<true, false>(Result, SearchIndex, QueryValues, ReconstructedPoseValuesBuffer, PoseIdx, INDEX_NONE, SearchFilters, ContinuingPoseCostAddend, ContinuingInteractionCostAddend, SearchContext, this, DynamicWeightsSqrt, bUpdateBestCandidates
-#if WITH_EDITOR && ENABLE_ANIM_DEBUG
-							, ++ResultIndex
-#endif // WITH_EDITOR && ENABLE_ANIM_DEBUG
+						EvaluatePoseKernel<true, false>(SearchResults, SearchIndex, QueryValues, ReconstructedPoseValuesBuffer, PoseIdx, INDEX_NONE, SearchFilters, ContinuingPoseCostAddend, ContinuingInteractionCostAddend, SearchContext, this, DynamicWeightsSqrt
+#if UE_POSE_SEARCH_TRACE_ENABLED
+							, bTrackResults
+#endif // UE_POSE_SEARCH_TRACE_ENABLED
+#if WITH_EDITOR && ENABLE_ANIM_DEBUG && UE_POSE_SEARCH_TRACE_ENABLED
+							, FSearchResult(), ++ResultIndex
+#endif // WITH_EDITOR && ENABLE_ANIM_DEBUG && UE_POSE_SEARCH_TRACE_ENABLED
 						);
 					}
 				}
@@ -2337,10 +2379,13 @@ UE::PoseSearch::FSearchResult UPoseSearchDatabase::SearchBruteForce(UE::PoseSear
 					const int32 LastPoseIdx = FirstPoseIdx + SearchIndexAsset.GetNumPoses();
 					for (int32 PoseIdx = FirstPoseIdx; PoseIdx < LastPoseIdx; ++PoseIdx)
 					{
-						EvaluatePoseKernel<false, true>(Result, SearchIndex, QueryValues, TArrayView<float>(), PoseIdx, INDEX_NONE, SearchFilters, ContinuingPoseCostAddend, ContinuingInteractionCostAddend, SearchContext, this, DynamicWeightsSqrt, bUpdateBestCandidates
-#if WITH_EDITOR && ENABLE_ANIM_DEBUG
-							, ++ResultIndex
-#endif // WITH_EDITOR && ENABLE_ANIM_DEBUG
+						EvaluatePoseKernel<false, true>(SearchResults, SearchIndex, QueryValues, TArrayView<float>(), PoseIdx, INDEX_NONE, SearchFilters, ContinuingPoseCostAddend, ContinuingInteractionCostAddend, SearchContext, this, DynamicWeightsSqrt
+#if UE_POSE_SEARCH_TRACE_ENABLED
+							, bTrackResults
+#endif // UE_POSE_SEARCH_TRACE_ENABLED
+#if WITH_EDITOR && ENABLE_ANIM_DEBUG && UE_POSE_SEARCH_TRACE_ENABLED
+							, FSearchResult(), ++ResultIndex
+#endif // WITH_EDITOR && ENABLE_ANIM_DEBUG && UE_POSE_SEARCH_TRACE_ENABLED
 						);
 					}
 				}
@@ -2355,10 +2400,13 @@ UE::PoseSearch::FSearchResult UPoseSearchDatabase::SearchBruteForce(UE::PoseSear
 					const int32 LastPoseIdx = FirstPoseIdx + SearchIndexAsset.GetNumPoses();
 					for (int32 PoseIdx = FirstPoseIdx; PoseIdx < LastPoseIdx; ++PoseIdx)
 					{
-						EvaluatePoseKernel<false, false>(Result, SearchIndex, QueryValues, TArrayView<float>(), PoseIdx, INDEX_NONE, SearchFilters, ContinuingPoseCostAddend, ContinuingInteractionCostAddend, SearchContext, this, DynamicWeightsSqrt, bUpdateBestCandidates
-#if WITH_EDITOR && ENABLE_ANIM_DEBUG
-							, ++ResultIndex
-#endif // WITH_EDITOR && ENABLE_ANIM_DEBUG
+						EvaluatePoseKernel<false, false>(SearchResults, SearchIndex, QueryValues, TArrayView<float>(), PoseIdx, INDEX_NONE, SearchFilters, ContinuingPoseCostAddend, ContinuingInteractionCostAddend, SearchContext, this, DynamicWeightsSqrt
+#if UE_POSE_SEARCH_TRACE_ENABLED
+							, bTrackResults
+#endif // UE_POSE_SEARCH_TRACE_ENABLED
+#if WITH_EDITOR && ENABLE_ANIM_DEBUG && UE_POSE_SEARCH_TRACE_ENABLED
+							, FSearchResult(), ++ResultIndex
+#endif // WITH_EDITOR && ENABLE_ANIM_DEBUG && UE_POSE_SEARCH_TRACE_ENABLED
 						);
 					}
 				}
@@ -2380,25 +2428,23 @@ UE::PoseSearch::FSearchResult UPoseSearchDatabase::SearchBruteForce(UE::PoseSear
 #endif // UE_POSE_SEARCH_TRACE_ENABLED
 	}
 
-	// finalizing Result properties
-	if (Result.PoseIdx != INDEX_NONE)
+	SearchResults.FinalizeResults();
+
+#if UE_POSE_SEARCH_TRACE_ENABLED
+	if (bTrackResults)
 	{
-		Result.AssetTime = GetNormalizedAssetTime(Result.PoseIdx);
-		Result.Database = this;
+		// in case we skipped the search, or we didn't find any candidates we still have to track we requested to evaluate this database, so we keep track of this
+		SearchContext.Track(this);
 	}
-
-	check(Result.DebugValidate());
-	return Result;
+#endif // UE_POSE_SEARCH_TRACE_ENABLED
 }
 
-UE::PoseSearch::FSearchResult UPoseSearchDatabase::SearchEvent(UE::PoseSearch::FSearchContext& SearchContext) const
+void UPoseSearchDatabase::SearchEvent(UE::PoseSearch::FSearchContext& SearchContext, UE::PoseSearch::ISearchResults& SearchResults) const
 {
 	SCOPE_CYCLE_COUNTER(STAT_PoseSearch_Event);
 
 	using namespace UE::PoseSearch;
 
-	FSearchResult Result;
-
 	const FSearchIndex& SearchIndex = GetSearchIndex();
 	if (SearchIndex.IsValuesEmpty())
 	{
@@ -2438,10 +2484,13 @@ UE::PoseSearch::FSearchResult UPoseSearchDatabase::SearchEvent(UE::PoseSearch::F
 			{
 				for (int32 EventPoseIdx : PosesWithEvent)
 				{
-					EvaluatePoseKernel<false, false>(Result, SearchIndex, QueryValues, TArrayView<float>(), EventPoseIdx, EventPoseIdx, SearchFilters, ContinuingPoseCostAddend, ContinuingInteractionCostAddend, SearchContext, this, DynamicWeightsSqrt, true
-#if WITH_EDITOR && ENABLE_ANIM_DEBUG
-						, EventPoseIdx
-#endif // WITH_EDITOR && ENABLE_ANIM_DEBUG
+					EvaluatePoseKernel<false, false>(SearchResults, SearchIndex, QueryValues, TArrayView<float>(), EventPoseIdx, EventPoseIdx, SearchFilters, ContinuingPoseCostAddend, ContinuingInteractionCostAddend, SearchContext, this, DynamicWeightsSqrt
+#if UE_POSE_SEARCH_TRACE_ENABLED
+						, true
+#endif // UE_POSE_SEARCH_TRACE_ENABLED
+#if WITH_EDITOR && ENABLE_ANIM_DEBUG && UE_POSE_SEARCH_TRACE_ENABLED
+						, FSearchResult(), EventPoseIdx
+#endif // WITH_EDITOR && ENABLE_ANIM_DEBUG && UE_POSE_SEARCH_TRACE_ENABLED
 					);
 				}
 			}
@@ -2455,12 +2504,15 @@ UE::PoseSearch::FSearchResult UPoseSearchDatabase::SearchEvent(UE::PoseSearch::F
 					const FSearchIndexAsset& SearchIndexAsset = SearchIndex.GetAssetForPose(EventPoseIdx);
 					const float EventTime = SearchIndexAsset.GetTimeFromPoseIndex(EventPoseIdx, Schema->SampleRate);
 					SearchIndexAsset.IteratePoseIndexesOverTime(EventTime - TimeToEventEnd, EventTime - TimeToEventStart, Schema->SampleRate,
-						[this, &Result, &SearchIndex, &QueryValues, EventPoseIdx, &SearchFilters, ContinuingPoseCostAddend, ContinuingInteractionCostAddend, &SearchContext, &DynamicWeightsSqrt](int32 PoseIdx)
+						[this, &SearchResults, &SearchIndex, &QueryValues, EventPoseIdx, &SearchFilters, ContinuingPoseCostAddend, ContinuingInteractionCostAddend, &SearchContext, &DynamicWeightsSqrt](int32 PoseIdx)
 						{
-							EvaluatePoseKernel<false, false>(Result, SearchIndex, QueryValues, TArrayView<float>(), PoseIdx, EventPoseIdx, SearchFilters, ContinuingPoseCostAddend, ContinuingInteractionCostAddend, SearchContext, this, DynamicWeightsSqrt, true
-#if WITH_EDITOR && ENABLE_ANIM_DEBUG
-								, PoseIdx
-#endif // WITH_EDITOR && ENABLE_ANIM_DEBUG
+							EvaluatePoseKernel<false, false>(SearchResults, SearchIndex, QueryValues, TArrayView<float>(), PoseIdx, EventPoseIdx, SearchFilters, ContinuingPoseCostAddend, ContinuingInteractionCostAddend, SearchContext, this, DynamicWeightsSqrt
+#if UE_POSE_SEARCH_TRACE_ENABLED
+								, true
+#endif // UE_POSE_SEARCH_TRACE_ENABLED
+#if WITH_EDITOR && ENABLE_ANIM_DEBUG && UE_POSE_SEARCH_TRACE_ENABLED
+								, FSearchResult(), PoseIdx
+#endif // WITH_EDITOR && ENABLE_ANIM_DEBUG && UE_POSE_SEARCH_TRACE_ENABLED
 							);
 						});
 				}
@@ -2468,12 +2520,10 @@ UE::PoseSearch::FSearchResult UPoseSearchDatabase::SearchEvent(UE::PoseSearch::F
 		}
 	}
 
-	// finalizing Result properties
-	if (Result.PoseIdx != INDEX_NONE)
-	{
-		Result.AssetTime = GetNormalizedAssetTime(Result.PoseIdx);
-		Result.Database = this;
-	}
+	SearchResults.FinalizeResults();
 
-	return Result;
+#if UE_POSE_SEARCH_TRACE_ENABLED
+	// in case we skipped the search, or we didn't find any candidates we still have to track we requested to evaluate this database, so we keep track of this
+	SearchContext.Track(this);
+#endif // UE_POSE_SEARCH_TRACE_ENABLED
 }
diff --git a/Engine/Plugins/Animation/PoseSearch/Source/Runtime/Private/PoseSearchInteractionIsland.cpp b/Engine/Plugins/Animation/PoseSearch/Source/Runtime/Private/PoseSearchInteractionIsland.cpp
--- a/Engine/Plugins/Animation/PoseSearch/Source/Runtime/Private/PoseSearchInteractionIsland.cpp
+++ b/Engine/Plugins/Animation/PoseSearch/Source/Runtime/Private/PoseSearchInteractionIsland.cpp
@@ -1046,7 +1046,9 @@ bool FInteractionIsland::DoSearch_AnyThread(const UObject* AnimContext, const TC
 
 		const UObject* AssetsToSearch[] = { Database };
 		// @todo: we could perform multiple UPoseSearchLibrary::MotionMatch in parallel!
-		PoseSearchResults[SearchIndex] = UPoseSearchLibrary::MotionMatch(SearchContext, AssetsToSearch, InteractionSearchContext.GetContinuingProperties());
+		FSearchResults_Single MotionMatchSearchResults;
+		UPoseSearchLibrary::MotionMatch(SearchContext, AssetsToSearch, InteractionSearchContext.GetContinuingProperties(), MotionMatchSearchResults);
+		PoseSearchResults[SearchIndex] = MotionMatchSearchResults.GetBestResult();
 	}
 		
 	InitSearchResults(SearchResults, PoseSearchResults, SearchContexts);
@@ -1090,7 +1092,7 @@ bool FInteractionIsland::GetResult_AnyThread(const UObject* AnimContext, FPoseSe
 				check(DatabaseAnimationAssetBase);
 
 				Result.SelectedAnim = DatabaseAnimationAssetBase->GetAnimationAsset();
-				Result.SelectedTime = SearchResult.AssetTime;
+				Result.SelectedTime = SearchResult.GetAssetTime();
 				Result.bIsContinuingPoseSearch = SearchResult.bIsContinuingPoseSearch;
 				Result.bLoop = SearchIndexAsset->IsLooping();
 				Result.bIsMirrored = SearchIndexAsset->IsMirrored();
diff --git a/Engine/Plugins/Animation/PoseSearch/Source/Runtime/Private/PoseSearchInteractionSubsystem.cpp b/Engine/Plugins/Animation/PoseSearch/Source/Runtime/Private/PoseSearchInteractionSubsystem.cpp
--- a/Engine/Plugins/Animation/PoseSearch/Source/Runtime/Private/PoseSearchInteractionSubsystem.cpp
+++ b/Engine/Plugins/Animation/PoseSearch/Source/Runtime/Private/PoseSearchInteractionSubsystem.cpp
@@ -421,7 +421,7 @@ void UPoseSearchInteractionSubsystem::PopulateContinuingProperties(float DeltaSe
 							// in case DatabaseAsset->GetAnimationAsset() is a blendspace, SearchResult->AssetTime is a normalized time in the interval [0,1] 
 							// so we need to convert the delta time in seconds to the asset normalized time before integrating SearchResult->AssetTime
 							const float NormalizedDeltaTime = DeltaSeconds / SearchIndexAsset->GetToRealTimeFactor();
-							SearchContext.PlayingAssetAccumulatedTime = SearchResult->AssetTime + NormalizedDeltaTime;
+							SearchContext.PlayingAssetAccumulatedTime = SearchResult->GetAssetTime() + NormalizedDeltaTime;
 							SearchContext.PlayingAsset = DatabaseAsset->GetAnimationAsset();
 							SearchContext.bIsPlayingAssetMirrored = SearchIndexAsset->IsMirrored();
 							SearchContext.PlayingAssetBlendParameters = SearchIndexAsset->GetBlendParameters();
diff --git a/Engine/Plugins/Animation/PoseSearch/Source/Runtime/Private/PoseSearchLibrary.cpp b/Engine/Plugins/Animation/PoseSearch/Source/Runtime/Private/PoseSearchLibrary.cpp
--- a/Engine/Plugins/Animation/PoseSearch/Source/Runtime/Private/PoseSearchLibrary.cpp
+++ b/Engine/Plugins/Animation/PoseSearch/Source/Runtime/Private/PoseSearchLibrary.cpp
@@ -727,14 +727,14 @@ void UPoseSearchLibrary::TraceMotionMatching(
 				FAnimationAssetSampler Sampler(AnimationAsset, FTransform::Identity,FVector::ZeroVector, FAnimationAssetSampler::DefaultRootTransformSamplingRate, true, false);
 
 				const float TimeStep = DeltaTime * WantedPlayRate;
-				const FTransform PrevRoot = Sampler.ExtractRootTransform(SearchResult.AssetTime);
-				const FTransform CurrRoot = Sampler.ExtractRootTransform(SearchResult.AssetTime + TimeStep);
+				const FTransform PrevRoot = Sampler.ExtractRootTransform(SearchResult.GetAssetTime());
+				const FTransform CurrRoot = Sampler.ExtractRootTransform(SearchResult.GetAssetTime() + TimeStep);
 				const FTransform RootMotionTransformDelta = PrevRoot.GetRelativeTransform(CurrRoot);
 				TraceState.AnimLinearVelocity = RootMotionTransformDelta.GetTranslation().Size() / DeltaTime;
 				TraceState.AnimAngularVelocity = FMath::RadiansToDegrees(RootMotionTransformDelta.GetRotation().GetAngle()) / DeltaTime;
 
 				// Need another root motion extraction for non-playrate version in case acceleration isn't the same.
-				const FTransform CurrRootNoTimescale = Sampler.ExtractRootTransform(SearchResult.AssetTime + DeltaTime);
+				const FTransform CurrRootNoTimescale = Sampler.ExtractRootTransform(SearchResult.GetAssetTime() + DeltaTime);
 				const FTransform RootMotionTransformDeltaNoTimescale = PrevRoot.GetRelativeTransform(CurrRootNoTimescale);
 				TraceState.AnimLinearVelocityNoTimescale = RootMotionTransformDeltaNoTimescale.GetTranslation().Size() / DeltaTime;
 				TraceState.AnimAngularVelocityNoTimescale = FMath::RadiansToDegrees(RootMotionTransformDeltaNoTimescale.GetRotation().GetAngle()) / DeltaTime;
@@ -744,7 +744,7 @@ void UPoseSearchLibrary::TraceMotionMatching(
 	}
 
 	TraceState.ElapsedPoseSearchTime = ElapsedPoseSearchTime;
-	TraceState.AssetPlayerTime = SearchResult.AssetTime;
+	TraceState.AssetPlayerTime = SearchResult.GetAssetTime();
 	TraceState.DeltaTime = DeltaTime;
 
 	TraceState.RecordingTime = RecordingTime;
@@ -898,35 +898,24 @@ void UPoseSearchLibrary::UpdateMotionMatchingState(
 		// calculating if it's worth bUseCachedChannelData (if we potentially have to build query with multiple schemas)
 		SearchContext.SetUseCachedChannelData(bShouldUseCachedChannelData && ShouldUseCachedChannelData(bSearchContinuingPose ? CurrentResultDatabase : nullptr, Databases));
 
+		FSearchResults_Single SearchResults;
+
 		// Evaluate continuing pose
 		if (bSearchContinuingPose)
 		{
-			SearchResult = CurrentResultDatabase->SearchContinuingPose(SearchContext);
-			SearchContext.UpdateCurrentBestCost(SearchResult.PoseCost);
+			CurrentResultDatabase->SearchContinuingPose(SearchContext, SearchResults);
 		}
 
 		for (const TObjectPtr<const UPoseSearchDatabase>& Database : Databases)
 		{
 			if (Database)
 			{
-				const FSearchResult NewSearchResult = Database->Search(SearchContext);
-
-#if WITH_EDITOR && ENABLE_ANIM_DEBUG && UE_POSE_SEARCH_TRACE_ENABLED
-				const FPoseSearchCost BestBruteForcePoseCost = NewSearchResult.BruteForcePoseCost < SearchResult.BruteForcePoseCost ? NewSearchResult.BruteForcePoseCost : SearchResult.BruteForcePoseCost;
-#endif // WITH_EDITOR && ENABLE_ANIM_DEBUG && UE_POSE_SEARCH_TRACE_ENABLED
-
-				if (NewSearchResult.PoseCost < SearchResult.PoseCost)
-				{
-					SearchResult = NewSearchResult;
-					SearchContext.UpdateCurrentBestCost(SearchResult.PoseCost);
-				}
-
-#if WITH_EDITOR && ENABLE_ANIM_DEBUG && UE_POSE_SEARCH_TRACE_ENABLED
-				SearchResult.BruteForcePoseCost = BestBruteForcePoseCost;
-#endif // WITH_EDITOR && ENABLE_ANIM_DEBUG && UE_POSE_SEARCH_TRACE_ENABLED
+				Database->Search(SearchContext, SearchResults);
 			}
 		}
 
+		SearchResult = SearchResults.GetBestResult();
+
 #if !NO_LOGGING
 		if (!SearchResult.IsValid())
 		{
@@ -1168,6 +1157,22 @@ UE::PoseSearch::FSearchResult UPoseSearchLibrary::MotionMatch(
 	const FPoseSearchContinuingProperties& ContinuingProperties,
 	const FPoseSearchFutureProperties& Future,
 	const FPoseSearchEvent& EventToSearch)
+{
+	using namespace UE::PoseSearch;
+	FSearchResults_Single SearchResults;
+	MotionMatch(AnimContexts, Roles, PoseHistories, AssetsToSearch, ContinuingProperties, Future, EventToSearch, SearchResults);
+	return SearchResults.GetBestResult();
+}
+
+void UPoseSearchLibrary::MotionMatch(
+	const TArrayView<const UObject*> AnimContexts,
+	const TArrayView<const UE::PoseSearch::FRole> Roles,
+	const TArrayView<const UE::PoseSearch::IPoseHistory*> PoseHistories, 
+	const TArrayView<const UObject*> AssetsToSearch,
+	const FPoseSearchContinuingProperties& ContinuingProperties,
+	const FPoseSearchFutureProperties& Future,
+	const FPoseSearchEvent& EventToSearch,
+	UE::PoseSearch::ISearchResults& SearchResults)
 {
 	TArray<FChooserEvaluationContext, TInlineAllocator<UE::PoseSearch::PreallocatedRolesNum>> Contexts;
 	const int NumContexts = AnimContexts.Num(); 
@@ -1177,17 +1182,18 @@ UE::PoseSearch::FSearchResult UPoseSearchLibrary::MotionMatch(
 		Contexts[i].AddObjectParam(const_cast<UObject*>(AnimContexts[i]));
 	}
 
-	return MotionMatch(Contexts, Roles, PoseHistories, AssetsToSearch, ContinuingProperties, Future, EventToSearch);
+	MotionMatch(Contexts, Roles, PoseHistories, AssetsToSearch, ContinuingProperties, Future, EventToSearch, SearchResults);
 }
 
-UE::PoseSearch::FSearchResult UPoseSearchLibrary::MotionMatch(
+void UPoseSearchLibrary::MotionMatch(
 	const TArrayView<const UObject*> AnimContexts,
 	const TArrayView<const UE::PoseSearch::FRole> Roles,
 	const TArrayView<const UE::PoseSearch::IPoseHistory*> PoseHistories, 
 	const TArrayView<const UObject*> AssetsToSearch,
 	const FPoseSearchContinuingProperties& ContinuingProperties,
 	float DesiredPermutationTimeOffset,
-	const FPoseSearchEvent& EventToSearch)
+	const FPoseSearchEvent& EventToSearch,
+	UE::PoseSearch::ISearchResults& SearchResults)
 {
 	TArray<FChooserEvaluationContext, TInlineAllocator<UE::PoseSearch::PreallocatedRolesNum>> Contexts;
 	const int NumContexts = AnimContexts.Num(); 
@@ -1197,24 +1203,23 @@ UE::PoseSearch::FSearchResult UPoseSearchLibrary::MotionMatch(
 		Contexts[i].AddObjectParam(const_cast<UObject*>(AnimContexts[i]));
 	}
 
-	return MotionMatch(Contexts, Roles, PoseHistories, AssetsToSearch, ContinuingProperties, DesiredPermutationTimeOffset, EventToSearch);
+	MotionMatch(Contexts, Roles, PoseHistories, AssetsToSearch, ContinuingProperties, DesiredPermutationTimeOffset, EventToSearch, SearchResults);
 }
 
-UE::PoseSearch::FSearchResult UPoseSearchLibrary::MotionMatch(
+void UPoseSearchLibrary::MotionMatch(
 	const TArrayView<FChooserEvaluationContext> Contexts,
 	const TArrayView<const UE::PoseSearch::FRole> Roles,
 	const TArrayView<const UE::PoseSearch::IPoseHistory*> PoseHistories, 
 	const TArrayView<const UObject*> AssetsToSearch,
 	const FPoseSearchContinuingProperties& ContinuingProperties,
 	const FPoseSearchFutureProperties& Future,
-	const FPoseSearchEvent& EventToSearch)
+	const FPoseSearchEvent& EventToSearch,
+	UE::PoseSearch::ISearchResults& SearchResults)
 {
 	check(!Contexts.IsEmpty() && Contexts.Num() == Roles.Num() && Contexts.Num() == PoseHistories.Num());
 
 	using namespace UE::PoseSearch;
 
-	FMemMark Mark(FMemStack::Get());
-
 	TArray<const IPoseHistory*, TInlineAllocator<PreallocatedRolesNum, TMemStackAllocator<>>> InternalPoseHistories;
 	InternalPoseHistories = PoseHistories;
 
@@ -1267,17 +1272,18 @@ UE::PoseSearch::FSearchResult UPoseSearchLibrary::MotionMatch(
 		}
 	}		
 
-	return MotionMatch(Contexts, Roles, InternalPoseHistories, AssetsToSearch, ContinuingProperties, FutureIntervalTime, EventToSearch);
+	MotionMatch(Contexts, Roles, InternalPoseHistories, AssetsToSearch, ContinuingProperties, FutureIntervalTime, EventToSearch, SearchResults);
 }
 
-UE::PoseSearch::FSearchResult UPoseSearchLibrary::MotionMatch(
+void UPoseSearchLibrary::MotionMatch(
 	const TArrayView<FChooserEvaluationContext> Contexts,
 	const TArrayView<const UE::PoseSearch::FRole> Roles,
 	const TArrayView<const UE::PoseSearch::IPoseHistory*> PoseHistories,
 	const TArrayView<const UObject*> AssetsToSearch,
 	const FPoseSearchContinuingProperties& ContinuingProperties,
 	const float DesiredPermutationTimeOffset,
-	const FPoseSearchEvent& EventToSearch)
+	const FPoseSearchEvent& EventToSearch,
+	UE::PoseSearch::ISearchResults& SearchResults)
 {
 	using namespace UE::PoseSearch;
 	
@@ -1287,18 +1293,17 @@ UE::PoseSearch::FSearchResult UPoseSearchLibrary::MotionMatch(
 		SearchContext.AddRole(Roles[RoleIndex], &Contexts[RoleIndex], PoseHistories[RoleIndex]);
 	}
 
-	return MotionMatch(SearchContext, AssetsToSearch, ContinuingProperties);
+	MotionMatch(SearchContext, AssetsToSearch, ContinuingProperties, SearchResults);
 }
 
-UE::PoseSearch::FSearchResult UPoseSearchLibrary::MotionMatch(
+void UPoseSearchLibrary::MotionMatch(
 	UE::PoseSearch::FSearchContext& SearchContext,
 	const TArrayView<const UObject*> AssetsToSearch,
-	const FPoseSearchContinuingProperties& ContinuingProperties)
+	const FPoseSearchContinuingProperties& ContinuingProperties,
+	UE::PoseSearch::ISearchResults& SearchResults)
 {
 	using namespace UE::PoseSearch;
 
-	FSearchResult SearchResult;
-
 	TConstArrayView<const UObject*> CurrentAssetToConsider = SearchContext.GetAssetsToConsider();
 	SearchContext.SetIsContinuingInteraction(ContinuingProperties.bIsContinuingInteraction);
 
@@ -1321,7 +1326,7 @@ UE::PoseSearch::FSearchResult UPoseSearchLibrary::MotionMatch(
 		{
 			// reconstructing and caching all the required continuing pose search results
 			FSearchResult DatabaseContinuingPoseSearchResult;
-			DatabaseContinuingPoseSearchResult.AssetTime = ContinuingProperties.PlayingAssetAccumulatedTime;
+			DatabaseContinuingPoseSearchResult.SetAssetTime(ContinuingProperties.PlayingAssetAccumulatedTime);
 			DatabaseContinuingPoseSearchResult.PoseIdx = Database->GetPoseIndex(ContinuingProperties.PlayingAsset.Get(), ContinuingProperties.PlayingAssetAccumulatedTime, ContinuingProperties.bIsPlayingAssetMirrored, ContinuingProperties.PlayingAssetBlendParameters);
 			DatabaseContinuingPoseSearchResult.Database = Database;
 			CachedContinuingPoseSearchResults.CheckedAdd(Database, DatabaseContinuingPoseSearchResult);
@@ -1335,22 +1340,7 @@ UE::PoseSearch::FSearchResult UPoseSearchLibrary::MotionMatch(
 			if (bCanAdvance && !bForceInterrupt)
 			{
 				SearchContext.UpdateContinuingPoseSearchResult(DatabaseContinuingPoseSearchResult, SchemaContinuingPoseSearchResult);
-
-				const FSearchResult NewSearchResult = Database->SearchContinuingPose(SearchContext);
-
-#if WITH_EDITOR && ENABLE_ANIM_DEBUG && UE_POSE_SEARCH_TRACE_ENABLED
-				const FPoseSearchCost BestBruteForcePoseCost = NewSearchResult.BruteForcePoseCost < SearchResult.BruteForcePoseCost ? NewSearchResult.BruteForcePoseCost : SearchResult.BruteForcePoseCost;
-#endif // WITH_EDITOR && ENABLE_ANIM_DEBUG && UE_POSE_SEARCH_TRACE_ENABLED
-
-				if (NewSearchResult.PoseCost < SearchResult.PoseCost)
-				{
-					SearchResult = NewSearchResult;
-					SearchContext.UpdateCurrentBestCost(SearchResult.PoseCost);
-				}
-
-#if WITH_EDITOR && ENABLE_ANIM_DEBUG && UE_POSE_SEARCH_TRACE_ENABLED
-				SearchResult.BruteForcePoseCost = BestBruteForcePoseCost;
-#endif // WITH_EDITOR && ENABLE_ANIM_DEBUG && UE_POSE_SEARCH_TRACE_ENABLED
+				Database->SearchContinuingPose(SearchContext, SearchResults);
 			}
 		}
 	}
@@ -1368,23 +1358,13 @@ UE::PoseSearch::FSearchResult UPoseSearchLibrary::MotionMatch(
 		SearchContext.UpdateContinuingPoseSearchResult(CachedContinuingPoseSearchResults.FindOrDefault(Database), CachedContinuingPoseSearchResults.FindOrDefault(Database->Schema));
 		SearchContext.SetAssetsToConsider(AssetsToSearchPerDatabasePair.Value);
 		
-		const FSearchResult NewSearchResult = Database->Search(SearchContext);
-
-#if WITH_EDITOR && ENABLE_ANIM_DEBUG && UE_POSE_SEARCH_TRACE_ENABLED
-		const FPoseSearchCost BestBruteForcePoseCost = NewSearchResult.BruteForcePoseCost < SearchResult.BruteForcePoseCost ? NewSearchResult.BruteForcePoseCost : SearchResult.BruteForcePoseCost;
-#endif // WITH_EDITOR && ENABLE_ANIM_DEBUG && UE_POSE_SEARCH_TRACE_ENABLED
-
-		if (NewSearchResult.PoseCost < SearchResult.PoseCost)
-		{
-			SearchResult = NewSearchResult;
-			SearchContext.UpdateCurrentBestCost(SearchResult.PoseCost);
-		}
-
-#if WITH_EDITOR && ENABLE_ANIM_DEBUG && UE_POSE_SEARCH_TRACE_ENABLED
-		SearchResult.BruteForcePoseCost = BestBruteForcePoseCost;
-#endif // WITH_EDITOR && ENABLE_ANIM_DEBUG && UE_POSE_SEARCH_TRACE_ENABLED
+		Database->Search(SearchContext, SearchResults);
 	}
 
+#if (ENABLE_DRAW_DEBUG && ENABLE_ANIM_DEBUG) || UE_POSE_SEARCH_TRACE_ENABLED
+	const FSearchResult SearchResult = SearchResults.GetBestResult();
+#endif // (ENABLE_DRAW_DEBUG && ENABLE_ANIM_DEBUG) || UE_POSE_SEARCH_TRACE_ENABLED
+
 #if ENABLE_DRAW_DEBUG && ENABLE_ANIM_DEBUG
 	if (SearchResult.IsValid())
 	{
@@ -1419,14 +1399,11 @@ UE::PoseSearch::FSearchResult UPoseSearchLibrary::MotionMatch(
 		}
 	}
 
-	// @todo: add and handle an interrupt mode input param MotionMatch
-	TraceMotionMatching(SearchContext, SearchResult, 0.f, DeltaSeconds, true, 1.f, EPoseSearchInterruptMode::DoNotInterrupt);
+	TraceMotionMatching(SearchContext, SearchResult, 0.f, DeltaSeconds, true, 1.f, ContinuingProperties.InterruptMode);
 #endif // UE_POSE_SEARCH_TRACE_ENABLED
 
 	// restoring the asset to consider array view
 	SearchContext.SetAssetsToConsider(CurrentAssetToConsider);
-
-	return SearchResult;
 }
 
 const FAnimNode_PoseSearchHistoryCollector_Base* UPoseSearchLibrary::FindPoseHistoryNode(
diff --git a/Engine/Plugins/Animation/PoseSearch/Source/Runtime/Private/PoseSearchModule.cpp b/Engine/Plugins/Animation/PoseSearch/Source/Runtime/Private/PoseSearchModule.cpp
--- a/Engine/Plugins/Animation/PoseSearch/Source/Runtime/Private/PoseSearchModule.cpp
+++ b/Engine/Plugins/Animation/PoseSearch/Source/Runtime/Private/PoseSearchModule.cpp
@@ -78,14 +78,14 @@ public:
 			MakeArrayView(&DefaultRole, 1), MakeArrayView(&PoseHistory, 1), AssetsToSearch, ContinuingProperties, Future, FPoseSearchEvent());
 
 		UE::Anim::IPoseSearchProvider::FSearchResult ProviderResult;
-		if (const UE::PoseSearch::FSearchIndexAsset* SearchIndexAsset = SearchResult.GetSearchIndexAsset())
+		if (const FSearchIndexAsset* SearchIndexAsset = SearchResult.GetSearchIndexAsset())
 		{
 			const UPoseSearchDatabase* Database = SearchResult.Database.Get();
 			if (const FPoseSearchDatabaseAnimationAssetBase* DatabaseAnimationAssetBase = Database->GetDatabaseAnimationAsset(*SearchIndexAsset))
 			{
 				ProviderResult.SelectedAsset = DatabaseAnimationAssetBase->GetAnimationAsset();
 				ProviderResult.Dissimilarity = SearchResult.PoseCost;
-				ProviderResult.TimeOffsetSeconds = SearchResult.AssetTime;
+				ProviderResult.TimeOffsetSeconds = SearchResult.GetAssetTime();
 				ProviderResult.bIsFromContinuingPlaying = SearchResult.bIsContinuingPoseSearch;
 				ProviderResult.bMirrored = SearchIndexAsset->IsMirrored();
 				ProviderResult.BlendParameters = SearchIndexAsset->GetBlendParameters();
diff --git a/Engine/Plugins/Animation/PoseSearch/Source/Runtime/Private/PoseSearchResult.cpp b/Engine/Plugins/Animation/PoseSearch/Source/Runtime/Private/PoseSearchResult.cpp
--- a/Engine/Plugins/Animation/PoseSearch/Source/Runtime/Private/PoseSearchResult.cpp
+++ b/Engine/Plugins/Animation/PoseSearch/Source/Runtime/Private/PoseSearchResult.cpp
@@ -62,7 +62,9 @@ bool FSearchResult::DebugValidate() const
 		const FSearchIndexAsset& SearchIndexAsset = SearchIndex.GetAssetForPose(PoseIdx);
 		if (const FPoseSearchDatabaseAnimationAssetBase* DatabaseAnimationAssetBase = Database->GetDatabaseAnimationAsset(SearchIndexAsset))
 		{
+			PRAGMA_DISABLE_DEPRECATION_WARNINGS
 			const float RealAssetTime = AssetTime * SearchIndexAsset.GetToRealTimeFactor();
+			PRAGMA_ENABLE_DEPRECATION_WARNINGS
 			const int32 RecalculatedPoseIdx = SearchIndexAsset.GetPoseIndexFromTime(RealAssetTime, Database->Schema->SampleRate);
 			if (RecalculatedPoseIdx != PoseIdx)
 			{
@@ -115,7 +117,9 @@ void FSearchResult::UpdateWithNormalizedTime(float NormalizedTime)
 #endif // WITH_EDITOR && DO_CHECK
 		
 		PoseIdx = SearchIndexAsset.GetPoseIndexFromTime(RealTime, Database->Schema->SampleRate);
+		PRAGMA_DISABLE_DEPRECATION_WARNINGS
 		AssetTime = NormalizedTime;
+		PRAGMA_ENABLE_DEPRECATION_WARNINGS
 	}
 	else
 	{
@@ -160,7 +164,9 @@ void FSearchResult::UpdateWithRealTime(float RealTime)
 #endif // WITH_EDITOR && DO_CHECK
 
 		PoseIdx = SearchIndexAsset.GetPoseIndexFromTime(RealTime, Database->Schema->SampleRate);
+		PRAGMA_DISABLE_DEPRECATION_WARNINGS
 		AssetTime = NormalizedTime;
+		PRAGMA_ENABLE_DEPRECATION_WARNINGS
 	}
 	else
 	{
@@ -191,7 +197,9 @@ float FSearchResult::CalculateTimeToEvent() const
 	const float PoseQuantizedTime = SearchIndexAsset.GetTimeFromPoseIndex(PoseIdx, Database->Schema->SampleRate);
 	// AssetTime is the current search result time, that differs from PoseQuantizedTime in case this search result is from a continuing pose search
 	// we calculate the QuantizationError as differnece between AssetTime and PoseQuantizedTime, time in seconds that this search result drifted away from the quantized time of the associated PoseIdx
+	PRAGMA_DISABLE_DEPRECATION_WARNINGS
 	const float QuantizationError = AssetTime - PoseQuantizedTime;
+	PRAGMA_ENABLE_DEPRECATION_WARNINGS
 
 	// time to event can be negative if PoseIdx already passed EventPoseIdx and the asset is not looping
 	const float TimeToEvent = DeltaTimeBetweenPoseAndEvent - QuantizationError;
@@ -204,7 +212,9 @@ void FSearchResult::InitFrom(const FPoseSearchBlueprintResult& BlueprintResult)
 	PoseIdx = BlueprintResult.SelectedDatabase ? BlueprintResult.SelectedDatabase->GetPoseIndex(BlueprintResult.SelectedAnim.Get(), BlueprintResult.SelectedTime, BlueprintResult.bIsMirrored, BlueprintResult.BlendParameters) : INDEX_NONE;
 	EventPoseIdx = INDEX_NONE;
 	Database = BlueprintResult.SelectedDatabase;
+	PRAGMA_DISABLE_DEPRECATION_WARNINGS
 	AssetTime = BlueprintResult.SelectedTime;
+	PRAGMA_ENABLE_DEPRECATION_WARNINGS
 	bIsContinuingPoseSearch = BlueprintResult.bIsContinuingPoseSearch;
 
 #if WITH_EDITOR && ENABLE_ANIM_DEBUG && UE_POSE_SEARCH_TRACE_ENABLED
@@ -213,6 +223,154 @@ void FSearchResult::InitFrom(const FPoseSearchBlueprintResult& BlueprintResult)
 #endif // WITH_EDITOR && ENABLE_ANIM_DEBUG && UE_POSE_SEARCH_TRACE_ENABLED
 }
 
+/////////////////////////////////////////////////////
+// FSearchResults_Single
+void FSearchResults_Single::UpdateWith(FSearchResult& SearchResult)
+{
+	check(SearchResult.IsValid());
+	if (SearchResult.PoseCost < SingleSearchResult.PoseCost)
+	{
+		SingleSearchResult = SearchResult;
+	}
+}
+
+void FSearchResults_Single::FinalizeResults()
+{
+	if (SingleSearchResult.IsValid() && !SingleSearchResult.IsAssetTimeValid())
+	{
+		SingleSearchResult.SetAssetTime(SingleSearchResult.Database->GetNormalizedAssetTime(SingleSearchResult.PoseIdx));
+		check(SingleSearchResult.DebugValidate());
+	}
+}
+
+FSearchResult FSearchResults_Single::GetBestResult() const
+{
+	return SingleSearchResult;
+}
+
+bool FSearchResults_Single::ShouldPerformSearch(float SearchMinimumCost) const
+{
+	return SingleSearchResult.PoseCost > SearchMinimumCost;
+}
+
+/////////////////////////////////////////////////////
+// FSearchResults_Multi
+FSearchResults_Multi::FSearchResults_Multi(int32 InMaxNumberOfWantedResults)
+	: MaxNumberOfWantedResults(InMaxNumberOfWantedResults)
+{
+	check(MaxNumberOfWantedResults > 0); // or else we need to refine the logic in UpdateWith
+	SearchResults.Reserve(MaxNumberOfWantedResults);
+}
+
+void FSearchResults_Multi::UpdateWith(FSearchResult& SearchResult)
+{
+	check(SearchResult.IsValid());
+	if (SearchResults.Num() < MaxNumberOfWantedResults)
+	{
+		SearchResults.Add(SearchResult);
+
+		// @todo: speed up this sort, since SearchResults started alreay sorted!
+		SearchResults.Sort([](const FSearchResult& A, const FSearchResult& B)
+			{
+				return A.PoseCost < B.PoseCost;
+			});
+	}
+	else
+	{
+		FSearchResult& WorstResult = SearchResults.Last();
+		if (SearchResult.PoseCost < WorstResult.PoseCost)
+		{
+			WorstResult = SearchResult;
+
+			// @todo: speed up this sort, since SearchResults started alreay sorted!
+			SearchResults.Sort([](const FSearchResult& A, const FSearchResult& B)
+				{
+					return A.PoseCost < B.PoseCost;
+				});
+		}
+	}
+}
+
+void FSearchResults_Multi::FinalizeResults()
+{
+	for (FSearchResult& SearchResult : SearchResults)
+	{
+		check(SearchResult.IsValid());
+		if (!SearchResult.IsAssetTimeValid())
+		{
+			SearchResult.SetAssetTime(SearchResult.Database->GetNormalizedAssetTime(SearchResult.PoseIdx));
+			check(SearchResult.DebugValidate());
+		}
+	}
+}
+
+FSearchResult FSearchResults_Multi::GetBestResult() const
+{
+	if (!SearchResults.IsEmpty())
+	{
+		return SearchResults[0];
+	}
+	return FSearchResult();
+}
+
+bool FSearchResults_Multi::ShouldPerformSearch(float SearchMinimumCost) const
+{
+	if (SearchResults.IsEmpty())
+	{
+		return true;
+	}
+
+	return SearchResults.Last().PoseCost > SearchMinimumCost;
+}
+
+/////////////////////////////////////////////////////
+// FSearchResults_AssetBests
+void FSearchResults_AssetBests::UpdateWith(FSearchResult& SearchResult)
+{
+	check(SearchResult.IsValid());
+	if (const FSearchIndexAsset* SearchIndexAsset = SearchResult.GetSearchIndexAsset())
+	{
+		FSearchResult& BestSearchResultForSourceAssetIdx = SourceAssetIdxToBestSearchResult.FindOrAdd(SearchIndexAsset->GetSourceAssetIdx());
+		if (SearchResult.PoseCost < BestSearchResultForSourceAssetIdx.PoseCost)
+		{
+			BestSearchResultForSourceAssetIdx = SearchResult;
+		}
+	}
+}
+
+void FSearchResults_AssetBests::FinalizeResults()
+{
+	for (TPair<int32, FSearchResult>& SourceAssetIdxToBestSearchResultPair : SourceAssetIdxToBestSearchResult)
+	{
+		FSearchResult& BestSearchResultForSourceAssetIdx = SourceAssetIdxToBestSearchResultPair.Value;
+		check(BestSearchResultForSourceAssetIdx.IsValid());
+		if (!BestSearchResultForSourceAssetIdx.IsAssetTimeValid())
+		{
+			BestSearchResultForSourceAssetIdx.SetAssetTime(BestSearchResultForSourceAssetIdx.Database->GetNormalizedAssetTime(BestSearchResultForSourceAssetIdx.PoseIdx));
+			check(BestSearchResultForSourceAssetIdx.DebugValidate());
+		}
+	}
+}
+
+FSearchResult FSearchResults_AssetBests::GetBestResult() const
+{
+	FSearchResult BestSearchResult;
+	for (const TPair<int32, FSearchResult>& SourceAssetIdxToBestSearchResultPair : SourceAssetIdxToBestSearchResult)
+	{
+		const FSearchResult& BestSearchResultForSourceAssetIdx = SourceAssetIdxToBestSearchResultPair.Value;
+		if (BestSearchResultForSourceAssetIdx.PoseCost < BestSearchResult.PoseCost)
+		{
+			BestSearchResult = BestSearchResultForSourceAssetIdx;
+		}
+	}
+	return BestSearchResult;
+}
+
+bool FSearchResults_AssetBests::ShouldPerformSearch(float SearchMinimumCost) const
+{
+	return true;
+}
+
 } // namespace UE::PoseSearch
 
 bool FPoseSearchBlueprintResult::InitFrom(const UE::PoseSearch::FSearchResult& SearchResult, float InWantedPlayRate)
@@ -231,7 +389,7 @@ bool FPoseSearchBlueprintResult::InitFrom(const UE::PoseSearch::FSearchResult& S
 #endif // WITH_EDITORONLY_DATA
 
 		SelectedAnim = DatabaseAsset->GetAnimationAsset();
-		SelectedTime = SearchResult.AssetTime;
+		SelectedTime = SearchResult.GetAssetTime();
 		bIsContinuingPoseSearch = SearchResult.bIsContinuingPoseSearch;
 		WantedPlayRate = InWantedPlayRate;
 		bLoop = SearchIndexAsset->IsLooping();
diff --git a/Engine/Plugins/Animation/PoseSearch/Source/Runtime/Public/PoseSearch/PoseSearchContext.h b/Engine/Plugins/Animation/PoseSearch/Source/Runtime/Public/PoseSearch/PoseSearchContext.h
--- a/Engine/Plugins/Animation/PoseSearch/Source/Runtime/Public/PoseSearch/PoseSearchContext.h
+++ b/Engine/Plugins/Animation/PoseSearch/Source/Runtime/Public/PoseSearch/PoseSearchContext.h
@@ -185,9 +185,14 @@ struct FSearchContext
 	// if bUseCharacterSpaceVelocities is true, velocities will be computed in root bone space, rather than world space
 	UE_API FVector GetSampleVelocity(float SampleTimeOffset, float OriginTimeOffset, int8 SchemaSampleBoneIdx, int8 SchemaOriginBoneIdx, const FRole& SampleRole, const FRole& OriginRole, bool bUseCharacterSpaceVelocities = true, EPermutationTimeType PermutationTimeType = EPermutationTimeType::UseSampleTime, const FVector* SampleBoneVelocityWorldOverride = nullptr);
 
+	UE_DEPRECATED(5.7, "No longer necessary API")
 	UE_API void ResetCurrentBestCost();
+
+	UE_DEPRECATED(5.7, "No longer necessary API")
 	UE_API void UpdateCurrentBestCost(const FPoseSearchCost& PoseSearchCost);
-	float GetCurrentBestTotalCost() const { return CurrentBestTotalCost; }
+
+	UE_DEPRECATED(5.7, "No longer necessary API. Use the FSearchResult::PoseCost from UPoseSearchDatabase::Search if a search cost is needed")
+	float GetCurrentBestTotalCost() const { return MAX_FLT; }
 
 	UE_API TConstArrayView<float> GetOrBuildQuery(const UPoseSearchSchema* Schema);
 	UE_API TConstArrayView<float> GetCachedQuery(const UPoseSearchSchema* Schema) const;
@@ -329,8 +334,6 @@ private:
 	// mapping channel unique identifier (hash) to FCachedChannel
 	TMap<uint32, FCachedChannel, TInlineSetAllocator<PreallocatedCachedChannelDataNum, TMemStackSetAllocator<>>> CachedChannels;
 
-	float CurrentBestTotalCost = MAX_flt;
-
 	// @todo: add FPoseSearchContinuingProperties here and reconstuct the continuing pose FSearchResult for the previous frame from it
 	// Experimental, this feature might be removed without warning, not for production use
 	bool bIsContinuingInteraction = false;
diff --git a/Engine/Plugins/Animation/PoseSearch/Source/Runtime/Public/PoseSearch/PoseSearchDatabase.h b/Engine/Plugins/Animation/PoseSearch/Source/Runtime/Public/PoseSearch/PoseSearchDatabase.h
--- a/Engine/Plugins/Animation/PoseSearch/Source/Runtime/Public/PoseSearch/PoseSearchDatabase.h
+++ b/Engine/Plugins/Animation/PoseSearch/Source/Runtime/Public/PoseSearch/PoseSearchDatabase.h
@@ -682,8 +682,15 @@ public:
 	UE_API virtual void Serialize(FArchive& Ar) override;
 	// End UObject
 	
+	// deprecate in favor of Search(UE::PoseSearch::FSearchContext& SearchContext, UE::PoseSearch::ISearchResults& SearchResults) once we settle on ISearchResults API signatures
 	UE_API UE::PoseSearch::FSearchResult Search(UE::PoseSearch::FSearchContext& SearchContext) const;
+	// deprecate in favor of SearchContinuingPose(UE::PoseSearch::FSearchContext& SearchContext, UE::PoseSearch::ISearchResults& SearchResults) once we settle on ISearchResults API signatures
 	UE_API UE::PoseSearch::FSearchResult SearchContinuingPose(UE::PoseSearch::FSearchContext& SearchContext) const;
+	
+	// Experimental, this feature might be removed without warning, not for production use
+	UE_API void Search(UE::PoseSearch::FSearchContext& SearchContext, UE::PoseSearch::ISearchResults& SearchResults) const;
+	// Experimental, this feature might be removed without warning, not for production use
+	UE_API void SearchContinuingPose(UE::PoseSearch::FSearchContext& SearchContext, UE::PoseSearch::ISearchResults& SearchResults) const;
 
 	UE_API bool Contains(const UObject* Object) const;
 
@@ -726,10 +733,10 @@ public:
 #endif // WITH_EDITORONLY_DATA
 
 private:
-	UE_API UE::PoseSearch::FSearchResult SearchPCAKDTree(UE::PoseSearch::FSearchContext& SearchContext) const;
-	UE_API UE::PoseSearch::FSearchResult SearchVPTree(UE::PoseSearch::FSearchContext& SearchContext) const;
-	UE_API UE::PoseSearch::FSearchResult SearchBruteForce(UE::PoseSearch::FSearchContext& SearchContext) const;
-	UE_API UE::PoseSearch::FSearchResult SearchEvent(UE::PoseSearch::FSearchContext& SearchContext) const;
+	UE_API void SearchPCAKDTree(UE::PoseSearch::FSearchContext& SearchContext, UE::PoseSearch::ISearchResults& SearchResults) const;
+	UE_API void SearchVPTree(UE::PoseSearch::FSearchContext& SearchContext, UE::PoseSearch::ISearchResults& SearchResults) const;
+	UE_API void SearchBruteForce(UE::PoseSearch::FSearchContext& SearchContext, UE::PoseSearch::ISearchResults& SearchResults) const;
+	UE_API void SearchEvent(UE::PoseSearch::FSearchContext& SearchContext, UE::PoseSearch::ISearchResults& SearchResults) const;
 
 	typedef TArray<int32, TInlineAllocator<256, TMemStackAllocator<>>> FSelectableAssetIdx;
 	UE_API void PopulateSelectableAssetIdx(FSelectableAssetIdx& SelectableAssetIdx, TConstArrayView<const UObject*> AssetsToConsider) const;
diff --git a/Engine/Plugins/Animation/PoseSearch/Source/Runtime/Public/PoseSearch/PoseSearchLibrary.h b/Engine/Plugins/Animation/PoseSearch/Source/Runtime/Public/PoseSearch/PoseSearchLibrary.h
--- a/Engine/Plugins/Animation/PoseSearch/Source/Runtime/Public/PoseSearch/PoseSearchLibrary.h
+++ b/Engine/Plugins/Animation/PoseSearch/Source/Runtime/Public/PoseSearch/PoseSearchLibrary.h
@@ -295,6 +295,7 @@ public:
 		const FPoseSearchFutureProperties& Future,
 		FPoseSearchBlueprintResult& Result);
 
+	// deprecate in favor of MotionMatch with the ISearchResults parameter once we settle on ISearchResults API signatures
 	static UE_API UE::PoseSearch::FSearchResult MotionMatch(
 		const TArrayView<const UObject*> AnimContexts,
 		const TArrayView<const UE::PoseSearch::FRole> Roles,
@@ -303,42 +304,57 @@ public:
 		const FPoseSearchContinuingProperties& ContinuingProperties,
 		const FPoseSearchFutureProperties& Future,
 		const FPoseSearchEvent& EventToSearch);
-	
+
 	// Experimental, this feature might be removed without warning, not for production use
-	static UE_API UE::PoseSearch::FSearchResult MotionMatch(
+	static UE_API void MotionMatch(
+		const TArrayView<const UObject*> AnimContexts,
+		const TArrayView<const UE::PoseSearch::FRole> Roles,
+		const TArrayView<const UE::PoseSearch::IPoseHistory*> PoseHistories, 
+		const TArrayView<const UObject*> AssetsToSearch,
+		const FPoseSearchContinuingProperties& ContinuingProperties,
+		const FPoseSearchFutureProperties& Future,
+		const FPoseSearchEvent& EventToSearch,
+		UE::PoseSearch::ISearchResults& SearchResults);
+
+	// Experimental, this feature might be removed without warning, not for production use
+	static UE_API void MotionMatch(
 		const TArrayView<FChooserEvaluationContext> Contexts,
 		const TArrayView<const UE::PoseSearch::FRole> Roles,
 		const TArrayView<const UE::PoseSearch::IPoseHistory*> PoseHistories, 
 		const TArrayView<const UObject*> AssetsToSearch,
 		const FPoseSearchContinuingProperties& ContinuingProperties,
 		const FPoseSearchFutureProperties& Future,
-		const FPoseSearchEvent& EventToSearch);
+		const FPoseSearchEvent& EventToSearch,
+		UE::PoseSearch::ISearchResults& SearchResults);
 			
 	// Experimental, this feature might be removed without warning, not for production use
-	static UE_API UE::PoseSearch::FSearchResult MotionMatch(
+	static UE_API void MotionMatch(
 		const TArrayView<FChooserEvaluationContext> Contexts,
 		const TArrayView<const UE::PoseSearch::FRole> Roles,
 		const TArrayView<const UE::PoseSearch::IPoseHistory*> PoseHistories,
 		const TArrayView<const UObject*> AssetsToSearch,
 		const FPoseSearchContinuingProperties& ContinuingProperties,
 		const float DesiredPermutationTimeOffset,
-		const FPoseSearchEvent& EventToSearch);
+		const FPoseSearchEvent& EventToSearch,
+		UE::PoseSearch::ISearchResults& SearchResults);
 
 	// Experimental, this feature might be removed without warning, not for production use
-	static UE_API UE::PoseSearch::FSearchResult MotionMatch(
+	static UE_API void MotionMatch(
 		const TArrayView<const UObject*> AnimContexts,
 		const TArrayView<const UE::PoseSearch::FRole> Roles,
 		const TArrayView<const UE::PoseSearch::IPoseHistory*> PoseHistories,
 		const TArrayView<const UObject*> AssetsToSearch,
 		const FPoseSearchContinuingProperties& ContinuingProperties,
 		const float DesiredPermutationTimeOffset,
-		const FPoseSearchEvent& EventToSearch);
+		const FPoseSearchEvent& EventToSearch,
+		UE::PoseSearch::ISearchResults& SearchResults);
 
 	// Experimental, this feature might be removed without warning, not for production use
-	static UE_API UE::PoseSearch::FSearchResult MotionMatch(
+	static UE_API void MotionMatch(
 		UE::PoseSearch::FSearchContext& SearchContext,
 		const TArrayView<const UObject*> AssetsToSearch,
-		const FPoseSearchContinuingProperties& ContinuingProperties);
+		const FPoseSearchContinuingProperties& ContinuingProperties,
+		UE::PoseSearch::ISearchResults& SearchResults);
 
 	static UE_API const FAnimNode_PoseSearchHistoryCollector_Base* FindPoseHistoryNode(
 		const FName PoseHistoryName,
diff --git a/Engine/Plugins/Animation/PoseSearch/Source/Runtime/Public/PoseSearch/PoseSearchResult.h b/Engine/Plugins/Animation/PoseSearch/Source/Runtime/Public/PoseSearch/PoseSearchResult.h
--- a/Engine/Plugins/Animation/PoseSearch/Source/Runtime/Public/PoseSearch/PoseSearchResult.h
+++ b/Engine/Plugins/Animation/PoseSearch/Source/Runtime/Public/PoseSearch/PoseSearchResult.h
@@ -50,7 +50,8 @@ struct FSearchResult : FDatabasePoseIdx
 
 	// time of the associated GetSearchIndexAsset from this search. During continuing pose searches,
 	// AssetTime is resynchronized via UpdateWithNormalizedTime or UpdateWithRealTime, so it's not precisely the quantized time associated to PoseIdx
-	float AssetTime = 0.f;
+	UE_DEPRECATED(5.7, "AssetTime will be made private. Use GetAssetTime and SetAssetTime instead")
+	float AssetTime = MAX_flt;
 
 	// true if the search result is from a continuing pose search
 	bool bIsContinuingPoseSearch = false;
@@ -61,7 +62,18 @@ struct FSearchResult : FDatabasePoseIdx
 	int32 BestPosePos = 0;
 #endif // WITH_EDITOR && ENABLE_ANIM_DEBUG && UE_POSE_SEARCH_TRACE_ENABLED
 
+	PRAGMA_DISABLE_DEPRECATION_WARNINGS
 	FSearchResult() = default;
+	FSearchResult(const FSearchResult& Other) = default;
+	FSearchResult(FSearchResult&& Other) = default;
+	FSearchResult& operator=(const FSearchResult& Other) = default;
+	FSearchResult& operator=(FSearchResult&& Other) = default;
+
+	static bool IsAssetTimeValid(float AssetTime) { return AssetTime != MAX_flt; }
+	bool IsAssetTimeValid() const { return IsAssetTimeValid(AssetTime); }
+	float GetAssetTime() const { check(IsAssetTimeValid(AssetTime)); return AssetTime; }
+	void SetAssetTime(float InAssetTime) { check(IsAssetTimeValid(InAssetTime)); AssetTime = InAssetTime; }
+	PRAGMA_ENABLE_DEPRECATION_WARNINGS
 
 	// Experimental, this feature might be removed without warning, not for production use
 	// Attempts to set the internal state to match the provided asset time including updating the internal DbPoseIdx. 
@@ -96,7 +108,9 @@ struct FSearchResult : FDatabasePoseIdx
 		return static_cast<const FDatabasePoseIdx&>(*this) == static_cast<const FDatabasePoseIdx&>(Other) &&
 			PoseCost == Other.PoseCost &&
 			EventPoseIdx == Other.EventPoseIdx &&
+			PRAGMA_DISABLE_DEPRECATION_WARNINGS
 			AssetTime == Other.AssetTime &&
+			PRAGMA_ENABLE_DEPRECATION_WARNINGS
 			bIsContinuingPoseSearch == Other.bIsContinuingPoseSearch;
 	}
 
@@ -104,6 +118,61 @@ struct FSearchResult : FDatabasePoseIdx
 	void InitFrom(const FPoseSearchBlueprintResult& BlueprintResult);
 };
 
+// Experimental, this feature might be removed without warning, not for production use
+struct ISearchResults
+{
+	virtual ~ISearchResults() {}
+	virtual void UpdateWith(FSearchResult& SearchResult) = 0;
+	virtual void FinalizeResults() = 0;
+	virtual FSearchResult GetBestResult() const = 0;
+
+	// every search will ask FSearchResults if it should perform the search, knowing that its minimum cost will be SearchMinimumCost:
+	// it's up to this method to perform early out optimizations if the holded results are already good enough fo skip the search
+	virtual bool ShouldPerformSearch(float SearchMinimumCost) const = 0;
+};
+
+// Experimental, this feature might be removed without warning, not for production use
+struct FSearchResults_Single : public ISearchResults
+{
+	POSESEARCH_API virtual void UpdateWith(FSearchResult& SearchResult) override;
+	POSESEARCH_API virtual void FinalizeResults() override;
+	POSESEARCH_API virtual FSearchResult GetBestResult() const override;
+	POSESEARCH_API virtual bool ShouldPerformSearch(float SearchMinimumCost) const override;
+
+private:
+	FSearchResult SingleSearchResult;
+};
+
+// Experimental, this feature might be removed without warning, not for production use
+struct FSearchResults_Multi : public ISearchResults
+{
+	POSESEARCH_API explicit FSearchResults_Multi(int32 InMaxNumberOfWantedResults = 1);
+	POSESEARCH_API virtual void UpdateWith(FSearchResult& SearchResult) override;
+	POSESEARCH_API virtual void FinalizeResults() override;
+	POSESEARCH_API virtual FSearchResult GetBestResult() const override;
+	POSESEARCH_API virtual bool ShouldPerformSearch(float SearchMinimumCost) const override;
+
+private:
+	// @todo: evaluate a max heap instead (using AlgoImpl::HeapSiftUp / AlgoImpl::HeapSiftDown)
+	// SearchResults[0] is the best result (lowest FSearchResult::PoseCost), SearchResults.Last() is the worst result (higher FSearchResult::PoseCost)
+	TArray<FSearchResult, TInlineAllocator<PreallocatedSearchesNum, TMemStackAllocator<>>> SearchResults;
+	int32 MaxNumberOfWantedResults = 1;
+};
+
+// Experimental, this feature might be removed without warning, not for production use
+struct FSearchResults_AssetBests : public ISearchResults
+{
+	POSESEARCH_API virtual void UpdateWith(FSearchResult& SearchResult) override;
+	POSESEARCH_API virtual void FinalizeResults() override;
+	POSESEARCH_API virtual FSearchResult GetBestResult() const override;
+	POSESEARCH_API virtual bool ShouldPerformSearch(float SearchMinimumCost) const override;
+
+	typedef TMap<int32, FSearchResult, TInlineSetAllocator<16, TMemStackSetAllocator<>>> FMap;
+	const FMap& GetSourceAssetIdxToBestSearchResult() const { return SourceAssetIdxToBestSearchResult; }
+private:
+	FMap SourceAssetIdxToBestSearchResult;
+};
+
 } // namespace UE::PoseSearch
 
 USTRUCT(BlueprintType, Category="Animation|Pose Search")
diff --git a/Engine/Plugins/Experimental/MovieScenePoseSearchTracks/Source/MovieScenePoseSearchTracks/Private/Sections/MovieSceneStitchAnimSection.cpp b/Engine/Plugins/Experimental/MovieScenePoseSearchTracks/Source/MovieScenePoseSearchTracks/Private/Sections/MovieSceneStitchAnimSection.cpp
--- a/Engine/Plugins/Experimental/MovieScenePoseSearchTracks/Source/MovieScenePoseSearchTracks/Private/Sections/MovieSceneStitchAnimSection.cpp
+++ b/Engine/Plugins/Experimental/MovieScenePoseSearchTracks/Source/MovieScenePoseSearchTracks/Private/Sections/MovieSceneStitchAnimSection.cpp
@@ -118,16 +118,18 @@ void FMovieSceneStitchAnimEvaluationTask::Execute(UE::UAF::FEvaluationVM& VM) co
 
 			const UObject* StitchDatabaseObj = StitchData.StitchDatabase.Get();
 			const UObject* AnimContext = ContextObject.Get();
-			const FSearchResult SearchResult = UPoseSearchLibrary::MotionMatch(MakeArrayView(&AnimContext, 1),
-				MakeArrayView(&DefaultRole, 1), MakeArrayView(&PoseHistory, 1), MakeArrayView(&StitchDatabaseObj, 1), ContinuingProperties, TimeToTarget, FPoseSearchEvent());
-
+			FSearchResults_Single SearchResults;
+			UPoseSearchLibrary::MotionMatch(MakeArrayView(&AnimContext, 1), MakeArrayView(&DefaultRole, 1), MakeArrayView(&PoseHistory, 1),
+				MakeArrayView(&StitchDatabaseObj, 1), ContinuingProperties, TimeToTarget, FPoseSearchEvent(), SearchResults);
+				
+			const FSearchResult SearchResult = SearchResults.GetBestResult();
 			if (const FSearchIndexAsset* SearchIndexAsset = SearchResult.GetSearchIndexAsset())
 			{
 				const UPoseSearchDatabase* Database = SearchResult.Database.Get();
 				if (const FPoseSearchDatabaseAnimationAssetBase* DatabaseAnimationAssetBase = Database->GetDatabaseAnimationAsset(*SearchIndexAsset))
 				{
 					MatchedAsset = Cast<UAnimationAsset>(DatabaseAnimationAssetBase->GetAnimationAsset());
-					MatchedAssetTime = SearchResult.AssetTime;
+					MatchedAssetTime = SearchResult.GetAssetTime();
 
 					// TODO: Add mirroring support
 					//ProviderResult.bMirrored = SearchIndexAsset->IsMirrored();
diff --git a/Engine/Plugins/Experimental/UAF/UAFPoseSearch/Source/UAFPoseSearch/Private/MotionMatchingTrait.cpp b/Engine/Plugins/Experimental/UAF/UAFPoseSearch/Source/UAFPoseSearch/Private/MotionMatchingTrait.cpp
--- a/Engine/Plugins/Experimental/UAF/UAFPoseSearch/Source/UAFPoseSearch/Private/MotionMatchingTrait.cpp
+++ b/Engine/Plugins/Experimental/UAF/UAFPoseSearch/Source/UAFPoseSearch/Private/MotionMatchingTrait.cpp
@@ -297,7 +297,9 @@ namespace UE::UAF
 
 					FPoseSearchContinuingProperties ContinuingProperties;
 					ContinuingProperties.InitFrom(MotionMatchingState.SearchResult, SharedData->GetNextUpdateInterruptMode(Binding));
-					SearchResult = UPoseSearchLibrary::MotionMatch(SearchContext, AssetsToSearch, ContinuingProperties);
+					FSearchResults_Single SearchResults;
+					UPoseSearchLibrary::MotionMatch(SearchContext, AssetsToSearch, ContinuingProperties, SearchResults);
+					SearchResult = SearchResults.GetBestResult();
 				}
 				else
 				{
