PoseSearch - support for multiple search results (from ISearchResults) for UPoseSearchLibrary::TraceMotionMatching

[RN] support for multiple search results (from ISearchResults) for UPoseSearchLibrary::TraceMotionMatching

diff --git a/Engine/Plugins/Animation/PoseSearch/Source/Editor/Private/PoseSearchDebuggerDatabaseView.h b/Engine/Plugins/Animation/PoseSearch/Source/Editor/Private/PoseSearchDebuggerDatabaseView.h
--- a/Engine/Plugins/Animation/PoseSearch/Source/Editor/Private/PoseSearchDebuggerDatabaseView.h
+++ b/Engine/Plugins/Animation/PoseSearch/Source/Editor/Private/PoseSearchDebuggerDatabaseView.h
@@ -115,7 +115,6 @@ private:
 		/** Widget for displaying the list of row objects */
 		TSharedPtr<SListView<TSharedRef<FDebuggerDatabaseRowData>>> ListView;
 
-		// @TODO: Explore options for active row other than displaying array of 1 element
 		/** List of row objects */
 		TArray<TSharedRef<FDebuggerDatabaseRowData>> Rows;
 
diff --git a/Engine/Plugins/Animation/PoseSearch/Source/Editor/Private/PoseSearchDebuggerViewModel.cpp b/Engine/Plugins/Animation/PoseSearch/Source/Editor/Private/PoseSearchDebuggerViewModel.cpp
--- a/Engine/Plugins/Animation/PoseSearch/Source/Editor/Private/PoseSearchDebuggerViewModel.cpp
+++ b/Engine/Plugins/Animation/PoseSearch/Source/Editor/Private/PoseSearchDebuggerViewModel.cpp
@@ -47,7 +47,9 @@ const UPoseSearchDatabase* FDebuggerViewModel::GetCurrentDatabase() const
 {
 	if (const FTraceMotionMatchingStateMessage* State = GetMotionMatchingState())
 	{
+		PRAGMA_DISABLE_DEPRECATION_WARNINGS
 		return State->GetCurrentDatabase();
+		PRAGMA_ENABLE_DEPRECATION_WARNINGS
 	}
 	return nullptr;
 }
diff --git a/Engine/Plugins/Animation/PoseSearch/Source/Runtime/Private/Chooser/PoseSearchChooserColumn.cpp b/Engine/Plugins/Animation/PoseSearch/Source/Runtime/Private/Chooser/PoseSearchChooserColumn.cpp
--- a/Engine/Plugins/Animation/PoseSearch/Source/Runtime/Private/Chooser/PoseSearchChooserColumn.cpp
+++ b/Engine/Plugins/Animation/PoseSearch/Source/Runtime/Private/Chooser/PoseSearchChooserColumn.cpp
@@ -271,9 +271,16 @@ void FPoseSearchColumn::Filter(FChooserEvaluationContext& Context, const FChoose
 		TArray<const UObject*, TInlineAllocator<128>> AssetsToConsider;
 		for (const FChooserIndexArray::FIndexData& IndexData : IndexListIn)
 		{
-			if (const FPoseSearchDatabaseAnimationAsset* Asset = InternalDatabase->GetDatabaseAnimationAsset(IndexData.Index))
+			const FPoseSearchDatabaseAnimationAsset* Asset = InternalDatabase->GetDatabaseAnimationAsset(IndexData.Index);
+			if (ensure(Asset))
 			{
-				AssetsToConsider.Add(Asset->GetAnimationAsset());
+				const UObject* AnimationAsset = Asset->GetAnimationAsset();
+#if DO_CHECK && WITH_EDITOR
+				// making sure InternalDatabase and the OuterChooser are still in synch
+				const UObject* ReferencedObject = GetReferencedObject(IndexData.Index);
+				check(ReferencedObject == AnimationAsset);
+#endif // DO_CHECK && WITH_EDITOR
+				AssetsToConsider.Add(AnimationAsset);
 			}
 		}
 
diff --git a/Engine/Plugins/Animation/PoseSearch/Source/Runtime/Private/PoseSearchLibrary.cpp b/Engine/Plugins/Animation/PoseSearch/Source/Runtime/Private/PoseSearchLibrary.cpp
--- a/Engine/Plugins/Animation/PoseSearch/Source/Runtime/Private/PoseSearchLibrary.cpp
+++ b/Engine/Plugins/Animation/PoseSearch/Source/Runtime/Private/PoseSearchLibrary.cpp
@@ -559,14 +559,31 @@ void FPoseSearchContinuingProperties::InitFrom(const FPoseSearchBlueprintResult&
 // UPoseSearchLibrary
 
 #if UE_POSE_SEARCH_TRACE_ENABLED
+
+void UPoseSearchLibrary::TraceMotionMatching(
+	UE::PoseSearch::FSearchContext& SearchContext,
+	const UE::PoseSearch::FSearchResult& SearchResult,
+	float ElapsedPoseSearchTime,
+	float DeltaTime,
+	bool bSearch,
+	float WantedPlayRate,
+	EPoseSearchInterruptMode InterruptMode)
+{
+	using namespace UE::PoseSearch;
+	FSearchResults_Single SearchResults;
+	if (bSearch && SearchResult.IsValid())
+	{
+		SearchResults.UpdateWith(SearchResult);
+	}
+	TraceMotionMatching(SearchContext, SearchResults, ElapsedPoseSearchTime, WantedPlayRate, InterruptMode);
+}
+
 void UPoseSearchLibrary::TraceMotionMatching(
-		UE::PoseSearch::FSearchContext& SearchContext,
-		const UE::PoseSearch::FSearchResult& SearchResult,
-		float ElapsedPoseSearchTime,
-		float DeltaTime,
-		bool bSearch,
-		float WantedPlayRate,
-		EPoseSearchInterruptMode InterruptMode)
+	UE::PoseSearch::FSearchContext& SearchContext,
+	const UE::PoseSearch::ISearchResults& SearchResults,
+	float ElapsedPoseSearchTime,
+	float WantedPlayRate,
+	EPoseSearchInterruptMode InterruptMode)
 {
 	using namespace UE::PoseSearch;
 
@@ -649,7 +666,7 @@ void UPoseSearchLibrary::TraceMotionMatching(
 
 	TArray<uint64, TInlineAllocator<64>> DatabaseIds;
 	int32 DbEntryIdx = 0;
-	const int32 CurrentPoseIdx = bSearch && SearchResult.PoseCost.IsValid() ? SearchResult.PoseIdx : INDEX_NONE;
+
 	TraceState.DatabaseEntries.SetNum(SearchContext.GetBestPoseCandidatesMap().Num());
 	for (TPair<const UPoseSearchDatabase*, FSearchContext::FBestPoseCandidates> DatabaseBestPoseCandidates : SearchContext.GetBestPoseCandidatesMap())
 	{
@@ -672,14 +689,21 @@ void UPoseSearchLibrary::TraceMotionMatching(
 			PoseEntry.DbPoseIdx = PoseCandidate.PoseIdx;
 			PoseEntry.Cost = PoseCandidate.Cost;
 			PoseEntry.PoseCandidateFlags = PoseCandidate.PoseCandidateFlags;
-			if (CurrentPoseIdx == PoseCandidate.PoseIdx && SearchResult.Database.Get() == Database)
+			
+			// looking for PoseCandidate/Database in SearchResults
+			if (SearchResults.IterateOverSearchResults([&PoseCandidate, Database](const FSearchResult& SearchResult)
+				{
+					return PoseCandidate.PoseIdx == SearchResult.PoseIdx && Database == SearchResult.Database.Get();
+				}))
 			{
+				// if found PoseCandidate is a Valid_CurrentPose
 				check(EnumHasAnyFlags(PoseEntry.PoseCandidateFlags, EPoseCandidateFlags::Valid_Pose | EPoseCandidateFlags::Valid_ContinuingPose));
-
 				EnumAddFlags(PoseEntry.PoseCandidateFlags, EPoseCandidateFlags::Valid_CurrentPose);
 
+				PRAGMA_DISABLE_DEPRECATION_WARNINGS
 				TraceState.CurrentDbEntryIdx = DbEntryIdx;
 				TraceState.CurrentPoseEntryIdx = DbEntry.PoseEntries.Add(PoseEntry);
+				PRAGMA_ENABLE_DEPRECATION_WARNINGS
 			}
 			else
 			{
@@ -693,6 +717,21 @@ void UPoseSearchLibrary::TraceMotionMatching(
 	DatabaseIds.Sort();
 	SearchId = HashCombineFast(SearchId, GetTypeHash(DatabaseIds));
 
+	// @todo: instead of using the SearchResult (the best search result with the lowest cost) to calculate velocities etc, 
+	//        shouldn't we aggregate the values from all the SearchResults (using SearchResults.IterateOverSearchResults)?
+	const FSearchResult SearchResult = SearchResults.GetBestResult();
+	const float AssetTime = SearchResult.IsValid() ? SearchResult.GetAssetTime() : 0.f;
+
+	// @todo: integrate DeltaTime into SearchContext, and implement it for UAF as well
+	float DeltaTime = FiniteDelta;
+	if (SearchContext.GetContexts()[0]->ObjectParams.Num() > 0)
+	{
+		if (const UAnimInstance* AnimInstance = Cast<UAnimInstance>(SearchContext.GetContexts()[0]->GetFirstObjectParam()))
+		{
+			DeltaTime = AnimInstance->GetDeltaSeconds();
+		}
+	}
+
 	if (DeltaTime > SMALL_NUMBER)
 	{
 		// simulation
@@ -714,7 +753,7 @@ void UPoseSearchLibrary::TraceMotionMatching(
 				TraceState.SimAngularVelocity += FMath::RadiansToDegrees(SimDelta.GetRotation().GetAngle()) / (DeltaTime * NumRoles);
 			}
 		}
-
+		
 		const FSearchIndexAsset* SearchIndexAsset = SearchResult.GetSearchIndexAsset();
 		const UPoseSearchDatabase* CurrentResultDatabase = SearchResult.Database.Get();
 		if (SearchIndexAsset && CurrentResultDatabase)
@@ -727,14 +766,14 @@ void UPoseSearchLibrary::TraceMotionMatching(
 				FAnimationAssetSampler Sampler(AnimationAsset, FTransform::Identity,FVector::ZeroVector, FAnimationAssetSampler::DefaultRootTransformSamplingRate, true, false);
 
 				const float TimeStep = DeltaTime * WantedPlayRate;
-				const FTransform PrevRoot = Sampler.ExtractRootTransform(SearchResult.GetAssetTime());
-				const FTransform CurrRoot = Sampler.ExtractRootTransform(SearchResult.GetAssetTime() + TimeStep);
+				const FTransform PrevRoot = Sampler.ExtractRootTransform(AssetTime);
+				const FTransform CurrRoot = Sampler.ExtractRootTransform(AssetTime + TimeStep);
 				const FTransform RootMotionTransformDelta = PrevRoot.GetRelativeTransform(CurrRoot);
 				TraceState.AnimLinearVelocity = RootMotionTransformDelta.GetTranslation().Size() / DeltaTime;
 				TraceState.AnimAngularVelocity = FMath::RadiansToDegrees(RootMotionTransformDelta.GetRotation().GetAngle()) / DeltaTime;
 
 				// Need another root motion extraction for non-playrate version in case acceleration isn't the same.
-				const FTransform CurrRootNoTimescale = Sampler.ExtractRootTransform(SearchResult.GetAssetTime() + DeltaTime);
+				const FTransform CurrRootNoTimescale = Sampler.ExtractRootTransform(AssetTime + DeltaTime);
 				const FTransform RootMotionTransformDeltaNoTimescale = PrevRoot.GetRelativeTransform(CurrRootNoTimescale);
 				TraceState.AnimLinearVelocityNoTimescale = RootMotionTransformDeltaNoTimescale.GetTranslation().Size() / DeltaTime;
 				TraceState.AnimAngularVelocityNoTimescale = FMath::RadiansToDegrees(RootMotionTransformDeltaNoTimescale.GetRotation().GetAngle()) / DeltaTime;
@@ -744,7 +783,7 @@ void UPoseSearchLibrary::TraceMotionMatching(
 	}
 
 	TraceState.ElapsedPoseSearchTime = ElapsedPoseSearchTime;
-	TraceState.AssetPlayerTime = SearchResult.GetAssetTime();
+	TraceState.AssetPlayerTime = AssetTime;;
 	TraceState.DeltaTime = DeltaTime;
 
 	TraceState.RecordingTime = RecordingTime;
@@ -1388,18 +1427,7 @@ void UPoseSearchLibrary::MotionMatch(
 #endif // ENABLE_DRAW_DEBUG && ENABLE_ANIM_DEBUG
 
 #if UE_POSE_SEARCH_TRACE_ENABLED
-
-	// collecting all the possible continuing pose search (it could be multiple searches, but most likely only one)
-	float DeltaSeconds = FiniteDelta;
-	if (SearchContext.GetContexts()[0]->ObjectParams.Num() > 0)
-	{
-		if (const UAnimInstance* AnimInstance = Cast<UAnimInstance>(SearchContext.GetContexts()[0]->GetFirstObjectParam()))
-		{
-			DeltaSeconds = AnimInstance->GetDeltaSeconds();
-		}
-	}
-
-	TraceMotionMatching(SearchContext, SearchResult, 0.f, DeltaSeconds, true, 1.f, ContinuingProperties.InterruptMode);
+	TraceMotionMatching(SearchContext, SearchResults, 0.f, 1.f, ContinuingProperties.InterruptMode);
 #endif // UE_POSE_SEARCH_TRACE_ENABLED
 
 	// restoring the asset to consider array view
diff --git a/Engine/Plugins/Animation/PoseSearch/Source/Runtime/Private/PoseSearchResult.cpp b/Engine/Plugins/Animation/PoseSearch/Source/Runtime/Private/PoseSearchResult.cpp
--- a/Engine/Plugins/Animation/PoseSearch/Source/Runtime/Private/PoseSearchResult.cpp
+++ b/Engine/Plugins/Animation/PoseSearch/Source/Runtime/Private/PoseSearchResult.cpp
@@ -225,7 +225,7 @@ void FSearchResult::InitFrom(const FPoseSearchBlueprintResult& BlueprintResult)
 
 /////////////////////////////////////////////////////
 // FSearchResults_Single
-void FSearchResults_Single::UpdateWith(FSearchResult& SearchResult)
+void FSearchResults_Single::UpdateWith(const FSearchResult& SearchResult)
 {
 	check(SearchResult.IsValid());
 	if (SearchResult.PoseCost < SingleSearchResult.PoseCost)
@@ -253,6 +253,11 @@ bool FSearchResults_Single::ShouldPerformSearch(float SearchMinimumCost) const
 	return SingleSearchResult.PoseCost > SearchMinimumCost;
 }
 
+bool FSearchResults_Single::IterateOverSearchResults(const TFunctionRef<bool(const FSearchResult& SearchResult)> IterateOverSearchResultsFunction) const
+{
+	return IterateOverSearchResultsFunction(SingleSearchResult);
+}
+
 /////////////////////////////////////////////////////
 // FSearchResults_Multi
 FSearchResults_Multi::FSearchResults_Multi(int32 InMaxNumberOfWantedResults)
@@ -262,7 +267,7 @@ FSearchResults_Multi::FSearchResults_Multi(int32 InMaxNumberOfWantedResults)
 	SearchResults.Reserve(MaxNumberOfWantedResults);
 }
 
-void FSearchResults_Multi::UpdateWith(FSearchResult& SearchResult)
+void FSearchResults_Multi::UpdateWith(const FSearchResult& SearchResult)
 {
 	check(SearchResult.IsValid());
 	if (SearchResults.Num() < MaxNumberOfWantedResults)
@@ -323,9 +328,21 @@ bool FSearchResults_Multi::ShouldPerformSearch(float SearchMinimumCost) const
 	return SearchResults.Last().PoseCost > SearchMinimumCost;
 }
 
+bool FSearchResults_Multi::IterateOverSearchResults(const TFunctionRef<bool(const FSearchResult& SearchResult)> IterateOverSearchResultsFunction) const
+{
+	for (const FSearchResult& SearchResult : SearchResults)
+	{
+		if (IterateOverSearchResultsFunction(SearchResult))
+		{
+			return true;
+		}
+	}
+	return false;
+}
+
 /////////////////////////////////////////////////////
 // FSearchResults_AssetBests
-void FSearchResults_AssetBests::UpdateWith(FSearchResult& SearchResult)
+void FSearchResults_AssetBests::UpdateWith(const FSearchResult& SearchResult)
 {
 	check(SearchResult.IsValid());
 	if (const FSearchIndexAsset* SearchIndexAsset = SearchResult.GetSearchIndexAsset())
@@ -371,6 +388,18 @@ bool FSearchResults_AssetBests::ShouldPerformSearch(float SearchMinimumCost) con
 	return true;
 }
 
+bool FSearchResults_AssetBests::IterateOverSearchResults(const TFunctionRef<bool(const FSearchResult& SearchResult)> IterateOverSearchResultsFunction) const
+{
+	for (const TPair<int32, FSearchResult>& SourceAssetIdxToBestSearchResultPair : SourceAssetIdxToBestSearchResult)
+	{
+		if (IterateOverSearchResultsFunction(SourceAssetIdxToBestSearchResultPair.Value))
+		{
+			return true;
+		}
+	}
+	return false;
+}
+
 } // namespace UE::PoseSearch
 
 bool FPoseSearchBlueprintResult::InitFrom(const UE::PoseSearch::FSearchResult& SearchResult, float InWantedPlayRate)
diff --git a/Engine/Plugins/Animation/PoseSearch/Source/Runtime/Private/Trace/PoseSearchTraceLogger.cpp b/Engine/Plugins/Animation/PoseSearch/Source/Runtime/Private/Trace/PoseSearchTraceLogger.cpp
--- a/Engine/Plugins/Animation/PoseSearch/Source/Runtime/Private/Trace/PoseSearchTraceLogger.cpp
+++ b/Engine/Plugins/Animation/PoseSearch/Source/Runtime/Private/Trace/PoseSearchTraceLogger.cpp
@@ -80,8 +80,10 @@ FArchive& operator<<(FArchive& Ar, FTraceMotionMatchingStateMessage& State)
 	Ar << State.Roles;
 	Ar << State.DatabaseEntries;
 	Ar << State.PoseHistories;	
+	PRAGMA_DISABLE_DEPRECATION_WARNINGS
 	Ar << State.CurrentDbEntryIdx;
 	Ar << State.CurrentPoseEntryIdx;
+	PRAGMA_ENABLE_DEPRECATION_WARNINGS
 
 	// InterruptMode has been introduced with AddedInterruptModeToDebugger version
 	if (Ar.CustomVer(FPoseSearchCustomVersion::GUID) >= FPoseSearchCustomVersion::AddedInterruptModeToDebugger)
@@ -110,24 +112,29 @@ void FTraceMotionMatchingStateMessage::Output()
 const UPoseSearchDatabase* FTraceMotionMatchingStateMessage::GetCurrentDatabase() const
 {
 	const UPoseSearchDatabase* Database = nullptr;
+	PRAGMA_DISABLE_DEPRECATION_WARNINGS
 	if (DatabaseEntries.IsValidIndex(CurrentDbEntryIdx))
 	{
 		Database = GetObjectFromId<UPoseSearchDatabase>(DatabaseEntries[CurrentDbEntryIdx].DatabaseId);
 	}
+	PRAGMA_ENABLE_DEPRECATION_WARNINGS
 	return Database;
 }
 
 int32 FTraceMotionMatchingStateMessage::GetCurrentDatabasePoseIndex() const
 {
+	PRAGMA_DISABLE_DEPRECATION_WARNINGS
 	if (const FTraceMotionMatchingStatePoseEntry* PoseEntry = GetCurrentPoseEntry())
 	{
 		return PoseEntry->DbPoseIdx;
 	}
+	PRAGMA_ENABLE_DEPRECATION_WARNINGS
 	return INDEX_NONE;
 }
 
 const FTraceMotionMatchingStatePoseEntry* FTraceMotionMatchingStateMessage::GetCurrentPoseEntry() const
 {
+	PRAGMA_DISABLE_DEPRECATION_WARNINGS
 	if (DatabaseEntries.IsValidIndex(CurrentDbEntryIdx))
 	{
 		const FTraceMotionMatchingStateDatabaseEntry& DbEntry = DatabaseEntries[CurrentDbEntryIdx];
@@ -136,6 +143,7 @@ const FTraceMotionMatchingStatePoseEntry* FTraceMotionMatchingStateMessage::GetC
 			return &DbEntry.PoseEntries[CurrentPoseEntryIdx];
 		}
 	}
+	PRAGMA_ENABLE_DEPRECATION_WARNINGS
 	return nullptr;
 }
 
diff --git a/Engine/Plugins/Animation/PoseSearch/Source/Runtime/Public/PoseSearch/PoseSearchLibrary.h b/Engine/Plugins/Animation/PoseSearch/Source/Runtime/Public/PoseSearch/PoseSearchLibrary.h
--- a/Engine/Plugins/Animation/PoseSearch/Source/Runtime/Public/PoseSearch/PoseSearchLibrary.h
+++ b/Engine/Plugins/Animation/PoseSearch/Source/Runtime/Public/PoseSearch/PoseSearchLibrary.h
@@ -161,6 +161,7 @@ class UPoseSearchLibrary : public UBlueprintFunctionLibrary
 
 #if UE_POSE_SEARCH_TRACE_ENABLED
 
+	// deprecate in favor of TraceMotionMatching using ISearchResults once we settle on ISearchResults API signatures
 	static void TraceMotionMatching(
 		UE::PoseSearch::FSearchContext& SearchContext,
 		const UE::PoseSearch::FSearchResult& SearchResult,
@@ -170,6 +171,14 @@ class UPoseSearchLibrary : public UBlueprintFunctionLibrary
 		float WantedPlayRate,
 		EPoseSearchInterruptMode InterruptMode);
 
+	// Experimental, this feature might be removed without warning, not for production use
+	static void TraceMotionMatching(
+		UE::PoseSearch::FSearchContext& SearchContext,
+		const UE::PoseSearch::ISearchResults& SearchResults,
+		float ElapsedPoseSearchTime,
+		float WantedPlayRate,
+		EPoseSearchInterruptMode InterruptMode);
+
 	UE_DEPRECATED(5.6, "Use TraceMotionMatching with different signature instead.")
 	static void TraceMotionMatching(
 		UE::PoseSearch::FSearchContext& SearchContext,
diff --git a/Engine/Plugins/Animation/PoseSearch/Source/Runtime/Public/PoseSearch/PoseSearchResult.h b/Engine/Plugins/Animation/PoseSearch/Source/Runtime/Public/PoseSearch/PoseSearchResult.h
--- a/Engine/Plugins/Animation/PoseSearch/Source/Runtime/Public/PoseSearch/PoseSearchResult.h
+++ b/Engine/Plugins/Animation/PoseSearch/Source/Runtime/Public/PoseSearch/PoseSearchResult.h
@@ -122,22 +122,27 @@ struct FSearchResult : FDatabasePoseIdx
 struct ISearchResults
 {
 	virtual ~ISearchResults() {}
-	virtual void UpdateWith(FSearchResult& SearchResult) = 0;
+	virtual void UpdateWith(const FSearchResult& SearchResult) = 0;
 	virtual void FinalizeResults() = 0;
 	virtual FSearchResult GetBestResult() const = 0;
 
 	// every search will ask FSearchResults if it should perform the search, knowing that its minimum cost will be SearchMinimumCost:
 	// it's up to this method to perform early out optimizations if the holded results are already good enough fo skip the search
 	virtual bool ShouldPerformSearch(float SearchMinimumCost) const = 0;
+
+	// calls IterateOverSearchResultsFunction for each FSearchResult this ISearchResults is holding
+	// is IterateOverSearchResultsFunction returns true, the iteration will stop and IterateOverSearchResults will return true otherwise will return false (useful to implement contains / find methods)
+	virtual bool IterateOverSearchResults(const TFunctionRef<bool(const FSearchResult& SearchResult)> IterateOverSearchResultsFunction) const = 0;
 };
 
 // Experimental, this feature might be removed without warning, not for production use
 struct FSearchResults_Single : public ISearchResults
 {
-	POSESEARCH_API virtual void UpdateWith(FSearchResult& SearchResult) override;
+	POSESEARCH_API virtual void UpdateWith(const FSearchResult& SearchResult) override;
 	POSESEARCH_API virtual void FinalizeResults() override;
 	POSESEARCH_API virtual FSearchResult GetBestResult() const override;
 	POSESEARCH_API virtual bool ShouldPerformSearch(float SearchMinimumCost) const override;
+	POSESEARCH_API virtual bool IterateOverSearchResults(const TFunctionRef<bool(const FSearchResult& SearchResult)> IterateOverSearchResultsFunction) const override;
 
 private:
 	FSearchResult SingleSearchResult;
@@ -147,10 +152,11 @@ private:
 struct FSearchResults_Multi : public ISearchResults
 {
 	POSESEARCH_API explicit FSearchResults_Multi(int32 InMaxNumberOfWantedResults = 1);
-	POSESEARCH_API virtual void UpdateWith(FSearchResult& SearchResult) override;
+	POSESEARCH_API virtual void UpdateWith(const FSearchResult& SearchResult) override;
 	POSESEARCH_API virtual void FinalizeResults() override;
 	POSESEARCH_API virtual FSearchResult GetBestResult() const override;
 	POSESEARCH_API virtual bool ShouldPerformSearch(float SearchMinimumCost) const override;
+	POSESEARCH_API virtual bool IterateOverSearchResults(const TFunctionRef<bool(const FSearchResult& SearchResult)> IterateOverSearchResultsFunction) const override;
 
 private:
 	// @todo: evaluate a max heap instead (using AlgoImpl::HeapSiftUp / AlgoImpl::HeapSiftDown)
@@ -162,10 +168,11 @@ private:
 // Experimental, this feature might be removed without warning, not for production use
 struct FSearchResults_AssetBests : public ISearchResults
 {
-	POSESEARCH_API virtual void UpdateWith(FSearchResult& SearchResult) override;
+	POSESEARCH_API virtual void UpdateWith(const FSearchResult& SearchResult) override;
 	POSESEARCH_API virtual void FinalizeResults() override;
 	POSESEARCH_API virtual FSearchResult GetBestResult() const override;
 	POSESEARCH_API virtual bool ShouldPerformSearch(float SearchMinimumCost) const override;
+	POSESEARCH_API virtual bool IterateOverSearchResults(const TFunctionRef<bool(const FSearchResult& SearchResult)> IterateOverSearchResultsFunction) const override;
 
 	typedef TMap<int32, FSearchResult, TInlineSetAllocator<16, TMemStackSetAllocator<>>> FMap;
 	const FMap& GetSourceAssetIdxToBestSearchResult() const { return SourceAssetIdxToBestSearchResult; }
diff --git a/Engine/Plugins/Animation/PoseSearch/Source/Runtime/Public/PoseSearch/Trace/PoseSearchTraceLogger.h b/Engine/Plugins/Animation/PoseSearch/Source/Runtime/Public/PoseSearch/Trace/PoseSearchTraceLogger.h
--- a/Engine/Plugins/Animation/PoseSearch/Source/Runtime/Public/PoseSearch/Trace/PoseSearchTraceLogger.h
+++ b/Engine/Plugins/Animation/PoseSearch/Source/Runtime/Public/PoseSearch/Trace/PoseSearchTraceLogger.h
@@ -63,6 +63,14 @@ POSESEARCH_API FArchive& operator<<(FArchive& Ar, FTraceMotionMatchingStateDatab
  */
 struct FTraceMotionMatchingStateMessage : public FTraceMessage
 {
+	PRAGMA_DISABLE_DEPRECATION_WARNINGS
+	FTraceMotionMatchingStateMessage() = default;
+	FTraceMotionMatchingStateMessage(const FTraceMotionMatchingStateMessage& Other) = default;
+	FTraceMotionMatchingStateMessage(FTraceMotionMatchingStateMessage&& Other) = default;
+	FTraceMotionMatchingStateMessage& operator=(const FTraceMotionMatchingStateMessage& Other) = default;
+	FTraceMotionMatchingStateMessage& operator=(FTraceMotionMatchingStateMessage&& Other) = default;
+	PRAGMA_ENABLE_DEPRECATION_WARNINGS
+
 	/** Amount of time since the last pose switch */
 	float ElapsedPoseSearchTime = 0.f;
 
@@ -89,10 +97,10 @@ struct FTraceMotionMatchingStateMessage : public FTraceMessage
 
 	TArray<UE::PoseSearch::FArchivedPoseHistory> PoseHistories;
 
-	/** Index of the current database in DatabaseEntries */
+	UE_DEPRECATED(5.7, "This API will soon be removed. Look for a FTraceMotionMatchingStateDatabaseEntry::PoseEntries with PoseCandidateFlags containg EPoseCandidateFlags::Valid_ContinuingPose for the 'current database and pose index' instead")
 	int32 CurrentDbEntryIdx = INDEX_NONE;
 
-	/** Index of the current pose in DatabaseEntries[CurrentDbEntryIdx].PoseEntries */
+	UE_DEPRECATED(5.7, "This API will soon be removed. Look for a FTraceMotionMatchingStateDatabaseEntry::PoseEntries with PoseCandidateFlags containg EPoseCandidateFlags::Valid_ContinuingPose for the 'current database and pose index' instead")
 	int32 CurrentPoseEntryIdx = INDEX_NONE;
 
 	EPoseSearchInterruptMode InterruptMode = EPoseSearchInterruptMode::DoNotInterrupt;
@@ -100,8 +108,11 @@ struct FTraceMotionMatchingStateMessage : public FTraceMessage
 	/** Output the current state info to the logger */
 	UE_API void Output();
 
+	UE_DEPRECATED(5.7, "This API will soon be removed. Look for a FTraceMotionMatchingStateDatabaseEntry::PoseEntries with PoseCandidateFlags containg EPoseCandidateFlags::Valid_ContinuingPose for the 'current database and pose index' instead")
 	UE_API const UPoseSearchDatabase* GetCurrentDatabase() const;
+	UE_DEPRECATED(5.7, "This API will soon be removed. Look for a FTraceMotionMatchingStateDatabaseEntry::PoseEntries with PoseCandidateFlags containg EPoseCandidateFlags::Valid_ContinuingPose for the 'current database and pose index' instead")
 	UE_API int32 GetCurrentDatabasePoseIndex() const;
+	UE_DEPRECATED(5.7, "This API will soon be removed. Look for a FTraceMotionMatchingStateDatabaseEntry::PoseEntries with PoseCandidateFlags containg EPoseCandidateFlags::Valid_ContinuingPose for the 'current database and pose index' instead")
 	UE_API const FTraceMotionMatchingStatePoseEntry* GetCurrentPoseEntry() const;
 
 	template<typename T>
