PoseSearch - support for multiple search results (from ISearchResults) reporting in pose search debugger

[RN] support for multiple search results (from ISearchResults) reporting in pose search debugger

diff --git a/Engine/Plugins/Animation/PoseSearch/Source/Editor/Private/PoseSearchDebuggerDatabaseView.cpp b/Engine/Plugins/Animation/PoseSearch/Source/Editor/Private/PoseSearchDebuggerDatabaseView.cpp
--- a/Engine/Plugins/Animation/PoseSearch/Source/Editor/Private/PoseSearchDebuggerDatabaseView.cpp
+++ b/Engine/Plugins/Animation/PoseSearch/Source/Editor/Private/PoseSearchDebuggerDatabaseView.cpp
@@ -46,10 +46,11 @@ public:
 	{
 		// processing all the DatabaseEntries to collect the LabelToChannels
 		TLabelBuilder LabelBuilder;
+		FAsyncPoseSearchDatabasesManagementCache DatabasesManagementCache;
 		for (const FTraceMotionMatchingStateDatabaseEntry& DbEntry : DatabaseEntries)
 		{
 			const UPoseSearchDatabase* Database = FTraceMotionMatchingStateMessage::GetObjectFromId<UPoseSearchDatabase>(DbEntry.DatabaseId);
-			if (EAsyncBuildIndexResult::Success == FAsyncPoseSearchDatabasesManagement::RequestAsyncBuildIndex(Database, ERequestAsyncBuildFlag::ContinueRequest))
+			if (EAsyncBuildIndexResult::Success == DatabasesManagementCache(Database))
 			{
 				for (const TObjectPtr<UPoseSearchFeatureChannel>& ChannelPtr : Database->Schema->GetChannels())
 				{
@@ -254,10 +255,12 @@ void SDebuggerDatabaseView::Update(const FTraceMotionMatchingStateMessage& State
 	bool bAddPCACost = false;
 	TArray<uint32> PoseToPCAValuesVectorIndexes;
 	TAlignedArray<float> DynamicWeightsSqrtBuffer;
+
+	FAsyncPoseSearchDatabasesManagementCache DatabasesManagementCache;
 	for (const FTraceMotionMatchingStateDatabaseEntry& DbEntry : State.DatabaseEntries)
 	{
 		const UPoseSearchDatabase* Database = FTraceMotionMatchingStateMessage::GetObjectFromId<UPoseSearchDatabase>(DbEntry.DatabaseId);
-		if (EAsyncBuildIndexResult::Success == FAsyncPoseSearchDatabasesManagement::RequestAsyncBuildIndex(Database, ERequestAsyncBuildFlag::ContinueRequest))
+		if (EAsyncBuildIndexResult::Success == DatabasesManagementCache(Database))
 		{
 			const FSearchIndex& SearchIndex = Database->GetSearchIndex();
 			SearchIndex.GetPoseToPCAValuesVectorIndexes(PoseToPCAValuesVectorIndexes);
@@ -602,14 +605,57 @@ void SDebuggerDatabaseView::OnUseRegexCheckboxChanged(ECheckBoxState State)
 	PopulateViewRows();
 }
 
-void SDebuggerDatabaseView::OnDatabaseRowSelectionChanged(TSharedPtr<FDebuggerDatabaseRowData> Row, ESelectInfo::Type SelectInfo)
+void SDebuggerDatabaseView::OnViewSelectionChanged(TSharedPtr<FDebuggerDatabaseRowData> Row, SDebuggerDatabaseView::FTable& Table)
 {
+	Table.PreviouslySelectedItems = Table.ListView->GetSelectedItems();
+
 	if (Row.IsValid())
 	{
 		OnPoseSelectionChanged.ExecuteIfBound(Row->SharedData->SourceDatabase.Get(), Row->PoseIdx, Row->AssetTime);
 	}
 }
 
+void SDebuggerDatabaseView::OnActiveViewSelectionChanged(TSharedPtr<FDebuggerDatabaseRowData> Row, ESelectInfo::Type SelectInfo)
+{
+	OnViewSelectionChanged(Row, ActiveView);
+}
+
+void SDebuggerDatabaseView::OnContinuingPoseViewSelectionChanged(TSharedPtr<FDebuggerDatabaseRowData> Row, ESelectInfo::Type SelectInfo)
+{
+	OnViewSelectionChanged(Row, ContinuingPoseView);
+}
+
+void SDebuggerDatabaseView::OnFilteredDatabaseViewSelectionChanged(TSharedPtr<FDebuggerDatabaseRowData> Row, ESelectInfo::Type SelectInfo)
+{
+	OnViewSelectionChanged(Row, FilteredDatabaseView);
+}
+
+void SDebuggerDatabaseView::OnViewMouseButtonClick(const TSharedPtr<FDebuggerDatabaseRowData> Row, SDebuggerDatabaseView::FTable& Table)
+{
+	const TArray<TSharedRef<FDebuggerDatabaseRowData>> CurrentlySelectedItems = Table.ListView->GetSelectedItems();
+
+	if (CurrentlySelectedItems.Num() == 1 && Table.PreviouslySelectedItems.Num() == 1 &&
+		CurrentlySelectedItems[0] == Row && Table.PreviouslySelectedItems[0] == Row)
+	{
+		Table.ListView->ClearSelection();
+	}
+}
+
+void SDebuggerDatabaseView::OnActiveViewMouseButtonClick(const TSharedRef<FDebuggerDatabaseRowData> Row)
+{
+	OnViewMouseButtonClick(Row, ActiveView);
+}
+
+void SDebuggerDatabaseView::OnContinuingPoseViewMouseButtonClick(const TSharedRef<FDebuggerDatabaseRowData> Row)
+{
+	OnViewMouseButtonClick(Row, ContinuingPoseView);
+}
+
+void SDebuggerDatabaseView::OnFilteredDatabaseViewMouseButtonClick(const TSharedRef<FDebuggerDatabaseRowData> Row)
+{
+	OnViewMouseButtonClick(Row, FilteredDatabaseView);
+}
+
 void SDebuggerDatabaseView::SortDatabaseRows()
 {
 	if (!UnfilteredDatabaseRows.IsEmpty())
@@ -828,8 +874,8 @@ void SDebuggerDatabaseView::Construct(const FArguments& InArgs)
 		.HeaderRow(ActiveView.HeaderRow.ToSharedRef())
 		.OnGenerateRow(this, &SDebuggerDatabaseView::HandleGenerateActiveRow)
 		.ExternalScrollbar(ActiveView.ScrollBar)
-		.SelectionMode(ESelectionMode::SingleToggle)
-		.ConsumeMouseWheel(EConsumeMouseWheel::Never);
+		.OnSelectionChanged(this, &SDebuggerDatabaseView::OnActiveViewSelectionChanged)
+		.OnMouseButtonClick(this, &SDebuggerDatabaseView::OnActiveViewMouseButtonClick);
 
 	ActiveView.RowStyle = FAppStyle::GetWidgetStyle<FTableRowStyle>("TableView.Row");
 	ActiveView.RowBrush = *FAppStyle::GetBrush("DetailsView.CategoryTop");
@@ -849,8 +895,8 @@ void SDebuggerDatabaseView::Construct(const FArguments& InArgs)
 		.HeaderRow(ContinuingPoseView.HeaderRow.ToSharedRef())
 		.OnGenerateRow(this, &SDebuggerDatabaseView::HandleGenerateContinuingPoseRow)
 		.ExternalScrollbar(ContinuingPoseView.ScrollBar)
-		.SelectionMode(ESelectionMode::SingleToggle)
-		.ConsumeMouseWheel(EConsumeMouseWheel::Never);
+		.OnSelectionChanged(this, &SDebuggerDatabaseView::OnContinuingPoseViewSelectionChanged)
+		.OnMouseButtonClick(this, &SDebuggerDatabaseView::OnContinuingPoseViewMouseButtonClick);
 
 	ContinuingPoseView.RowStyle = FAppStyle::GetWidgetStyle<FTableRowStyle>("TableView.Row");
 	ContinuingPoseView.RowBrush = *FAppStyle::GetBrush("DetailsView.CategoryTop");
@@ -869,9 +915,8 @@ void SDebuggerDatabaseView::Construct(const FArguments& InArgs)
 		.HeaderRow(FilteredDatabaseView.HeaderRow.ToSharedRef())
 		.OnGenerateRow(this, &SDebuggerDatabaseView::HandleGenerateDatabaseRow)
 		.ExternalScrollbar(FilteredDatabaseView.ScrollBar)
-		.SelectionMode(ESelectionMode::Multi)
-		.ConsumeMouseWheel(EConsumeMouseWheel::WhenScrollingPossible)
-		.OnSelectionChanged(this, &SDebuggerDatabaseView::OnDatabaseRowSelectionChanged);
+		.OnSelectionChanged(this, &SDebuggerDatabaseView::OnFilteredDatabaseViewSelectionChanged)
+		.OnMouseButtonClick(this, &SDebuggerDatabaseView::OnFilteredDatabaseViewMouseButtonClick);
 
 	FilteredDatabaseView.RowStyle = FAppStyle::GetWidgetStyle<FTableRowStyle>("TableView.Row");
 	// Set selected color to white to retain visibility when multi-selecting
diff --git a/Engine/Plugins/Animation/PoseSearch/Source/Editor/Private/PoseSearchDebuggerDatabaseView.h b/Engine/Plugins/Animation/PoseSearch/Source/Editor/Private/PoseSearchDebuggerDatabaseView.h
--- a/Engine/Plugins/Animation/PoseSearch/Source/Editor/Private/PoseSearchDebuggerDatabaseView.h
+++ b/Engine/Plugins/Animation/PoseSearch/Source/Editor/Private/PoseSearchDebuggerDatabaseView.h
@@ -39,6 +39,30 @@ class SDebuggerDatabaseView : public SCompoundWidget
 	using FColumnMap = TMap<FName, TSharedRef<DebuggerDatabaseColumns::IColumn>>;
 
 private:
+
+	struct FTable
+	{
+		/** Header row*/
+		TSharedPtr<SHeaderRow> HeaderRow;
+
+		/** Widget for displaying the list of row objects */
+		TSharedPtr<SListView<TSharedRef<FDebuggerDatabaseRowData>>> ListView;
+
+		/** List of row objects */
+		TArray<TSharedRef<FDebuggerDatabaseRowData>> Rows;
+
+		/** Background style for the list view */
+		FTableRowStyle RowStyle;
+
+		/** Row color */
+		FSlateBrush RowBrush;
+
+		/** Scroll bar for the data table */
+		TSharedPtr<SScrollBar> ScrollBar;
+
+		TArray<TSharedRef<FDebuggerDatabaseRowData>> PreviouslySelectedItems;
+	};
+
 	/** Adds a column to the existing list */
 	void AddColumn(TSharedRef<DebuggerDatabaseColumns::IColumn>&& Column);
 
@@ -78,8 +102,15 @@ private:
 	/** Called when the UseRegex Checkbox is Toggled */
 	void OnUseRegexCheckboxChanged(ECheckBoxState State);
 	
-	/** Row selection to update model view */
-	void OnDatabaseRowSelectionChanged(TSharedPtr<FDebuggerDatabaseRowData> Row, ESelectInfo::Type SelectInfo);
+	void OnViewSelectionChanged(TSharedPtr<FDebuggerDatabaseRowData> Row, FTable& Table);
+	void OnActiveViewSelectionChanged(TSharedPtr<FDebuggerDatabaseRowData> Row, ESelectInfo::Type SelectInfo);
+	void OnContinuingPoseViewSelectionChanged(TSharedPtr<FDebuggerDatabaseRowData> Row, ESelectInfo::Type SelectInfo);
+	void OnFilteredDatabaseViewSelectionChanged(TSharedPtr<FDebuggerDatabaseRowData> Row, ESelectInfo::Type SelectInfo);
+
+	void OnViewMouseButtonClick(const TSharedPtr<FDebuggerDatabaseRowData> Row, FTable& Table);
+	void OnActiveViewMouseButtonClick(const TSharedRef<FDebuggerDatabaseRowData> Row);
+	void OnContinuingPoseViewMouseButtonClick(const TSharedRef<FDebuggerDatabaseRowData> Row);
+	void OnFilteredDatabaseViewMouseButtonClick(const TSharedRef<FDebuggerDatabaseRowData> Row);
 
 	/** Generates a database row widget for the given data */
 	TSharedRef<ITableRow> HandleGenerateDatabaseRow(TSharedRef<FDebuggerDatabaseRowData> Item, const TSharedRef<STableViewBase>& OwnerTable) const;
@@ -107,27 +138,6 @@ private:
 	bool bHasPCACostColumn = false;
 	bool bWasVerbose = false;
 
-	struct FTable
-	{
-		/** Header row*/
-		TSharedPtr<SHeaderRow> HeaderRow;
-
-		/** Widget for displaying the list of row objects */
-		TSharedPtr<SListView<TSharedRef<FDebuggerDatabaseRowData>>> ListView;
-
-		/** List of row objects */
-		TArray<TSharedRef<FDebuggerDatabaseRowData>> Rows;
-
-		/** Background style for the list view */
-		FTableRowStyle RowStyle;
-
-		/** Row color */
-		FSlateBrush RowBrush;
-
-		/** Scroll bar for the data table */
-		TSharedPtr<SScrollBar> ScrollBar;
-	};
-
 	/** Active row at the top of the view */
 	FTable ActiveView;
 
diff --git a/Engine/Plugins/Animation/PoseSearch/Source/Editor/Private/PoseSearchDebuggerView.cpp b/Engine/Plugins/Animation/PoseSearch/Source/Editor/Private/PoseSearchDebuggerView.cpp
--- a/Engine/Plugins/Animation/PoseSearch/Source/Editor/Private/PoseSearchDebuggerView.cpp
+++ b/Engine/Plugins/Animation/PoseSearch/Source/Editor/Private/PoseSearchDebuggerView.cpp
@@ -289,9 +289,9 @@ void SDebuggerView::Tick(const FGeometry& AllottedGeometry, const double InCurre
 	if (const FTraceMotionMatchingStateMessage* State = Model->GetMotionMatchingState())
 	{
 		FRoleToIndex RoleToIndex;
-		TArray<FChooserEvaluationContext> AnimContextsData;
-		TArray<FChooserEvaluationContext*> AnimContexts;
-		TArray<const IPoseHistory*> PoseHistories;
+		TArray<FChooserEvaluationContext, TInlineAllocator<PreallocatedRolesNum>> AnimContextsData;
+		TArray<FChooserEvaluationContext*, TInlineAllocator<PreallocatedRolesNum>> AnimContexts;
+		TArray<const IPoseHistory*, TInlineAllocator<PreallocatedRolesNum>> PoseHistories;
 		UWorld* World = nullptr;
 
 		const int32 NumRoles = State->Roles.Num();
@@ -355,77 +355,65 @@ void SDebuggerView::Tick(const FGeometry& AllottedGeometry, const double InCurre
 			}
 		}
 #endif
+		
+		FAsyncPoseSearchDatabasesManagementCache DatabasesManagementCache;
 
-		// Draw query vector
-		if (Model->GetDrawQuery())
-		{
-			if (const UPoseSearchDatabase* CurrentDatabase = Model->GetCurrentDatabase())
+		// Draw selected poses
+		int32 NumOfRemainingRowsToDraw = 250;
+		TSet<const UPoseSearchDatabase*, DefaultKeyFuncs<const UPoseSearchDatabase*>, TInlineSetAllocator<8>> SelectedRowsDatabases;
+		TMap<const UPoseSearchDatabase*, FDebugDrawParams, TInlineSetAllocator<8>> CachedDebugDrawParams;
+		
+		auto DrawFeatureVectors = [&DatabasesManagementCache, &SelectedRowsDatabases, &CachedDebugDrawParams, &AnimContexts, &PoseHistories, &RoleToIndex, &NumOfRemainingRowsToDraw](TArray<TSharedRef<FDebuggerDatabaseRowData>> Rows)
 			{
-				for (const FTraceMotionMatchingStateDatabaseEntry& DbEntry : State->DatabaseEntries)
+				for (const TSharedRef<FDebuggerDatabaseRowData>& Row : Rows)
 				{
-					const UPoseSearchDatabase* Database = FTraceMotionMatchingStateMessage::GetObjectFromId<UPoseSearchDatabase>(DbEntry.DatabaseId);
-					if (Database && Database == CurrentDatabase &&
-						EAsyncBuildIndexResult::Success == FAsyncPoseSearchDatabasesManagement::RequestAsyncBuildIndex(CurrentDatabase, ERequestAsyncBuildFlag::ContinueRequest) &&
-						DbEntry.QueryVector.Num() == Database->Schema->SchemaCardinality)
+					if (NumOfRemainingRowsToDraw == 0)
 					{
-						FDebugDrawParams DrawParams(AnimContexts, PoseHistories, RoleToIndex, CurrentDatabase);
-						DrawParams.DrawFeatureVector(DbEntry.QueryVector);
 						break;
 					}
-				}
-			}
-		}
-
-		// Draw selected poses
-		const TSharedPtr<SListView<TSharedRef<FDebuggerDatabaseRowData>>>& DatabaseRows = DatabaseView->GetDatabaseRows();
-		TArray<TSharedRef<FDebuggerDatabaseRowData>> SelectedRows = DatabaseRows->GetSelectedItems();
 
-		// Draw any selected database vectors
-		constexpr int32 MaxRowsToDraw = 250;
-
-		const int32 NumRowsToDraw = FMath::Min(MaxRowsToDraw, SelectedRows.Num());
-		for (int32 RowIdx = 0; RowIdx < NumRowsToDraw; ++RowIdx)
-		{
-			const TSharedRef<FDebuggerDatabaseRowData>& Row = SelectedRows[RowIdx];
-			const UPoseSearchDatabase* RowDatabase = Row->SharedData->SourceDatabase.Get();
-			if (EAsyncBuildIndexResult::Success == FAsyncPoseSearchDatabasesManagement::RequestAsyncBuildIndex(RowDatabase, ERequestAsyncBuildFlag::ContinueRequest))
-			{
-				FDebugDrawParams DrawParams(AnimContexts, PoseHistories, RoleToIndex, RowDatabase);
-				DrawParams.DrawFeatureVector(Row->PoseIdx);
-			}
-		}
+					const UPoseSearchDatabase* RowDatabase = Row->SharedData->SourceDatabase.Get();
+					if (EAsyncBuildIndexResult::Success == DatabasesManagementCache(RowDatabase))
+					{
+						SelectedRowsDatabases.Add(RowDatabase);
+						FDebugDrawParams* DrawParams = CachedDebugDrawParams.Find(RowDatabase);
+						if (!DrawParams)
+						{
+							DrawParams = &CachedDebugDrawParams.Add(RowDatabase);
+							DrawParams->Init(AnimContexts, PoseHistories, RoleToIndex, RowDatabase);
+						}
 
-		// Draw active pose
-		TArray<TSharedRef<FDebuggerDatabaseRowData>> ActiveRows = DatabaseView->GetActiveRow()->GetSelectedItems();
+						DrawParams->DrawFeatureVector(Row->PoseIdx);
+						--NumOfRemainingRowsToDraw;
+					}
+				}
+			};
 
-		// Active row should only have 0 or 1
-		check(ActiveRows.Num() < 2);
+		// @todo: avoid returning arrays by copy here
+		DrawFeatureVectors(DatabaseView->GetActiveRow()->GetSelectedItems());
+		DrawFeatureVectors(DatabaseView->GetContinuingPoseRow()->GetSelectedItems());
+		DrawFeatureVectors(DatabaseView->GetDatabaseRows()->GetSelectedItems());
 
-		if (!ActiveRows.IsEmpty())
+		// Draw queries: either the one for the selected row databases or ALL of them
+		if (Model->GetDrawQuery())
 		{
-			const UPoseSearchDatabase* Database = ActiveRows[0]->SharedData->SourceDatabase.Get();
-			if (EAsyncBuildIndexResult::Success == FAsyncPoseSearchDatabasesManagement::RequestAsyncBuildIndex(Database, ERequestAsyncBuildFlag::ContinueRequest))
+			for (const FTraceMotionMatchingStateDatabaseEntry& DbEntry : State->DatabaseEntries)
 			{
-				// Use the motion-matching state's pose idx, as the active row may be update-throttled at this point
-				FDebugDrawParams DrawParams(AnimContexts, PoseHistories, RoleToIndex, Database);
-				DrawParams.DrawFeatureVector(ActiveRows[0]->PoseIdx);
-			}
-		}
-
-
-		// Draw continuing pose
-		TArray<TSharedRef<FDebuggerDatabaseRowData>> ContinuingRows = DatabaseView->GetContinuingPoseRow()->GetSelectedItems();
-
-		// ContinuingPose row should only have 0 or 1
-		check(ContinuingRows.Num() < 2);
+				const UPoseSearchDatabase* Database = FTraceMotionMatchingStateMessage::GetObjectFromId<UPoseSearchDatabase>(DbEntry.DatabaseId);
+				if (Database &&
+					(SelectedRowsDatabases.IsEmpty() || SelectedRowsDatabases.Contains(Database)) &&
+					EAsyncBuildIndexResult::Success == DatabasesManagementCache(Database) &&
+					DbEntry.QueryVector.Num() == Database->Schema->SchemaCardinality)
+				{
+					FDebugDrawParams* DrawParams = CachedDebugDrawParams.Find(Database);
+					if (!DrawParams)
+					{
+						DrawParams = &CachedDebugDrawParams.Add(Database);
+						DrawParams->Init(AnimContexts, PoseHistories, RoleToIndex, Database);
+					}
 
-		if (!ContinuingRows.IsEmpty())
-		{
-			const UPoseSearchDatabase* Database = ContinuingRows[0]->SharedData->SourceDatabase.Get();
-			if (EAsyncBuildIndexResult::Success == FAsyncPoseSearchDatabasesManagement::RequestAsyncBuildIndex(Database, ERequestAsyncBuildFlag::ContinueRequest))
-			{
-				FDebugDrawParams DrawParams(AnimContexts, PoseHistories, RoleToIndex, Database);
-				DrawParams.DrawFeatureVector(ContinuingRows[0]->PoseIdx);
+					DrawParams->DrawFeatureVector(DbEntry.QueryVector);
+				}
 			}
 		}
 	}
diff --git a/Engine/Plugins/Animation/PoseSearch/Source/Editor/Private/PoseSearchDebuggerViewModel.cpp b/Engine/Plugins/Animation/PoseSearch/Source/Editor/Private/PoseSearchDebuggerViewModel.cpp
--- a/Engine/Plugins/Animation/PoseSearch/Source/Editor/Private/PoseSearchDebuggerViewModel.cpp
+++ b/Engine/Plugins/Animation/PoseSearch/Source/Editor/Private/PoseSearchDebuggerViewModel.cpp
@@ -43,17 +43,6 @@ const FTraceMotionMatchingStateMessage* FDebuggerViewModel::GetMotionMatchingSta
 	return nullptr;
 }
 
-const UPoseSearchDatabase* FDebuggerViewModel::GetCurrentDatabase() const
-{
-	if (const FTraceMotionMatchingStateMessage* State = GetMotionMatchingState())
-	{
-		PRAGMA_DISABLE_DEPRECATION_WARNINGS
-		return State->GetCurrentDatabase();
-		PRAGMA_ENABLE_DEPRECATION_WARNINGS
-	}
-	return nullptr;
-}
-
 int32 FDebuggerViewModel::GetNodesNum() const
 {
 	return MotionMatchingStates.Num();
diff --git a/Engine/Plugins/Animation/PoseSearch/Source/Editor/Private/PoseSearchDebuggerViewModel.h b/Engine/Plugins/Animation/PoseSearch/Source/Editor/Private/PoseSearchDebuggerViewModel.h
--- a/Engine/Plugins/Animation/PoseSearch/Source/Editor/Private/PoseSearchDebuggerViewModel.h
+++ b/Engine/Plugins/Animation/PoseSearch/Source/Editor/Private/PoseSearchDebuggerViewModel.h
@@ -29,7 +29,6 @@ public:
 	const TMap<uint64, TWeakObjectPtr<AActor>>& GetDebugDrawActors() const { return DebugDrawActors; }
 	const TArray<FTraceMotionMatchingStateMessage>& GetMotionMatchingStates() const { return MotionMatchingStates; }
 
-	const UPoseSearchDatabase* GetCurrentDatabase() const;
 	int32 GetNodesNum() const;
 
 	/** Update motion matching states for frame */
diff --git a/Engine/Plugins/Animation/PoseSearch/Source/Runtime/Private/PoseSearchContext.cpp b/Engine/Plugins/Animation/PoseSearch/Source/Runtime/Private/PoseSearchContext.cpp
--- a/Engine/Plugins/Animation/PoseSearch/Source/Runtime/Private/PoseSearchContext.cpp
+++ b/Engine/Plugins/Animation/PoseSearch/Source/Runtime/Private/PoseSearchContext.cpp
@@ -219,11 +219,62 @@ FDebugDrawParams::FDebugDrawParams(TArrayView<const USkinnedMeshComponent*> InMe
 PRAGMA_ENABLE_DEPRECATION_WARNINGS
 
 FDebugDrawParams::FDebugDrawParams(TConstArrayView<FChooserEvaluationContext*> InAnimContexts, TConstArrayView<const IPoseHistory*> InPoseHistories, const FRoleToIndex& InRoleToIndex, const UPoseSearchDatabase* InDatabase)
-: AnimContexts(InAnimContexts)
-, PoseHistories(InPoseHistories)
-, RoleToIndex(InRoleToIndex)
-, Database(InDatabase)
+: FDebugDrawParams()
+{
+	Init(InAnimContexts, InPoseHistories, InRoleToIndex, InDatabase);
+}
+
+FDebugDrawParams::FDebugDrawParams(const FDebugDrawParams& Other)
+: AnimContexts(Other.AnimContexts)
+, PoseHistories(Other.PoseHistories)
+, RoleToIndex(Other.RoleToIndex)
+, Database(Other.Database)
+, DynamicWeightsSqrtBuffer(Other.DynamicWeightsSqrtBuffer)
+{
+	// pointing DynamicWeightsSqrt this object internal DynamicWeightsSqrtBuffer memory buffer!
+	DynamicWeightsSqrt = DynamicWeightsSqrtBuffer;
+}
+
+FDebugDrawParams::FDebugDrawParams(FDebugDrawParams&& Other)
+: AnimContexts(MoveTemp(Other.AnimContexts))
+, PoseHistories(MoveTemp(Other.PoseHistories))
+, RoleToIndex(MoveTemp(Other.RoleToIndex))
+, Database(Other.Database)
+, DynamicWeightsSqrtBuffer(MoveTemp(Other.DynamicWeightsSqrtBuffer))
+{
+	// pointing DynamicWeightsSqrt this object internal DynamicWeightsSqrtBuffer memory buffer!
+	DynamicWeightsSqrt = DynamicWeightsSqrtBuffer;
+
+	Other.Database = nullptr;
+}
+
+FDebugDrawParams& FDebugDrawParams::operator=(const FDebugDrawParams& Other)
+{
+	if (this != &Other)
+	{
+		this->~FDebugDrawParams();
+		new(this) FDebugDrawParams(Other);
+	}
+	return *this;
+}
+
+FDebugDrawParams& FDebugDrawParams::operator=(FDebugDrawParams&& Other)
+{
+	if (this != &Other)
+	{
+		this->~FDebugDrawParams();
+		new(this) FDebugDrawParams(MoveTemp(Other));
+	}
+	return *this;
+}
+
+void FDebugDrawParams::Init(TConstArrayView<FChooserEvaluationContext*> InAnimContexts, TConstArrayView<const IPoseHistory*> InPoseHistories, const FRoleToIndex& InRoleToIndex, const UPoseSearchDatabase* InDatabase)
 {
+	AnimContexts = InAnimContexts;
+	PoseHistories = InPoseHistories;
+	RoleToIndex = InRoleToIndex;
+	Database = InDatabase;
+
 	check(RoleToIndex.Num() == PoseHistories.Num());
 	check(IsValid(RoleToIndex));
 
@@ -716,6 +767,103 @@ FCachedQuery::FCachedQuery(const UPoseSearchSchema* InSchema)
 	Values.SetNumZeroed(Schema->SchemaCardinality);
 }
 
+#if UE_POSE_SEARCH_TRACE_ENABLED
+//////////////////////////////////////////////////////////////////////////
+// FSearchContext::FBestPoseCandidates
+FSearchContext::FBestPoseCandidates::FBestPoseCandidates()
+{
+	// preallocating memory to avoid multiple reallocations / rehashing
+	TransientPoseCandidates.Reserve(MaxNumberOfCollectedPoseCandidatesPerDatabase);
+	PermanentPoseCandidates.Reserve(16);
+}
+
+void FSearchContext::FBestPoseCandidates::Add(int32 PoseIdx, EPoseCandidateFlags PoseCandidateFlags, const FPoseSearchCost& Cost)
+{
+	if (EnumHasAnyFlags(PoseCandidateFlags, EPoseCandidateFlags::Valid_ContinuingPose | EPoseCandidateFlags::Valid_CurrentPose))
+	{
+		const FPoseCandidate PoseCandidate(PoseIdx, Cost, PoseCandidateFlags);
+		bool bIsAlreadyInSet = false;
+		FPoseCandidate& FoundPoseCandidate = PermanentPoseCandidates.FindOrAdd(PoseCandidate, &bIsAlreadyInSet);
+		if (bIsAlreadyInSet)
+		{
+			if (Cost < FoundPoseCandidate.Cost)
+			{
+				FoundPoseCandidate.Cost = Cost;
+			}
+			FoundPoseCandidate.PoseCandidateFlags |= PoseCandidateFlags;
+		}
+	}
+	else if (TransientPoseCandidates.Num() < MaxNumberOfCollectedPoseCandidatesPerDatabase)
+	{
+		const FPoseCandidate PoseCandidate(PoseIdx, Cost, PoseCandidateFlags);
+		TransientPoseCandidates.HeapPush(PoseCandidate, FPoseCandidateHeapCompare());
+	}
+	else if (Cost < TransientPoseCandidates.HeapTop().Cost)
+	{
+		// popping the max heap (TransientPoseCandidates) head since it costs more than the candidate we're about to insert
+		FPoseCandidate PoppedPoseCandidate;
+		TransientPoseCandidates.HeapPop(PoppedPoseCandidate, FPoseCandidateHeapCompare(), EAllowShrinking::No);
+
+		const FPoseCandidate PoseCandidate(PoseIdx, Cost, PoseCandidateFlags);
+		TransientPoseCandidates.HeapPush(PoseCandidate, FPoseCandidateHeapCompare());
+	}
+}
+
+bool FSearchContext::FBestPoseCandidates::IterateOverBestPoseCandidates(const TFunctionRef<bool(const FPoseCandidate& PoseCandidate)> IterateOverBestPoseCandidatesFunction) const
+{
+	TSet<uint32> DiscoveredPoseIndexes;
+	DiscoveredPoseIndexes.Reserve(PermanentPoseCandidates.Num() + TransientPoseCandidates.Num());
+
+	for (const FPoseCandidate& PoseCandidate : PermanentPoseCandidates)
+	{
+		if (IterateOverBestPoseCandidatesFunction(PoseCandidate))
+		{
+			return true;
+		}
+			
+		DiscoveredPoseIndexes.Add(PoseCandidate.PoseIdx);
+	}
+
+	for (const FPoseCandidate& PoseCandidate : TransientPoseCandidates)
+	{
+		bool bIsAlreadyInSet = false;
+		DiscoveredPoseIndexes.FindOrAdd(PoseCandidate.PoseIdx, &bIsAlreadyInSet);
+		if (!bIsAlreadyInSet && IterateOverBestPoseCandidatesFunction(PoseCandidate))
+		{
+			return true;
+		}
+	}
+		
+	return false;
+}
+
+int32 FSearchContext::FBestPoseCandidates::Num() const
+{
+	return PermanentPoseCandidates.Num() + TransientPoseCandidates.Num();
+}
+
+FSearchContext::FPoseCandidate FSearchContext::FBestPoseCandidates::GetUnsortedCandidate(int32 Index) const
+{
+	PRAGMA_DISABLE_DEPRECATION_WARNINGS
+	check(Index >= 0 && Index < Num());
+	PRAGMA_ENABLE_DEPRECATION_WARNINGS
+
+	if (Index < PermanentPoseCandidates.Num())
+	{
+		TSet<FPoseCandidate, FPoseCandidatePoseIdxKeyFunc>::TConstIterator ConstIterator = PermanentPoseCandidates.CreateConstIterator();
+		for (int32 Iteration = 0; Iteration < Index; ++Iteration)
+		{
+			++ConstIterator;
+			check(ConstIterator);
+		}
+		return *ConstIterator;
+	}
+
+	return TransientPoseCandidates[Index - PermanentPoseCandidates.Num()];
+}
+
+#endif // UE_POSE_SEARCH_TRACE_ENABLED
+
 //////////////////////////////////////////////////////////////////////////
 // FSearchContext
 FSearchContext::FSearchContext(float InDesiredPermutationTimeOffset, const FPoseIndicesHistory* InPoseIndicesHistory,
diff --git a/Engine/Plugins/Animation/PoseSearch/Source/Runtime/Private/PoseSearchDerivedData.cpp b/Engine/Plugins/Animation/PoseSearch/Source/Runtime/Private/PoseSearchDerivedData.cpp
--- a/Engine/Plugins/Animation/PoseSearch/Source/Runtime/Private/PoseSearchDerivedData.cpp
+++ b/Engine/Plugins/Animation/PoseSearch/Source/Runtime/Private/PoseSearchDerivedData.cpp
@@ -2823,6 +2823,20 @@ EAsyncBuildIndexResult FAsyncPoseSearchDatabasesManagement::RequestAsyncBuildInd
 	return RequestAsyncBuildIndexInternal(Database, Flag);
 }
 
+//////////////////////////////////////////////////////////////////////////
+// FAsyncPoseSearchDatabasesManagementCache
+EAsyncBuildIndexResult FAsyncPoseSearchDatabasesManagementCache::operator()(const UPoseSearchDatabase* Database, ERequestAsyncBuildFlag RequestAsyncBuildFlag)
+{
+	if (const EAsyncBuildIndexResult* AsyncBuildIndexResult = CachedRequests.Find(Database))
+	{
+		return *AsyncBuildIndexResult;
+	}
+
+	const EAsyncBuildIndexResult NewAsyncBuildIndexResult = FAsyncPoseSearchDatabasesManagement::RequestAsyncBuildIndex(Database, RequestAsyncBuildFlag);
+	CachedRequests.Add(Database) = NewAsyncBuildIndexResult;
+	return NewAsyncBuildIndexResult;
+}
+
 } // namespace UE::PoseSearch
 
 #undef LOCTEXT_NAMESPACE
diff --git a/Engine/Plugins/Animation/PoseSearch/Source/Runtime/Private/PoseSearchLibrary.cpp b/Engine/Plugins/Animation/PoseSearch/Source/Runtime/Private/PoseSearchLibrary.cpp
--- a/Engine/Plugins/Animation/PoseSearch/Source/Runtime/Private/PoseSearchLibrary.cpp
+++ b/Engine/Plugins/Animation/PoseSearch/Source/Runtime/Private/PoseSearchLibrary.cpp
@@ -664,6 +664,17 @@ void UPoseSearchLibrary::TraceMotionMatching(
 		TraceState.PoseHistories[AnimInstanceIndex].InitFrom(SearchContext.GetPoseHistories()[AnimInstanceIndex]);
 	}
 
+	// finalizing SearchContext.GetBestPoseCandidatesMap() by calling SearchContext::Track for all the SearchResults flagging them as EPoseCandidateFlags::Valid_CurrentPose
+	// because, those candidate could have been discarded becasue of their cost was too high and didn't fit into SearchContext.GetBestPoseCandidatesMap() limits
+	SearchResults.IterateOverSearchResults([&SearchContext](const FSearchResult& SearchResult)
+		{
+			if (SearchResult.IsValid())
+			{
+				SearchContext.Track(SearchResult.Database.Get(), SearchResult.PoseIdx, EPoseCandidateFlags::Valid_CurrentPose, SearchResult.PoseCost);
+			}
+			return false;
+		});
+
 	TArray<uint64, TInlineAllocator<64>> DatabaseIds;
 	int32 DbEntryIdx = 0;
 
@@ -681,39 +692,29 @@ void UPoseSearchLibrary::TraceMotionMatching(
 		DbEntry.DatabaseId = FObjectTrace::GetObjectId(Database);
 		DatabaseIds.Add(DbEntry.DatabaseId);
 
-		for (int32 CandidateIdx = 0; CandidateIdx < DatabaseBestPoseCandidates.Value.Num(); ++CandidateIdx)
-		{
-			const FSearchContext::FPoseCandidate PoseCandidate = DatabaseBestPoseCandidates.Value.GetUnsortedCandidate(CandidateIdx);
-
-			FTraceMotionMatchingStatePoseEntry PoseEntry;
-			PoseEntry.DbPoseIdx = PoseCandidate.PoseIdx;
-			PoseEntry.Cost = PoseCandidate.Cost;
-			PoseEntry.PoseCandidateFlags = PoseCandidate.PoseCandidateFlags;
-			
-			// looking for PoseCandidate/Database in SearchResults
-			if (SearchResults.IterateOverSearchResults([&PoseCandidate, Database](const FSearchResult& SearchResult)
-				{
-					return PoseCandidate.PoseIdx == SearchResult.PoseIdx && Database == SearchResult.Database.Get();
-				}))
-			{
-				// if found PoseCandidate is a Valid_CurrentPose
-				check(EnumHasAnyFlags(PoseEntry.PoseCandidateFlags, EPoseCandidateFlags::Valid_Pose | EPoseCandidateFlags::Valid_ContinuingPose));
-				EnumAddFlags(PoseEntry.PoseCandidateFlags, EPoseCandidateFlags::Valid_CurrentPose);
-
-				PRAGMA_DISABLE_DEPRECATION_WARNINGS
-				TraceState.CurrentDbEntryIdx = DbEntryIdx;
-				TraceState.CurrentPoseEntryIdx = DbEntry.PoseEntries.Add(PoseEntry);
-				PRAGMA_ENABLE_DEPRECATION_WARNINGS
-			}
-			else
+		DatabaseBestPoseCandidates.Value.IterateOverBestPoseCandidates([&DbEntry](const FSearchContext::FPoseCandidate& PoseCandidate)
 			{
+				// @todo replace FTraceMotionMatchingStatePoseEntry with FSearchContext::FPoseCandidate
+				FTraceMotionMatchingStatePoseEntry PoseEntry;
+				PoseEntry.DbPoseIdx = PoseCandidate.PoseIdx;
+				PoseEntry.Cost = PoseCandidate.Cost;
+				PoseEntry.PoseCandidateFlags = PoseCandidate.PoseCandidateFlags;
 				DbEntry.PoseEntries.Add(PoseEntry);
-			}
-		}
+				return false;
+			});
 
 		++DbEntryIdx;
 	}
 
+	// @todo: reenable this code if needed
+	//PRAGMA_DISABLE_DEPRECATION_WARNINGS
+	//if (SearchResult.IsValid())
+	//{
+	//	TraceState.CurrentDbEntryIdx = DbEntryIdx;
+	//	TraceState.CurrentPoseEntryIdx = DbEntry.PoseEntries.Add(PoseEntry);
+	//}
+	//PRAGMA_ENABLE_DEPRECATION_WARNINGS
+
 	DatabaseIds.Sort();
 	SearchId = HashCombineFast(SearchId, GetTypeHash(DatabaseIds));
 
diff --git a/Engine/Plugins/Animation/PoseSearch/Source/Runtime/Public/PoseSearch/PoseSearchContext.h b/Engine/Plugins/Animation/PoseSearch/Source/Runtime/Public/PoseSearch/PoseSearchContext.h
--- a/Engine/Plugins/Animation/PoseSearch/Source/Runtime/Public/PoseSearch/PoseSearchContext.h
+++ b/Engine/Plugins/Animation/PoseSearch/Source/Runtime/Public/PoseSearch/PoseSearchContext.h
@@ -79,7 +79,15 @@ struct FDebugDrawParams
 	UE_API FDebugDrawParams(TArrayView<const USkinnedMeshComponent*> InMeshes, TConstArrayView<const IPoseHistory*> InPoseHistories, const FRoleToIndex& InRoleToIndex, const UPoseSearchDatabase* InDatabase, EDebugDrawFlags InFlags = EDebugDrawFlags::None);
 	UE_API PRAGMA_ENABLE_DEPRECATION_WARNINGS
 	FDebugDrawParams(TConstArrayView<FChooserEvaluationContext*> InAnimContexts, TConstArrayView<const IPoseHistory*> InPoseHistories, const FRoleToIndex& InRoleToIndex, const UPoseSearchDatabase* InDatabase);
-	
+	FDebugDrawParams() = default;
+
+	UE_API FDebugDrawParams(const FDebugDrawParams& Other);
+	UE_API FDebugDrawParams(FDebugDrawParams&& Other);
+	UE_API FDebugDrawParams& operator=(const FDebugDrawParams& Other);
+	UE_API FDebugDrawParams& operator=(FDebugDrawParams&& Other);
+
+	UE_API void Init(TConstArrayView<FChooserEvaluationContext*> InAnimContexts, TConstArrayView<const IPoseHistory*> InPoseHistories, const FRoleToIndex& InRoleToIndex, const UPoseSearchDatabase* InDatabase);
+
 	UE_API const FSearchIndex* GetSearchIndex() const;
 	UE_API const UPoseSearchSchema* GetSchema() const;
 
@@ -104,12 +112,12 @@ struct FDebugDrawParams
 private:
 	bool CanDraw() const;
 
-	const TConstArrayView<FChooserEvaluationContext*> AnimContexts;
-	const TConstArrayView<const IPoseHistory*> PoseHistories;
+	TConstArrayView<FChooserEvaluationContext*> AnimContexts;
+	TConstArrayView<const IPoseHistory*> PoseHistories;
 
 	// NoTe: mapping Role to the index of the associated asset that this FDebugDrawParams is drawing.
 	// NOT the index of the UPoseSearchSchema::Skeletons! Use UPoseSearchSchema::GetRoledSkeleton API to resolve that Role to FPoseSearchRoledSkeleton 
-	const FRoleToIndex RoleToIndex;
+	FRoleToIndex RoleToIndex;
 
 	const UPoseSearchDatabase* Database = nullptr;
 	TAlignedArray<float> DynamicWeightsSqrtBuffer;
@@ -344,109 +352,70 @@ private:
 
 #if UE_POSE_SEARCH_TRACE_ENABLED
 
-	struct FPoseCandidateIdCost
-	{
-		int32 PoseIdx = 0;
-		FPoseSearchCost Cost;
-		bool operator<(const FPoseCandidateIdCost& Other) const { return Other.Cost < Cost; } // Reverse compare because BestCandidates is a max heap
-	};
-	
 public:
 
-	struct FPoseCandidate : public FPoseCandidateIdCost
+	struct FPoseCandidate
 	{
-		EPoseCandidateFlags PoseCandidateFlags = EPoseCandidateFlags::None;
-	};
-
-	void Track(const UPoseSearchDatabase* Database, int32 PoseIdx = INDEX_NONE, EPoseCandidateFlags PoseCandidateFlags = EPoseCandidateFlags::None, const FPoseSearchCost& Cost = FPoseSearchCost())
-	{
-		if (Database)
+		FPoseCandidate(int32 InPoseIdx = 0, FPoseSearchCost InCost = FPoseSearchCost(), EPoseCandidateFlags InPoseCandidateFlags = EPoseCandidateFlags::None)
+			: PoseIdx(InPoseIdx)
+			, Cost(InCost)
+			, PoseCandidateFlags(InPoseCandidateFlags)
 		{
-			FBestPoseCandidates& BestPoseCandidates = BestPoseCandidatesMap.FindOrAdd(Database);
-			if (PoseIdx != INDEX_NONE)
-			{
-				BestPoseCandidates.Add(PoseIdx, PoseCandidateFlags, Cost);
-			}
 		}
-	}
+
+		int32 PoseIdx = 0;
+		FPoseSearchCost Cost;
+		EPoseCandidateFlags PoseCandidateFlags = EPoseCandidateFlags::None;
+	};
 
 	struct FBestPoseCandidates
 	{
-		FBestPoseCandidates()
-		{
-			// preallocating memory to avoid multiple reallocations / rehashing
-			PoseCandidateHeap.Reserve(MaxNumberOfCollectedPoseCandidatesPerDatabase);
-			PoseIdxToFlags.Empty(MaxNumberOfCollectedPoseCandidatesPerDatabase);
-		}
+		UE_API FBestPoseCandidates();
+		UE_API void Add(int32 PoseIdx, EPoseCandidateFlags PoseCandidateFlags, const FPoseSearchCost& Cost);
+		UE_API bool IterateOverBestPoseCandidates(const TFunctionRef<bool(const FPoseCandidate& PoseCandidate)> IterateOverBestPoseCandidatesFunction) const;
+
+		UE_DEPRECATED(5.7, "Use IterateOverBestPoseCandidates to iterate over FBestPoseCandidates instead")
+		UE_API int32 Num() const;
+		UE_DEPRECATED(5.7, "Use IterateOverBestPoseCandidates to iterate over FBestPoseCandidates instead")
+		UE_API FPoseCandidate GetUnsortedCandidate(int32 Index) const;
+		
+	private:
 
-		void Add(int32 PoseIdx, EPoseCandidateFlags PoseCandidateFlags, const FPoseSearchCost& Cost)
+		// @todo: evaluate memstack usage here!
+		
+		struct FPoseCandidateHeapCompare
 		{
-			check(PoseIdx >= 0);
-			if (EPoseCandidateFlags* PoseIdxPoseCandidateFlags = PoseIdxToFlags.Find(PoseIdx))
-			{
-				*PoseIdxPoseCandidateFlags |= PoseCandidateFlags;
-			}
-			else if (PoseCandidateHeap.Num() < MaxNumberOfCollectedPoseCandidatesPerDatabase || Cost < PoseCandidateHeap.HeapTop().Cost)
+			bool operator()(const FPoseCandidate& A, const FPoseCandidate& B) const
 			{
-				bool bPoppedContinuingPoseCandidate = false;
-				FPoseCandidate ContinuingPoseCandidate;
-				while (PoseCandidateHeap.Num() >= MaxNumberOfCollectedPoseCandidatesPerDatabase)
-				{
-					FPoseCandidate PoppedPoseCandidate;
-					Pop(PoppedPoseCandidate);
-
-					if (EnumHasAnyFlags(PoppedPoseCandidate.PoseCandidateFlags, EPoseCandidateFlags::Valid_ContinuingPose))
-					{
-						// we can only have one continuing pose candidate
-						check(!bPoppedContinuingPoseCandidate);
-						ContinuingPoseCandidate = PoppedPoseCandidate;
-						bPoppedContinuingPoseCandidate = true;
-					}					
-				}
-
-				if (bPoppedContinuingPoseCandidate)
-				{
-					// if we popped the continuing pose candidate, we make some space for it and push it back
-					FPoseCandidate PoppedPoseCandidate;
-					Pop(PoppedPoseCandidate);
-					PoseCandidateHeap.HeapPush(ContinuingPoseCandidate);
-					PoseIdxToFlags.Add(ContinuingPoseCandidate.PoseIdx, ContinuingPoseCandidate.PoseCandidateFlags);
-				}
-
-				FPoseCandidate PoseCandidate;
-				PoseCandidate.PoseIdx = PoseIdx;
-				PoseCandidate.Cost = Cost;
-				PoseCandidateHeap.HeapPush(PoseCandidate);
-				PoseIdxToFlags.Add(PoseIdx, PoseCandidateFlags);
+				// using > to create a max heap where TransientPoseCandidates.HeapTop().Cost is the greatest Cost (to behave like the std::priority_queue)
+				return A.Cost > B.Cost;
 			}
-		}
+		};
 
-		int32 Num() const
-		{
-			return PoseCandidateHeap.Num();
-		}
+		// max heap of Valid_Pose(s). It could contains duplicate poses, that we filter out in IterateOverBestPoseCandidates!
+		TArray<FPoseCandidate> TransientPoseCandidates;
 
-		FPoseCandidate GetUnsortedCandidate(int32 Index) const
+		struct FPoseCandidatePoseIdxKeyFunc : BaseKeyFuncs<FPoseCandidate, FPoseCandidate, false>
 		{
-			FPoseCandidate PoseCandidate;
-			const FPoseCandidateIdCost& PoseCandidateIdCost = PoseCandidateHeap[Index];
-			PoseCandidate.PoseIdx = PoseCandidateIdCost.PoseIdx;
-			PoseCandidate.Cost = PoseCandidateIdCost.Cost;
-			PoseCandidate.PoseCandidateFlags = PoseIdxToFlags[PoseCandidateIdCost.PoseIdx];
-			return PoseCandidate;
-		}
+			static KeyInitType GetSetKey(ElementInitType& Element) { return Element; }
+			static bool Matches(KeyInitType A, KeyInitType B) { return A.PoseIdx == B.PoseIdx; }
+			static uint32 GetKeyHash(KeyInitType Key) { return GetTypeHash(Key.PoseIdx); }
+		};
 
-	private:
-		void Pop(FPoseCandidate& OutItem)
+		// set of Valid_ContinuingPose(s) and Valid_CurrentPose(s). the set contains unique poses with the lowest cost and an union of their EPoseCandidateFlags
+		TSet<FPoseCandidate, FPoseCandidatePoseIdxKeyFunc> PermanentPoseCandidates;
+	};
+	
+	void Track(const UPoseSearchDatabase* Database, int32 PoseIdx = INDEX_NONE, EPoseCandidateFlags PoseCandidateFlags = EPoseCandidateFlags::None, const FPoseSearchCost& Cost = FPoseSearchCost())
+	{
+		check(Database);
+		FBestPoseCandidates& BestPoseCandidates = BestPoseCandidatesMap.FindOrAdd(Database);
+		if (PoseIdx != INDEX_NONE)
 		{
-			PoseCandidateHeap.HeapPop(OutItem, EAllowShrinking::No);
-			OutItem.PoseCandidateFlags = PoseIdxToFlags.FindAndRemoveChecked(OutItem.PoseIdx);
+			BestPoseCandidates.Add(PoseIdx, PoseCandidateFlags, Cost);
 		}
+	}
 
-		TArray<FPoseCandidateIdCost> PoseCandidateHeap;
-		TMap<int32, EPoseCandidateFlags> PoseIdxToFlags;
-	};
-	
 	const TMap<const UPoseSearchDatabase*, FBestPoseCandidates>& GetBestPoseCandidatesMap() const
 	{
 		return BestPoseCandidatesMap;
diff --git a/Engine/Plugins/Animation/PoseSearch/Source/Runtime/Public/PoseSearch/PoseSearchDerivedData.h b/Engine/Plugins/Animation/PoseSearch/Source/Runtime/Public/PoseSearch/PoseSearchDerivedData.h
--- a/Engine/Plugins/Animation/PoseSearch/Source/Runtime/Public/PoseSearch/PoseSearchDerivedData.h
+++ b/Engine/Plugins/Animation/PoseSearch/Source/Runtime/Public/PoseSearch/PoseSearchDerivedData.h
@@ -8,6 +8,8 @@
 #include "UObject/GCObject.h"
 #include "UObject/ObjectSaveContext.h"
 
+#define UE_API POSESEARCH_API
+
 class UAnimSequenceBase;
 class UPoseSearchDatabase;
 
@@ -34,7 +36,7 @@ namespace UE::PoseSearch
 	class FAsyncPoseSearchDatabasesManagement : public FTickableGameObject, public FTickableCookObject, public FGCObject
 	{
 	public:
-		static POSESEARCH_API EAsyncBuildIndexResult RequestAsyncBuildIndex(const UPoseSearchDatabase* Database, ERequestAsyncBuildFlag Flag);
+		static UE_API EAsyncBuildIndexResult RequestAsyncBuildIndex(const UPoseSearchDatabase* Database, ERequestAsyncBuildFlag Flag);
 
 	private:
 		FAsyncPoseSearchDatabasesManagement();
@@ -91,6 +93,16 @@ namespace UE::PoseSearch
 
 		static FCriticalSection Mutex;
 	};
+
+	// Experimental, this feature might be removed without warning, not for production use
+	struct FAsyncPoseSearchDatabasesManagementCache
+	{
+		UE_API EAsyncBuildIndexResult operator()(const UPoseSearchDatabase* Database, ERequestAsyncBuildFlag RequestAsyncBuildFlag = ERequestAsyncBuildFlag::ContinueRequest);
+	private:
+		TMap<const UPoseSearchDatabase*, EAsyncBuildIndexResult, TInlineSetAllocator<32>> CachedRequests;
+	};
 } // namespace UE::PoseSearch
 
+#undef UE_API
+
 #endif // WITH_EDITOR
