PoseSearch
- converting ISearchResults into FSearchResults without virtual calls
- moving FPoseIndicesHistory into pose history node to be accessable from any MotionMatch function call (PoseSearchColumn, MM interactions don't need state to hold on to those values!)
- making PoseHistory mandatory for MM searches

diff --git a/Engine/Plugins/Animation/PoseSearch/Source/Editor/Private/Chooser/PoseSearchColumnEditor.cpp b/Engine/Plugins/Animation/PoseSearch/Source/Editor/Private/Chooser/PoseSearchColumnEditor.cpp
--- a/Engine/Plugins/Animation/PoseSearch/Source/Editor/Private/Chooser/PoseSearchColumnEditor.cpp
+++ b/Engine/Plugins/Animation/PoseSearch/Source/Editor/Private/Chooser/PoseSearchColumnEditor.cpp
@@ -1,14 +1,14 @@
 // Copyright Epic Games, Inc. All Rights Reserved.
 
 #include "PoseSearchColumnEditor.h"
-
 #include "Animation/AnimationAsset.h"
+#include "ChooserColumnHeader.h"
 #include "ObjectChooserWidgetFactories.h"
 #include "PoseSearch/Chooser/PoseSearchChooserColumn.h"
+#include "PoseSearch/PoseSearchDatabase.h"
 #include "SPropertyAccessChainWidget.h"
 #include "Widgets/Input/SNumericEntryBox.h"
 #include "Widgets/SBoxPanel.h"
-#include <ChooserColumnHeader.h>
 
 #define LOCTEXT_NAMESPACE "PoseSearchColumnEditor"
 
diff --git a/Engine/Plugins/Animation/PoseSearch/Source/Runtime/Private/AnimNode_MotionMatching.cpp b/Engine/Plugins/Animation/PoseSearch/Source/Runtime/Private/AnimNode_MotionMatching.cpp
--- a/Engine/Plugins/Animation/PoseSearch/Source/Runtime/Private/AnimNode_MotionMatching.cpp
+++ b/Engine/Plugins/Animation/PoseSearch/Source/Runtime/Private/AnimNode_MotionMatching.cpp
@@ -151,135 +151,147 @@ void FAnimNode_MotionMatching::UpdateAssetPlayer(const FAnimationUpdateContext&
 		PoseHistory = &PoseHistoryProvider->GetPoseHistory();
 	}
 
-	FFloatInterval PoseSearchPlayRate = PlayRate;
+	if (!PoseHistory)
+	{
+		UE_LOG(LogPoseSearch, Error, TEXT("FAnimNode_MotionMatching::UpdateAssetPlayer, missing IPoseHistory"));
+	}
+	else
+	{
+		FFloatInterval PoseSearchPlayRate = PlayRate;
 #if ENABLE_ANIM_DEBUG
-	bool bPoseSearchPlayRateEnabled = GVarAnimNodeMotionMatchingPlayRateEnabled == UE::Private::EPlayRateState::Enabled || GVarAnimNodeMotionMatchingPlayRateEnabled == UE::Private::EPlayRateState::PoseSearchOnly;
-	PoseSearchPlayRate = bPoseSearchPlayRateEnabled ? PoseSearchPlayRate : FFloatInterval(1.f, 1.f);
+		bool bPoseSearchPlayRateEnabled = GVarAnimNodeMotionMatchingPlayRateEnabled == UE::Private::EPlayRateState::Enabled || GVarAnimNodeMotionMatchingPlayRateEnabled == UE::Private::EPlayRateState::PoseSearchOnly;
+		PoseSearchPlayRate = bPoseSearchPlayRateEnabled ? PoseSearchPlayRate : FFloatInterval(1.f, 1.f);
 #endif // ENABLE_ANIM_DEBUG
 
-	const float DeltaTime = Context.GetDeltaTime();
+		const float DeltaTime = Context.GetDeltaTime();
 
-	FSearchContext SearchContext(0.f, &MotionMatchingState.PoseIndicesHistory, PoseJumpThresholdTime, EventToSearch.GetPlayRateOverriddenEvent(PoseSearchPlayRate));
-	FChooserEvaluationContext ChooserEvaluationContext(const_cast<UObject*>(Context.AnimInstanceProxy->GetAnimInstanceObject()));
-	SearchContext.AddRole(GetCommonDefaultRole(DatabasesToSearch), &ChooserEvaluationContext, PoseHistory);
+		FSearchContext SearchContext(0.f, PoseJumpThresholdTime, EventToSearch.GetPlayRateOverriddenEvent(PoseSearchPlayRate));
+		FChooserEvaluationContext ChooserEvaluationContext(const_cast<UObject*>(Context.AnimInstanceProxy->GetAnimInstanceObject()));
+		SearchContext.AddRole(GetCommonDefaultRole(DatabasesToSearch), &ChooserEvaluationContext, PoseHistory);
 
-	// If we can't advance or enough time has elapsed since the last pose jump then search
-	const bool bCanAdvance = SearchContext.GetContinuingPoseSearchResult().PoseIdx != INDEX_NONE;
-	const bool bSearch = !bCanAdvance || (MotionMatchingState.ElapsedPoseSearchTime >= SearchThrottleTime);
-	FSearchResult SearchResult;
-	if (bSearch)
-	{
-		MotionMatchingState.ElapsedPoseSearchTime = 0.f;
-			
-		TArrayView<const UObject*> AssetsToSearch((const UObject**)FMemory_Alloca(DatabasesToSearch.Num() * sizeof(UObject*)), (int32)DatabasesToSearch.Num());
-		for (int32 DatabaseToSearchIndex = 0; DatabaseToSearchIndex < DatabasesToSearch.Num(); ++DatabaseToSearchIndex)
+		// If we can't advance or enough time has elapsed since the last pose jump then search
+		const bool bCanAdvance = SearchContext.GetContinuingPoseSearchResult().PoseIdx != INDEX_NONE;
+		const bool bSearch = !bCanAdvance || (MotionMatchingState.ElapsedPoseSearchTime >= SearchThrottleTime);
+		FSearchResult SearchResult;
+		if (bSearch)
 		{
-			AssetsToSearch[DatabaseToSearchIndex] = DatabasesToSearch[DatabaseToSearchIndex].Get();
-		}
+			MotionMatchingState.ElapsedPoseSearchTime = 0.f;
 
-		FSearchResults_Single SearchResults;
-		FPoseSearchContinuingProperties ContinuingProperties;
-		ContinuingProperties.InitFrom(MotionMatchingState.SearchResult, NextUpdateInterruptMode);
-		UPoseSearchLibrary::MotionMatch(SearchContext, AssetsToSearch, ContinuingProperties, SearchResults);
-		SearchResult = SearchResults.GetBestResult();
-	}
-	else
-	{
-		MotionMatchingState.ElapsedPoseSearchTime += DeltaTime;
+			TArrayView<const UObject*> AssetsToSearch((const UObject**)FMemory_Alloca(DatabasesToSearch.Num() * sizeof(UObject*)), (int32)DatabasesToSearch.Num());
+			for (int32 DatabaseToSearchIndex = 0; DatabaseToSearchIndex < DatabasesToSearch.Num(); ++DatabaseToSearchIndex)
+			{
+				AssetsToSearch[DatabaseToSearchIndex] = DatabasesToSearch[DatabaseToSearchIndex].Get();
+			}
 
-		// if we didn't perform a search we carry on with the continuing pose search result
-		SearchResult = SearchContext.GetContinuingPoseSearchResult();
-		SearchResult.bIsContinuingPoseSearch = true;
+			FSearchResults_Single SearchResults;
+			FPoseSearchContinuingProperties ContinuingProperties;
+			ContinuingProperties.InitFrom(MotionMatchingState.SearchResult, NextUpdateInterruptMode);
+			UPoseSearchLibrary::MotionMatch(SearchContext, AssetsToSearch, ContinuingProperties, SearchResults);
+			SearchResult = SearchResults.GetBestResult();
+		}
+		else
+		{
+			MotionMatchingState.ElapsedPoseSearchTime += DeltaTime;
+
+			// if we didn't perform a search we carry on with the continuing pose search result
+			SearchResult = SearchContext.GetContinuingPoseSearchResult();
+			SearchResult.bIsContinuingPoseSearch = true;
 
 #if UE_POSE_SEARCH_TRACE_ENABLED
-		// in case we skipped the search, we still have to track we would have requested to evaluate DatabasesToSearch and SearchResult.Database
-		for (const TObjectPtr<const UPoseSearchDatabase>& DatabaseToSearch : DatabasesToSearch)
-		{
-			SearchContext.Track(DatabaseToSearch.Get());
-		}
+			// in case we skipped the search, we still have to track we would have requested to evaluate DatabasesToSearch and SearchResult.Database
+			for (const TObjectPtr<const UPoseSearchDatabase>& DatabaseToSearch : DatabasesToSearch)
+			{
+				SearchContext.Track(DatabaseToSearch.Get());
+			}
 
-		SearchContext.Track(SearchResult.Database.Get());
+			SearchContext.Track(SearchResult.Database.Get());
 #endif // UE_POSE_SEARCH_TRACE_ENABLED
-	}
+		}
 
-	const float WantedPlayRate = CalculateWantedPlayRate(SearchResult, SearchContext, PoseSearchPlayRate, PoseHistory ? PoseHistory->GetTrajectorySpeedMultiplier() : 1.f, EventToSearch);
-	MotionMatchingState.SearchResult.InitFrom(SearchResult, WantedPlayRate);
-	MotionMatchingState.PoseIndicesHistory.Update(SearchResult, DeltaTime, PoseReselectHistory);
+		const float WantedPlayRate = CalculateWantedPlayRate(SearchResult, SearchContext, PoseSearchPlayRate, PoseHistory ? PoseHistory->GetTrajectorySpeedMultiplier() : 1.f, EventToSearch);
+		MotionMatchingState.SearchResult.InitFrom(SearchResult, WantedPlayRate);
 
-	PRAGMA_DISABLE_DEPRECATION_WARNINGS
-	MotionMatchingState.CurrentSearchResult = SearchResult;
-	MotionMatchingState.WantedPlayRate = WantedPlayRate;
-	PRAGMA_ENABLE_DEPRECATION_WARNINGS
-	
-	const bool bJumpToPose = MotionMatchingState.SearchResult.SelectedAnim && !MotionMatchingState.SearchResult.bIsContinuingPoseSearch;
+		if (const FPoseIndicesHistory* PoseIndicesHistory = PoseHistory->GetPoseIndicesHistory())
+		{
+			// const casting here is safe since we're in the thread owning the pose history, and it's the correct place to update the previously selected poses
+			const_cast<FPoseIndicesHistory*>(PoseIndicesHistory)->Update(SearchResult, DeltaTime, PoseReselectHistory);
+		}
+
+		PRAGMA_DISABLE_DEPRECATION_WARNINGS
+		MotionMatchingState.CurrentSearchResult = SearchResult;
+		MotionMatchingState.WantedPlayRate = WantedPlayRate;
+		PRAGMA_ENABLE_DEPRECATION_WARNINGS
+
+		const bool bJumpToPose = MotionMatchingState.SearchResult.SelectedAnim && !MotionMatchingState.SearchResult.bIsContinuingPoseSearch;
 
-	PRAGMA_DISABLE_DEPRECATION_WARNINGS
-	MotionMatchingState.bJumpedToPose = bJumpToPose;
-	PRAGMA_ENABLE_DEPRECATION_WARNINGS
+		PRAGMA_DISABLE_DEPRECATION_WARNINGS
+		MotionMatchingState.bJumpedToPose = bJumpToPose;
+		PRAGMA_ENABLE_DEPRECATION_WARNINGS
 
-	float DesiredPlayRate = MotionMatchingState.SearchResult.WantedPlayRate * PlayRateMultiplier;
+		float DesiredPlayRate = MotionMatchingState.SearchResult.WantedPlayRate * PlayRateMultiplier;
 #if ENABLE_ANIM_DEBUG
-	DesiredPlayRate = GVarAnimNodeMotionMatchingPlayRateEnabled == UE::Private::EPlayRateState::Enabled ? DesiredPlayRate : PlayRateMultiplier;
+		DesiredPlayRate = GVarAnimNodeMotionMatchingPlayRateEnabled == UE::Private::EPlayRateState::Enabled ? DesiredPlayRate : PlayRateMultiplier;
 #endif // ENABLE_ANIM_DEBUG
 
-	UE::Anim::FNodeFunctionCaller::CallFunction(GetOnUpdateMotionMatchingStateFunction(), Context, *this);
+		UE::Anim::FNodeFunctionCaller::CallFunction(GetOnUpdateMotionMatchingStateFunction(), Context, *this);
 
-	// If a new pose is requested, blend into the new asset via BlendStackNode
-	if (bJumpToPose)
-	{
-		const UPoseSearchDatabase* CurrentResultDatabase = MotionMatchingState.SearchResult.SelectedDatabase.Get();
-		if (CurrentResultDatabase && CurrentResultDatabase->Schema)
+		// If a new pose is requested, blend into the new asset via BlendStackNode
+		if (bJumpToPose)
 		{
-			if (UAnimationAsset* AnimationAsset = MotionMatchingState.SearchResult.GetAnimationAssetForRole())
+			const UPoseSearchDatabase* CurrentResultDatabase = MotionMatchingState.SearchResult.SelectedDatabase.Get();
+			if (CurrentResultDatabase && CurrentResultDatabase->Schema)
 			{
-				// Clear up any sync group info before pushing new asset player (which will have sync info since its the highest weighted).
-				for (FBlendStackAnimPlayer& AnimPlayer : AnimPlayers)
+				if (UAnimationAsset* AnimationAsset = MotionMatchingState.SearchResult.GetAnimationAssetForRole())
 				{
-					if (FAnimNode_AssetPlayerBase* AssetPlayerNode = AnimPlayer.GetAssetPlayerNode())
+					// Clear up any sync group info before pushing new asset player (which will have sync info since its the highest weighted).
+					for (FBlendStackAnimPlayer& AnimPlayer : AnimPlayers)
 					{
-						AssetPlayerNode->SetGroupMethod(EAnimSyncMethod::DoNotSync);
-						AssetPlayerNode->SetGroupRole(EAnimGroupRole::CanBeLeader);
-						AssetPlayerNode->SetGroupName(NAME_None);
+						if (FAnimNode_AssetPlayerBase* AssetPlayerNode = AnimPlayer.GetAssetPlayerNode())
+						{
+							AssetPlayerNode->SetGroupMethod(EAnimSyncMethod::DoNotSync);
+							AssetPlayerNode->SetGroupRole(EAnimGroupRole::CanBeLeader);
+							AssetPlayerNode->SetGroupName(NAME_None);
+						}
 					}
-				}
 
-				bool bAbortBlend = false;
-				if (!AnimPlayers.IsEmpty() && GetBlendspaceParametersDeltaThreshold() > 0.f && AnimationAsset->IsA<UBlendSpace>())
-				{
-					if (AnimPlayers[0].GetAnimationAsset() == AnimationAsset)
+					bool bAbortBlend = false;
+					if (!AnimPlayers.IsEmpty() && GetBlendspaceParametersDeltaThreshold() > 0.f && AnimationAsset->IsA<UBlendSpace>())
 					{
-						const FVector CurrentBlendParameters = AnimPlayers[0].GetBlendParameters();
-						const FVector DesiredBlendParameters = GetBlendspaceParameters();
-						const float Delta = (CurrentBlendParameters - DesiredBlendParameters).SizeSquared();
-						if (Delta < FMath::Square(GetBlendspaceParametersDeltaThreshold()))
+						if (AnimPlayers[0].GetAnimationAsset() == AnimationAsset)
 						{
-							// If we haven't changed assets, and our currently playing blendspace xy is within threshold of change, then keep playing it.
-							// Time differences should be OK because of the call to AdjustAssetTime before the search.
-							bAbortBlend = true;
+							const FVector CurrentBlendParameters = AnimPlayers[0].GetBlendParameters();
+							const FVector DesiredBlendParameters = GetBlendspaceParameters();
+							const float Delta = (CurrentBlendParameters - DesiredBlendParameters).SizeSquared();
+							if (Delta < FMath::Square(GetBlendspaceParametersDeltaThreshold()))
+							{
+								// If we haven't changed assets, and our currently playing blendspace xy is within threshold of change, then keep playing it.
+								// Time differences should be OK because of the call to AdjustAssetTime before the search.
+								bAbortBlend = true;
+							}
 						}
 					}
+
+					if (!bAbortBlend)
+					{
+						FAnimNode_BlendStack_Standalone::BlendTo(Context, AnimationAsset, MotionMatchingState.SearchResult.SelectedTime, MotionMatchingState.SearchResult.bLoop,
+							MotionMatchingState.SearchResult.bIsMirrored, CurrentResultDatabase->Schema->GetMirrorDataTable(MotionMatchingState.SearchResult.Role), BlendTime,
+							BlendProfile, BlendOption, bUseInertialBlend, NAME_None, MotionMatchingState.SearchResult.BlendParameters, DesiredPlayRate, 0,
+							GetGroupName(), GetGroupRole(), GetGroupMethod(), GetOverridePositionWhenJoiningSyncGroupAsLeader());
+					}
 				}
-				
-				if (!bAbortBlend)
+				else
 				{
-					FAnimNode_BlendStack_Standalone::BlendTo(Context, AnimationAsset, MotionMatchingState.SearchResult.SelectedTime, MotionMatchingState.SearchResult.bLoop,
-						MotionMatchingState.SearchResult.bIsMirrored, CurrentResultDatabase->Schema->GetMirrorDataTable(MotionMatchingState.SearchResult.Role), BlendTime,
-						BlendProfile, BlendOption, bUseInertialBlend, NAME_None, MotionMatchingState.SearchResult.BlendParameters, DesiredPlayRate, 0,
-						GetGroupName(), GetGroupRole(), GetGroupMethod(), GetOverridePositionWhenJoiningSyncGroupAsLeader());
+					checkNoEntry();
 				}
 			}
-			else
-			{
-				checkNoEntry();
-			}
 		}
-	}
 
-	const bool bDidBlendToRequestAnInertialBlend = bJumpToPose && bUseInertialBlend;
-	UE::Anim::TOptionalScopedGraphMessage<UE::Anim::FAnimInertializationSyncScope> InertializationSync(bDidBlendToRequestAnInertialBlend, Context);
-	
-	FAnimNode_BlendStack_Standalone::UpdatePlayRate(DesiredPlayRate);
-	FAnimNode_BlendStack_Standalone::UpdateBlendspaceParameters(GetBlendspaceUpdateMode(), GetBlendspaceParameters());
+		const bool bDidBlendToRequestAnInertialBlend = bJumpToPose && bUseInertialBlend;
+		UE::Anim::TOptionalScopedGraphMessage<UE::Anim::FAnimInertializationSyncScope> InertializationSync(bDidBlendToRequestAnInertialBlend, Context);
+
+		FAnimNode_BlendStack_Standalone::UpdatePlayRate(DesiredPlayRate);
+		FAnimNode_BlendStack_Standalone::UpdateBlendspaceParameters(GetBlendspaceUpdateMode(), GetBlendspaceParameters());
+	}
 
 	FAnimNode_BlendStack_Standalone::UpdateAssetPlayer(Context);
 
diff --git a/Engine/Plugins/Animation/PoseSearch/Source/Runtime/Private/Chooser/PoseSearchChooserColumn.cpp b/Engine/Plugins/Animation/PoseSearch/Source/Runtime/Private/Chooser/PoseSearchChooserColumn.cpp
--- a/Engine/Plugins/Animation/PoseSearch/Source/Runtime/Private/Chooser/PoseSearchChooserColumn.cpp
+++ b/Engine/Plugins/Animation/PoseSearch/Source/Runtime/Private/Chooser/PoseSearchChooserColumn.cpp
@@ -22,6 +22,20 @@
 
 #define LOCTEXT_NAMESPACE "FPoseSearchColumn"
 
+namespace UE::PoseSearch
+{
+	// @todo: we should rework how to pass temporary data from Filter to SetOutputs, but for now this makes things working
+	struct FPoseSearchColumnContextScratchArea : public TThreadSingleton<FPoseSearchColumnContextScratchArea>
+	{
+		TArray<UE::PoseSearch::FSearchResult> SearchResults;
+		const UE::PoseSearch::IPoseHistory* PoseHistory = nullptr;
+#if DO_CHECK
+		const FPoseSearchColumn* DebugOwner = nullptr;
+#endif // DO_CHECK
+	};
+
+} // namespace UE::PoseSearch
+
 bool FPoseHistoryContextProperty::GetValue(FChooserEvaluationContext& Context, FPoseHistoryReference& OutResult) const
 {
 	const FPoseHistoryReference* Reference;
@@ -33,12 +47,10 @@ bool FPoseHistoryContextProperty::GetValue(FChooserEvaluationContext& Context, F
 	return false;
 }
 
-
 FPoseSearchColumn::FPoseSearchColumn()
 {
 }
 
-
 void FPoseSearchColumn::CheckConsistency() const
 {
 #if DO_CHECK && WITH_EDITORONLY_DATA
@@ -237,6 +249,12 @@ void FPoseSearchColumn::Filter(FChooserEvaluationContext& Context, const FChoose
 
 	CheckConsistency();
 
+	FPoseSearchColumnContextScratchArea& ScratchArea = FPoseSearchColumnContextScratchArea::Get();
+#if DO_CHECK
+	ScratchArea.DebugOwner = this;
+#endif // DO_CHECK
+	ScratchArea.SearchResults.Reset();
+
 	const UE::PoseSearch::IPoseHistory* PoseHistory = nullptr;
 	FPoseHistoryReference PoseHistoryReference;
 	if (const FChooserParameterPoseHistoryBase* PoseHistoryParameter = InputValue.GetPtr<FChooserParameterPoseHistoryBase>())
@@ -266,8 +284,14 @@ void FPoseSearchColumn::Filter(FChooserEvaluationContext& Context, const FChoose
 		}
 	}
 
-	if (PoseHistory)
+	if (!PoseHistory)
+	{
+		UE_LOG(LogPoseSearch, Error, TEXT("FPoseSearchColumn::Filter, missing IPoseHistory"));
+	}
+	else
 	{
+		ScratchArea.PoseHistory = PoseHistory;
+
 		TArray<const UObject*, TInlineAllocator<128>> AssetsToConsider;
 		for (const FChooserIndexArray::FIndexData& IndexData : IndexListIn)
 		{
@@ -286,7 +310,7 @@ void FPoseSearchColumn::Filter(FChooserEvaluationContext& Context, const FChoose
 
 		const FFloatInterval PoseJumpThresholdTime(0.f, 0.f);
 		const UObject* DatabasesToSearch[] = { InternalDatabase.Get() };
-		FSearchContext SearchContext(0.f, nullptr /*&MotionMatchingState.PoseIndicesHistory*/, PoseJumpThresholdTime, FPoseSearchEvent());
+		FSearchContext SearchContext(0.f, PoseJumpThresholdTime, FPoseSearchEvent());
 
 		const FRole Role = InternalDatabase->Schema ? InternalDatabase->Schema->GetDefaultRole() : DefaultRole;
 		SearchContext.AddRole(Role, &Context, PoseHistory);
@@ -299,46 +323,44 @@ void FPoseSearchColumn::Filter(FChooserEvaluationContext& Context, const FChoose
 			ContinuingProperties.PlayingAssetAccumulatedTime = ChooserPlayerSettings->PlayingAssetAccumulatedTime;
 			ContinuingProperties.bIsPlayingAssetMirrored = ChooserPlayerSettings->bIsPlayingAssetMirrored;
 			ContinuingProperties.PlayingAssetBlendParameters = ChooserPlayerSettings->PlayingAssetBlendParameters;
+			ContinuingProperties.InterruptMode = InterruptMode;
 		}
-		FSearchResults_AssetBests SearchResults;
-		UPoseSearchLibrary::MotionMatch(SearchContext, DatabasesToSearch, ContinuingProperties, SearchResults);
-
-		// @todo: WIP
-		// SearchResults.GetSourceAssetIdxToBestSearchResult() contains the mapping between asset indes in the PSD and the frame with the best cost (SearchResult)
-		//for (const FChooserIndexArray::FIndexData& IndexData : IndexListIn)
-		//{
-		//	if (const FSearchResult* SearchResult = SearchResults.GetSourceAssetIdxToBestSearchResult().Find(IndexData.Index))
-		//	{
-		//		IndexListOut.Push({ IndexData.Index, SearchResult->PoseCost });
-		//	}
-		//}
-
-		// @todo: shouldn't OutputStartTime, OutputCost, OutputMirror be arrays?
-		const FSearchResult SearchResult = SearchResults.GetBestResult();
-		if (const UE::PoseSearch::FSearchIndexAsset* SearchIndexAsset = SearchResult.GetSearchIndexAsset())
+
+		FSearchResult BestSearchResult;
+		if (MaxNumberOfResults == 1)
 		{
-			if (const FPoseSearchDatabaseAnimationAsset* DatabaseAnimationAsset = SearchResult.Database->GetDatabaseAnimationAsset(*SearchIndexAsset))
+			FSearchResults_Single SearchResults;
+			UPoseSearchLibrary::MotionMatch(SearchContext, DatabasesToSearch, ContinuingProperties, SearchResults);
+			BestSearchResult = SearchResults.GetBestResult();
+
+			if (const FSearchIndexAsset* SearchIndexAsset = BestSearchResult.GetSearchIndexAsset())
+			{
+				check(BestSearchResult.Database == InternalDatabase);
+
+				// @todo: handle duplicate entries, by providing to the UPoseSearchLibrary::MotionMatch a SourceAssetIdxsToConsider rather than AssetsToConsider
+				const int32 SearchResultSourceAssetIdx = SearchIndexAsset->GetSourceAssetIdx();
+				check(SearchResultSourceAssetIdx >= 0);
+				ScratchArea.SearchResults.Add(BestSearchResult);
+				IndexListOut.Push({ static_cast<uint32>(SearchResultSourceAssetIdx), 0.f});
+			}
+		}
+		else
+		{
+			FSearchResults_AssetBests SearchResults;
+			UPoseSearchLibrary::MotionMatch(SearchContext, DatabasesToSearch, ContinuingProperties, SearchResults);
+			BestSearchResult = SearchResults.GetBestResult();
+			if (MaxNumberOfResults > 0)
 			{
-				const UObject* Result = DatabaseAnimationAsset->GetAnimationAsset();
-				int Index;
-				if (AssetsToConsider.Find(Result, Index))
+				// keeping only up to MaxNumberOfResults results
+				SearchResults.Shrink(MaxNumberOfResults);
+			}
+
+			for (const FChooserIndexArray::FIndexData& IndexData : IndexListIn)
+			{
+				if (const FSearchResult* FoundSearchResult = SearchResults.FindSearchResultFor(InternalDatabase.Get(), IndexData.Index))
 				{
-				 	IndexListOut.Push(IndexListIn[Index]);
-
-				 	if (const FChooserParameterFloatBase* StartTime = OutputStartTime.GetPtr<FChooserParameterFloatBase>())
-				 	{
-				 		StartTime->SetValue(Context, SearchResult.GetAssetTime());
-				 	}
-
-				 	if (const FChooserParameterFloatBase* Cost = OutputCost.GetPtr<FChooserParameterFloatBase>())
-                    {
-						Cost->SetValue(Context, SearchResult.PoseCost);
-                    }
-
-				 	if (const FChooserParameterBoolBase* Mirror = OutputMirror.GetPtr<FChooserParameterBoolBase>())
-				 	{
-				 		Mirror->SetValue(Context, SearchIndexAsset->IsMirrored());
-				 	}
+					ScratchArea.SearchResults.Add(*FoundSearchResult);
+					IndexListOut.Push({ IndexData.Index, FoundSearchResult->PoseCost });
 				}
 			}
 		}
@@ -348,23 +370,74 @@ void FPoseSearchColumn::Filter(FChooserEvaluationContext& Context, const FChoose
 	{
 		// if nothing passed the pose match, or if PoseHistory was not found, then ignore pose matching and just pass through everything
 		IndexListOut = IndexListIn;
+	}
+}
 
-		// set some reasonable defaults for our other outputs in the passthrough  case
-		if (const FChooserParameterFloatBase* StartTime = OutputStartTime.GetPtr<FChooserParameterFloatBase>())
+void FPoseSearchColumn::SetOutputs(FChooserEvaluationContext& Context, int RowIndex) const
+{
+	using namespace UE::PoseSearch;
+
+	// set an arbitrary "high" cost value so that cost threshold implementations can still wait for the pose search to find a good match
+	float CostValue = 100.f;
+	float StartTimeValue = 0.f;
+	bool bMirrorValue = false;
+
+	FPoseSearchColumnContextScratchArea& ScratchArea = FPoseSearchColumnContextScratchArea::Get();
+#if DO_CHECK
+	check(ScratchArea.DebugOwner == this);
+#endif // DO_CHECK
+
+	for (const FSearchResult& SearchResult : ScratchArea.SearchResults)
+	{
+		check(SearchResult.IsValid() && SearchResult.Database == InternalDatabase);
+		if (const FSearchIndexAsset* SearchIndexAsset = SearchResult.GetSearchIndexAsset())
 		{
-			StartTime->SetValue(Context, 0.0);
+			const int32 SearchResultSourceAssetIdx = SearchIndexAsset->GetSourceAssetIdx();
+			if (SearchResultSourceAssetIdx == RowIndex)
+			{
+				StartTimeValue = SearchResult.GetAssetTime();
+				CostValue = SearchResult.PoseCost;
+				bMirrorValue = SearchIndexAsset->IsMirrored();
+
+				if (ScratchArea.PoseHistory)
+				{
+					if (const FPoseIndicesHistory* PoseIndicesHistory = ScratchArea.PoseHistory->GetPoseIndicesHistory())
+					{
+						// @todo: integrate DeltaTime into SearchContext, and implement it for UAF as well
+						float DeltaTime = FiniteDelta;
+						if (const UAnimInstance* AnimInstance = Cast<UAnimInstance>(Context.GetFirstObjectParam()))
+						{
+							DeltaTime = AnimInstance->GetDeltaSeconds();
+						}
+
+						// const casting here is safe since we're in the thread owning the pose history, and it's the correct place to update the previously selected poses
+						const_cast<FPoseIndicesHistory*>(PoseIndicesHistory)->Update(SearchResult, DeltaTime, PoseReselectHistory);
+					}
+				}
+				break;
+			}
 		}
+	}
+	
+#if DO_CHECK
+	ScratchArea.DebugOwner = nullptr;
+#endif // DO_CHECK
+	ScratchArea.SearchResults.Reset();
+	ScratchArea.PoseHistory = nullptr;
+
+	if (const FChooserParameterFloatBase* StartTime = OutputStartTime.GetPtr<FChooserParameterFloatBase>())
+	{
+		StartTime->SetValue(Context, StartTimeValue);
+	}
 		
-		if (const FChooserParameterFloatBase* Cost = OutputCost.GetPtr<FChooserParameterFloatBase>())
-		{
-			// set an arbitrary "high" cost value so that cost threshold implementations can still wait for the pose search to find a good match
-			Cost->SetValue(Context, 100.0);
-		}
+	if (const FChooserParameterFloatBase* Cost = OutputCost.GetPtr<FChooserParameterFloatBase>())
+	{
+		Cost->SetValue(Context, CostValue);
+	}
 
-		if (const FChooserParameterBoolBase* Mirror = OutputMirror.GetPtr<FChooserParameterBoolBase>())
-		{
-			Mirror->SetValue(Context, false);
-		}
+	if (const FChooserParameterBoolBase* Mirror = OutputMirror.GetPtr<FChooserParameterBoolBase>())
+	{
+		Mirror->SetValue(Context, bMirrorValue);
 	}
 }
 
diff --git a/Engine/Plugins/Animation/PoseSearch/Source/Runtime/Private/PoseSearchContext.cpp b/Engine/Plugins/Animation/PoseSearch/Source/Runtime/Private/PoseSearchContext.cpp
--- a/Engine/Plugins/Animation/PoseSearch/Source/Runtime/Private/PoseSearchContext.cpp
+++ b/Engine/Plugins/Animation/PoseSearch/Source/Runtime/Private/PoseSearchContext.cpp
@@ -48,12 +48,12 @@ const USkeleton* GetContextSkeleton(const UObject* AnimContext)
 	// @todo: make sure this is not called from worker threads!
 	// check(IsInGameThread());
 
+	const USkeleton* Skeleton = nullptr;
 	if (const UAnimInstance* AnimInstance = Cast<UAnimInstance>(AnimContext))
 	{
-		return AnimInstance->GetRequiredBonesOnAnyThread().GetSkeletonAsset();
+		Skeleton = AnimInstance->GetRequiredBonesOnAnyThread().GetSkeletonAsset();
 	}
-
-	if (const UActorComponent* ActorComponent = Cast<UActorComponent>(AnimContext))
+	else if (const UActorComponent* ActorComponent = Cast<UActorComponent>(AnimContext))
 	{
 		const AActor* Actor = ActorComponent->GetOwner();
 		check(Actor);
@@ -62,14 +62,13 @@ const USkeleton* GetContextSkeleton(const UObject* AnimContext)
 		{
 			if (const USkeletalMesh* SkeletalMesh = SkeletalMeshComponent->GetSkeletalMeshAsset())
 			{
-				return SkeletalMesh->GetSkeleton();
+				Skeleton = SkeletalMesh->GetSkeleton();
 			}
 		}
-		return nullptr;
 	}
 
-	unimplemented();
-	return nullptr;
+	check(Skeleton)
+	return Skeleton;
 }
 	
 const USkeleton* GetContextSkeleton(FChooserEvaluationContext& Context)
@@ -77,12 +76,8 @@ const USkeleton* GetContextSkeleton(FChooserEvaluationContext& Context)
 	// @todo: make sure this is not called from worker threads!
 	// check(IsInGameThread());
 
-	if (Context.ObjectParams.Num() > 0)
-	{
-		return GetContextSkeleton(Context.ObjectParams[0].Object);
-	}
-
-	return nullptr;
+	check(!Context.ObjectParams.IsEmpty());
+	return GetContextSkeleton(Context.ObjectParams[0].Object);
 }
 
 const USkeletalMeshComponent* GetContextSkeletalMeshComponent(const UObject* AnimContext)
@@ -90,20 +85,21 @@ const USkeletalMeshComponent* GetContextSkeletalMeshComponent(const UObject* Ani
 	// thread unsafe code. making sure we're on game thread!
 	check(IsInGameThread());
 
+	const USkeletalMeshComponent* SkeletalMeshComponent = nullptr;
 	if (const UAnimInstance* AnimInstance = Cast<UAnimInstance>(AnimContext))
 	{
-		return AnimInstance->GetSkelMeshComponent();
+		SkeletalMeshComponent = AnimInstance->GetSkelMeshComponent();
 	}
-	
-	if (const UActorComponent* ActorComponent = Cast<UActorComponent>(AnimContext))
+	else if (const UActorComponent* ActorComponent = Cast<UActorComponent>(AnimContext))
 	{
 		const AActor* Actor = ActorComponent->GetOwner();
 		check(Actor);
 		// @todo: this code depends on how AnimNext gather its context object, and will likely change in the future
-		return Actor->GetComponentByClass<USkeletalMeshComponent>();
+		SkeletalMeshComponent = Actor->GetComponentByClass<USkeletalMeshComponent>();
 	}
 
-	return nullptr;
+	check(SkeletalMeshComponent);
+	return SkeletalMeshComponent;
 }
 
 const FBoneContainer GetBoneContainer(const UObject* AnimContext)
@@ -182,20 +178,15 @@ static FAnimInstanceProxy* GetAnimInstanceProxy(TConstArrayView<FChooserEvaluati
 
 static const USkinnedMeshComponent* GetMesh(TConstArrayView<FChooserEvaluationContext*> AnimContexts, int32 RoleIndex = 0)
 {
-	if (AnimContexts.IsValidIndex(RoleIndex))
-	{
-		return Cast<USkinnedMeshComponent>(AnimContexts[RoleIndex]->GetFirstObjectParam());
-	}
-	return nullptr;
+	check(AnimContexts.IsValidIndex(RoleIndex));
+	return Cast<USkinnedMeshComponent>(AnimContexts[RoleIndex]->GetFirstObjectParam());
 }
 
 static const UWorld* GetWorld(TConstArrayView<FChooserEvaluationContext*> AnimContexts)
 {
-	if (const USkinnedMeshComponent* SkinnedMeshComponent = GetMesh(AnimContexts))
-	{
-		return SkinnedMeshComponent->GetWorld();
-	}
-	return nullptr;
+	const USkinnedMeshComponent* SkinnedMeshComponent = GetMesh(AnimContexts);
+	check(SkinnedMeshComponent);
+	return SkinnedMeshComponent->GetWorld();
 }
 
 //////////////////////////////////////////////////////////////////////////
@@ -354,28 +345,26 @@ FVector FDebugDrawParams::ExtractPosition(TConstArrayView<float> PoseVector, flo
 			return WorldBonePosition;
 		}
 
-		if (const int32* RoleIndex = RoleToIndex.Find(Role))
+		const int32 RoleIndex = RoleToIndex[Role];
+		const IPoseHistory* PoseHistory = PoseHistories[RoleIndex];
+		check(PoseHistory);
+			
+		if (const USkeleton* Skeleton = Schema->GetSkeleton(Role))
 		{
-			if (const IPoseHistory* PoseHistory = PoseHistories[*RoleIndex])
-			{
-				if (const USkeleton* Skeleton = Schema->GetSkeleton(Role))
-				{
-					const FBoneIndexType BoneIndexType = SchemaBoneIdx == TrajectorySchemaBoneIdx ? ComponentSpaceIndexType : Schema->GetBoneReferences(Role)[SchemaBoneIdx].BoneIndex;
+			const FBoneIndexType BoneIndexType = SchemaBoneIdx == TrajectorySchemaBoneIdx ? ComponentSpaceIndexType : Schema->GetBoneReferences(Role)[SchemaBoneIdx].BoneIndex;
 
-					FTransform WorldBoneTransform;
-					if (PoseHistory->GetTransformAtTime(SampleTimeOffset + PermutationSampleTimeOffset, WorldBoneTransform, Skeleton, BoneIndexType, WorldSpaceIndexType))
-					{
-						return WorldBoneTransform.GetTranslation();
-					}
-				}
+			FTransform WorldBoneTransform;
+			if (PoseHistory->GetTransformAtTime(SampleTimeOffset + PermutationSampleTimeOffset, WorldBoneTransform, Skeleton, BoneIndexType, WorldSpaceIndexType))
+			{
+				return WorldBoneTransform.GetTranslation();
 			}
+		}
 
-			if (SchemaBoneIdx > RootSchemaBoneIdx)
+		if (SchemaBoneIdx > RootSchemaBoneIdx)
+		{
+			if (const USkinnedMeshComponent* Mesh = GetMesh(AnimContexts, RoleIndex))
 			{
-				if (const USkinnedMeshComponent* Mesh = GetMesh(AnimContexts, *RoleIndex))
-				{
-					return Mesh->GetSocketTransform(Schema->GetBoneReferences(Role)[SchemaBoneIdx].BoneName).GetTranslation();
-				}
+				return Mesh->GetSocketTransform(Schema->GetBoneReferences(Role)[SchemaBoneIdx].BoneName).GetTranslation();
 			}
 		}
 	}
@@ -493,28 +482,26 @@ FQuat FDebugDrawParams::ExtractRotation(TConstArrayView<float> PoseVector, float
 			}
 		}
 
-		if (const int32* RoleIndex = RoleToIndex.Find(Role))
+		const int32 RoleIndex = RoleToIndex[Role];
+		const IPoseHistory* PoseHistory = PoseHistories[RoleIndex];
+		check(PoseHistory);
+			
+		if (const USkeleton* Skeleton = Schema->GetSkeleton(Role))
 		{
-			if (const IPoseHistory* PoseHistory = PoseHistories[*RoleIndex])
-			{
-				if (const USkeleton* Skeleton = Schema->GetSkeleton(Role))
-				{
-					const FBoneIndexType BoneIndexType = SchemaBoneIdx == TrajectorySchemaBoneIdx ? ComponentSpaceIndexType : Schema->GetBoneReferences(Role)[SchemaBoneIdx].BoneIndex;
+			const FBoneIndexType BoneIndexType = SchemaBoneIdx == TrajectorySchemaBoneIdx ? ComponentSpaceIndexType : Schema->GetBoneReferences(Role)[SchemaBoneIdx].BoneIndex;
 
-					FTransform WorldBoneTransform;
-					if (PoseHistory->GetTransformAtTime(SampleTimeOffset + PermutationSampleTimeOffset, WorldBoneTransform, Skeleton, BoneIndexType, WorldSpaceIndexType))
-					{
-						return WorldBoneTransform.GetRotation();
-					}
-				}
+			FTransform WorldBoneTransform;
+			if (PoseHistory->GetTransformAtTime(SampleTimeOffset + PermutationSampleTimeOffset, WorldBoneTransform, Skeleton, BoneIndexType, WorldSpaceIndexType))
+			{
+				return WorldBoneTransform.GetRotation();
 			}
+		}
 
-			if (SchemaBoneIdx > RootSchemaBoneIdx)
+		if (SchemaBoneIdx > RootSchemaBoneIdx)
+		{
+			if (const USkinnedMeshComponent* Mesh = GetMesh(AnimContexts, RoleIndex))
 			{
-				if (const USkinnedMeshComponent* Mesh = GetMesh(AnimContexts, *RoleIndex))
-				{
-					return Mesh->GetSocketTransform(Schema->GetBoneReferences(Role)[SchemaBoneIdx].BoneName).GetRotation();
-				}
+				return Mesh->GetSocketTransform(Schema->GetBoneReferences(Role)[SchemaBoneIdx].BoneName).GetRotation();
 			}
 		}
 	}
@@ -525,10 +512,7 @@ FQuat FDebugDrawParams::ExtractRotation(TConstArrayView<float> PoseVector, float
 FTransform FDebugDrawParams::GetRootBoneTransform(const FRole& Role, float SampleTimeOffset) const
 {
 	FTransform RootBoneTransform = FTransform::Identity;
-	if (const int32* RoleIndex = RoleToIndex.Find(Role))
-	{
-		PoseHistories[*RoleIndex]->GetTransformAtTime(SampleTimeOffset, RootBoneTransform, nullptr, RootBoneIndexType, WorldSpaceIndexType);
-	}
+	PoseHistories[RoleToIndex[Role]]->GetTransformAtTime(SampleTimeOffset, RootBoneTransform, nullptr, RootBoneIndexType, WorldSpaceIndexType);
 	return RootBoneTransform;
 }
 
@@ -868,36 +852,30 @@ FSearchContext::FPoseCandidate FSearchContext::FBestPoseCandidates::GetUnsortedC
 // FSearchContext
 FSearchContext::FSearchContext(float InDesiredPermutationTimeOffset, const FPoseIndicesHistory* InPoseIndicesHistory,
 	const FSearchResult& InCurrentResult, const FFloatInterval& InPoseJumpThresholdTime, bool bInUseCachedChannelData)
-PRAGMA_DISABLE_DEPRECATION_WARNINGS
-: FSearchContext(InDesiredPermutationTimeOffset, InPoseIndicesHistory, InCurrentResult, InPoseJumpThresholdTime, FPoseSearchEvent())
-PRAGMA_ENABLE_DEPRECATION_WARNINGS
+: FSearchContext(InDesiredPermutationTimeOffset, InPoseJumpThresholdTime, FPoseSearchEvent())
 {
+	UpdateContinuingPoseSearchResult(InCurrentResult, InCurrentResult);
 }
 
 FSearchContext::FSearchContext(float InDesiredPermutationTimeOffset, const FPoseIndicesHistory* InPoseIndicesHistory,
-	const FFloatInterval& InPoseJumpThresholdTime, const FPoseSearchEvent& InEventToSearch)
-PRAGMA_DISABLE_DEPRECATION_WARNINGS
-: FSearchContext(InDesiredPermutationTimeOffset, InPoseIndicesHistory, FSearchResult(), InPoseJumpThresholdTime, FPoseSearchEvent())
-PRAGMA_ENABLE_DEPRECATION_WARNINGS
+	const FSearchResult& InCurrentResult, const FFloatInterval& InPoseJumpThresholdTime, const FPoseSearchEvent& InEventToSearch)
+: FSearchContext(InDesiredPermutationTimeOffset, InPoseJumpThresholdTime, FPoseSearchEvent())
 {
+	UpdateContinuingPoseSearchResult(InCurrentResult, InCurrentResult);
 }
 
-FSearchContext::FSearchContext(float InDesiredPermutationTimeOffset, const FPoseIndicesHistory* InPoseIndicesHistory,
-	const FSearchResult& InContinuingPoseSearchResult, const FFloatInterval& InPoseJumpThresholdTime, const FPoseSearchEvent& InEventToSearch)
+FSearchContext::FSearchContext(float InDesiredPermutationTimeOffset, const FFloatInterval& InPoseJumpThresholdTime, const FPoseSearchEvent& InEventToSearch)
 : AnimContexts()
 , PoseHistories()
 , RoleToIndex()
 , AssetsToConsider()
 , EventToSearch(InEventToSearch)
 , DesiredPermutationTimeOffset(InDesiredPermutationTimeOffset)
-, PoseIndicesHistory(InPoseIndicesHistory)
 , ContinuingPoseSearchResult(FSearchResult())
 , ContinuingPoseValuesDatabasePoseIdx(FDatabasePoseIdx())
 , PoseJumpThresholdTime(InPoseJumpThresholdTime)
 , bUseCachedChannelData(false)
 {
-	UpdateContinuingPoseSearchResult(InContinuingPoseSearchResult, InContinuingPoseSearchResult);
-
 #if UE_POSE_SEARCH_TRACE_ENABLED
 	BestPoseCandidatesMap.Reserve(16);
 #endif // UE_POSE_SEARCH_TRACE_ENABLED
@@ -906,6 +884,7 @@ FSearchContext::FSearchContext(float InDesiredPermutationTimeOffset, const FPose
 void FSearchContext::AddRole(const FRole& Role, FChooserEvaluationContext* AnimContext, const IPoseHistory* PoseHistory)
 {
 	check(RoleToIndex.Num() == AnimContexts.Num() && RoleToIndex.Num() == PoseHistories.Num());
+	check(AnimContext && PoseHistory);
 
 	AnimContexts.Add(AnimContext);
 	PoseHistories.Add(PoseHistory);
@@ -982,14 +961,9 @@ float FSearchContext::GetSampleCurveValueInternal(float SampleTime, const FName&
 
 	float OutCurveValue = 0.0f;
 	const IPoseHistory* PoseHistory = GetPoseHistory(SampleRole);
-	{
-		if (ensure(PoseHistory))
-		{
-			PoseHistory->GetCurveValueAtTime(SampleTime, CurveName, OutCurveValue);
-			return OutCurveValue;
-		}
-	}
-
+	check(PoseHistory);
+	
+	PoseHistory->GetCurveValueAtTime(SampleTime, CurveName, OutCurveValue);
 	return OutCurveValue;
 }
 
@@ -1049,24 +1023,16 @@ FTransform FSearchContext::GetWorldRootBoneTransformAtTime(float SampleTime, con
 
 	FTransform PoseHistoryTransformAtTime;
 	const IPoseHistory* PoseHistory = GetPoseHistory(SampleRole);
-	#if WITH_EDITOR
-	if (!PoseHistory)
+	check(PoseHistory)
+	if (PoseHistory->GetTransformAtTime(SampleTime, PoseHistoryTransformAtTime, Schema->GetSkeleton(SampleRole), RootBoneIndexType, WorldSpaceIndexType))
 	{
-		UE_LOG(LogPoseSearch, Error, TEXT("FSearchContext::GetWorldRootBoneTransformAtTime - Couldn't search for world space root boneTransform by %s, because no IPoseHistory has been found!"), *Schema->GetName());
+		// PoseHistoryTransformAtTime contains the WorldRootBoneTransform (RootBoneTransform from the blended entry in the pose history times ComponentToWorldTransform from the trajectory)
+		return PoseHistoryTransformAtTime;
 	}
-	else
-	#endif // WITH_EDITOR
-	{
-		if (ensure(PoseHistory) && PoseHistory->GetTransformAtTime(SampleTime, PoseHistoryTransformAtTime, Schema->GetSkeleton(SampleRole), RootBoneIndexType, WorldSpaceIndexType))
-		{
-			// PoseHistoryTransformAtTime contains the WorldRootBoneTransform (RootBoneTransform from the blended entry in the pose history times ComponentToWorldTransform from the trajectory)
-			return PoseHistoryTransformAtTime;
-		}
 
-		// @ todo: should we use the component transform here in case we lack the root bone transform?
-		// NoTe: when GetTransformAtTime returns false PoseHistoryTransformAtTime doesn't contain a valid root bone transform in world space, BUT a valid ComponentToWorldTransform!
-		// So we use PoseHistoryTransformAtTime as such below (in the form of "const FTransform& ComponentToWorldTransform = PoseHistory ? PoseHistoryTransformAtTime : ...)
-	}
+	// @ todo: should we use the component transform here in case we lack the root bone transform?
+	// NoTe: when GetTransformAtTime returns false PoseHistoryTransformAtTime doesn't contain a valid root bone transform in world space, BUT a valid ComponentToWorldTransform!
+	// So we use PoseHistoryTransformAtTime as such below (in the form of "const FTransform& ComponentToWorldTransform = PoseHistory ? PoseHistoryTransformAtTime : ...)
 
 	// accessing AnimInstance or SkeletalMesh properties is not thread safe from a worker thread! Logging an error and returning identity transform
 	
@@ -1099,6 +1065,16 @@ FTransform FSearchContext::GetWorldRootBoneTransformAtTime(float SampleTime, con
 	return FTransform::Identity;
 }
 
+const FPoseIndicesHistory* FSearchContext::GetPoseIndicesHistory() const
+{
+	for (const IPoseHistory* PoseHistory : PoseHistories)
+	{
+		check(PoseHistory);
+		return PoseHistory->GetPoseIndicesHistory();
+	}
+	return nullptr;
+}
+
 bool FSearchContext::ArePoseHistoriesValid() const
 {
 	for (const IPoseHistory* PoseHistory : PoseHistories)
@@ -1113,57 +1089,27 @@ bool FSearchContext::ArePoseHistoriesValid() const
 
 const IPoseHistory* FSearchContext::GetPoseHistory(const FRole& Role) const
 {
-	if (const int32* RoleIndex = RoleToIndex.Find(Role))
-	{
-		return PoseHistories[*RoleIndex];
-	}
-
-	UE_LOG(LogPoseSearch, Error, TEXT("FSearchContext::GetPoseHistory - Role %s could not be found!"), *Role.ToString());
-	return nullptr;
+	return PoseHistories[RoleToIndex[Role]];
 }
 
 const UAnimInstance* FSearchContext::GetAnimInstance(const FRole& Role) const
 {
-	if (const int32* RoleIndex = RoleToIndex.Find(Role))
-	{
-		return Cast<UAnimInstance>(AnimContexts[*RoleIndex]->GetFirstObjectParam());
-	}
-
-	UE_LOG(LogPoseSearch, Error, TEXT("FSearchContext::GetAnimInstance - Role %s could not be found!"), *Role.ToString());
-	return nullptr;
+	return Cast<UAnimInstance>(AnimContexts[RoleToIndex[Role]]->GetFirstObjectParam());
 }
 
 const UObject* FSearchContext::GetAnimContext(const FRole& Role) const
 {
-	if (const int32* RoleIndex = RoleToIndex.Find(Role))
-	{
-		return AnimContexts[*RoleIndex]->GetFirstObjectParam();
-	}
-		
-	UE_LOG(LogPoseSearch, Error, TEXT("FSearchContext::GetAnimContext - Role %s could not be found!"), *Role.ToString());
-	return nullptr;
+	return AnimContexts[RoleToIndex[Role]]->GetFirstObjectParam();
 }
 
 const FChooserEvaluationContext* FSearchContext::GetContext(const FRole& Role) const
 {
-	if (const int32* RoleIndex = RoleToIndex.Find(Role))
-	{
-		return AnimContexts[*RoleIndex];
-	}
-		
-	UE_LOG(LogPoseSearch, Error, TEXT("FSearchContext::GetAnimContext - Role %s could not be found!"), *Role.ToString());
-	return nullptr;	
+	return AnimContexts[RoleToIndex[Role]];
 }
 	
 FChooserEvaluationContext* FSearchContext::GetContext(const FRole& Role)
 {
-	if (const int32* RoleIndex = RoleToIndex.Find(Role))
-	{
-		return AnimContexts[*RoleIndex];
-	}
-		
-	UE_LOG(LogPoseSearch, Error, TEXT("FSearchContext::GetAnimContext - Role %s could not be found!"), *Role.ToString());
-	return nullptr;	
+	return AnimContexts[RoleToIndex[Role]];
 }
 
 FTransform FSearchContext::GetWorldBoneTransformAtTime(float SampleTime, const FRole& SampleRole, int8 SchemaBoneIdx)
@@ -1208,35 +1154,26 @@ FTransform FSearchContext::GetWorldBoneTransformAtTime(float SampleTime, const F
 
 		// collecting the local bone transforms from the IPoseHistory
 		const IPoseHistory* PoseHistory = GetPoseHistory(SampleRole);
-		#if WITH_EDITOR
-		if (!PoseHistory)
-		{
-			UE_LOG(LogPoseSearch, Error, TEXT("FSearchContext::GetWorldBoneTransformAtTime - Couldn't search for bones requested by %s, because no IPoseHistory has been found!"), *Schema->GetName());
-		}
-		else
-		#endif // WITH_EDITOR
+		check(PoseHistory);
+		
+		const USkeleton* Skeleton = Schema->GetSkeleton(SampleRole);
+		FTransform LocalBoneTransform;
+		if (!PoseHistory->GetTransformAtTime(SampleTime, LocalBoneTransform, Skeleton, BoneIndexType, RootBoneIndexType))
 		{
-			check(PoseHistory);
-
-			const USkeleton* Skeleton = Schema->GetSkeleton(SampleRole);
-			FTransform LocalBoneTransform;
-			if (!PoseHistory->GetTransformAtTime(SampleTime, LocalBoneTransform, Skeleton, BoneIndexType, RootBoneIndexType))
+			if (Skeleton)
 			{
-				if (Skeleton)
-				{
-					if (!PoseHistory->IsEmpty())
-					{
-						UE_LOG(LogPoseSearch, Warning, TEXT("FSearchContext::GetWorldBoneTransformAtTime - Couldn't find BoneIndexType %d (%s) for Skeleton %s in the input IPoseHistory requested by %s. Consider adding it to the Pose History!"), BoneIndexType, *Skeleton->GetReferenceSkeleton().GetBoneName(BoneIndexType).ToString(), *Skeleton->GetName(), *Schema->GetName());
-					}
-				}
-				else
+				if (!PoseHistory->IsEmpty())
 				{
-					UE_LOG(LogPoseSearch, Warning, TEXT("FSearchContext::GetWorldBoneTransformAtTime - Schema '%s' Skeleton is not properly set"), *Schema->GetName());
+					UE_LOG(LogPoseSearch, Warning, TEXT("FSearchContext::GetWorldBoneTransformAtTime - Couldn't find BoneIndexType %d (%s) for Skeleton %s in the input IPoseHistory requested by %s. Consider adding it to the Pose History!"), BoneIndexType, *Skeleton->GetReferenceSkeleton().GetBoneName(BoneIndexType).ToString(), *Skeleton->GetName(), *Schema->GetName());
 				}
 			}
-
-			WorldBoneTransform = LocalBoneTransform * WorldBoneTransform;
+			else
+			{
+				UE_LOG(LogPoseSearch, Warning, TEXT("FSearchContext::GetWorldBoneTransformAtTime - Schema '%s' Skeleton is not properly set"), *Schema->GetName());
+			}
 		}
+
+		WorldBoneTransform = LocalBoneTransform * WorldBoneTransform;
 	}
 
 	CachedTransforms.Add(BoneCachedTransformKey) = WorldBoneTransform;
diff --git a/Engine/Plugins/Animation/PoseSearch/Source/Runtime/Private/PoseSearchDatabase.cpp b/Engine/Plugins/Animation/PoseSearch/Source/Runtime/Private/PoseSearchDatabase.cpp
--- a/Engine/Plugins/Animation/PoseSearch/Source/Runtime/Private/PoseSearchDatabase.cpp
+++ b/Engine/Plugins/Animation/PoseSearch/Source/Runtime/Private/PoseSearchDatabase.cpp
@@ -277,7 +277,7 @@ private:
 };
 
 template<bool bReconstructPoseValues, bool bAlignedAndPadded>
-static inline void EvaluatePoseKernel(ISearchResults& SearchResults, const FSearchIndex& SearchIndex, TConstArrayView<float> QueryValues, TArrayView<float> ReconstructedPoseValuesBuffer,
+static inline void EvaluatePoseKernel(FSearchResults& SearchResults, const FSearchIndex& SearchIndex, TConstArrayView<float> QueryValues, TArrayView<float> ReconstructedPoseValuesBuffer,
 	int32 PoseIdx, int32 EventPoseIdx, const FSearchFilters& SearchFilters, float ContinuingPoseCostAddend, float ContinuingInteractionCostAddend,
 	FSearchContext& SearchContext, const UPoseSearchDatabase* Database, TConstArrayView<float> DynamicWeightsSqrt
 #if UE_POSE_SEARCH_TRACE_ENABLED
@@ -1533,7 +1533,7 @@ float UPoseSearchDatabase::GetNormalizedAssetTime(int32 PoseIdx) const
 	return NormalizedAssetTime;
 }
 
-void UPoseSearchDatabase::Search(UE::PoseSearch::FSearchContext& SearchContext, UE::PoseSearch::ISearchResults& SearchResults) const
+void UPoseSearchDatabase::Search(UE::PoseSearch::FSearchContext& SearchContext, UE::PoseSearch::FSearchResults& SearchResults) const
 {
 	using namespace UE::PoseSearch;
 
@@ -1662,26 +1662,31 @@ void UPoseSearchDatabase::PopulateNonSelectableIdx(FNonSelectableIdx& NonSelecta
 		}
 	}
 
-	if (SearchContext.GetPoseIndicesHistory())
+	for (const IPoseHistory* PoseHistory : SearchContext.GetPoseHistories())
 	{
-		const FObjectKey DatabaseKey(this);
-		for (auto It = SearchContext.GetPoseIndicesHistory()->IndexToTime.CreateConstIterator(); It; ++It)
+		check(PoseHistory);
+		
+		if (const FPoseIndicesHistory* PoseIndicesHistory = PoseHistory->GetPoseIndicesHistory())
 		{
-			const FHistoricalPoseIndex& HistoricalPoseIndex = It.Key();
-			if (HistoricalPoseIndex.DatabaseKey == DatabaseKey)
+			const FObjectKey DatabaseKey(this);
+			for (auto It = PoseIndicesHistory->IndexToTime.CreateConstIterator(); It; ++It)
 			{
-				NonSelectableIdx.AddUnique(HistoricalPoseIndex.PoseIndex);
+				const FHistoricalPoseIndex& HistoricalPoseIndex = It.Key();
+				if (HistoricalPoseIndex.DatabaseKey == DatabaseKey)
+				{
+					NonSelectableIdx.AddUnique(HistoricalPoseIndex.PoseIndex);
 
 #if UE_POSE_SEARCH_TRACE_ENABLED
-				check(HistoricalPoseIndex.PoseIndex >= 0);
+					check(HistoricalPoseIndex.PoseIndex >= 0);
 
-				// if we're editing the database and removing assets it's possible that the PoseIndicesHistory contains invalid pose indexes
-				if (HistoricalPoseIndex.PoseIndex < SearchIndex.GetNumPoses())
-				{
-					const FPoseSearchCost PoseCost(CompareFeatureVectors<false>(SearchIndex.GetPoseValuesSafe(HistoricalPoseIndex.PoseIndex, BufferUsedForReconstruction), QueryValues, DynamicWeightsSqrt), SearchIndex.PoseMetadata[HistoricalPoseIndex.PoseIndex].GetCostAddend(), ContinuingPoseCostAddend, ContinuingInteractionCostAddend);
-					SearchContext.Track(this, HistoricalPoseIndex.PoseIndex, EPoseCandidateFlags::DiscardedBy_PoseReselectHistory, PoseCost);
-				}
+					// if we're editing the database and removing assets it's possible that the PoseIndicesHistory contains invalid pose indexes
+					if (HistoricalPoseIndex.PoseIndex < SearchIndex.GetNumPoses())
+					{
+						const FPoseSearchCost PoseCost(CompareFeatureVectors<false>(SearchIndex.GetPoseValuesSafe(HistoricalPoseIndex.PoseIndex, BufferUsedForReconstruction), QueryValues, DynamicWeightsSqrt), SearchIndex.PoseMetadata[HistoricalPoseIndex.PoseIndex].GetCostAddend(), ContinuingPoseCostAddend, ContinuingInteractionCostAddend);
+						SearchContext.Track(this, HistoricalPoseIndex.PoseIndex, EPoseCandidateFlags::DiscardedBy_PoseReselectHistory, PoseCost);
+					}
 #endif // UE_POSE_SEARCH_TRACE_ENABLED
+				}
 			}
 		}
 	}
@@ -1731,7 +1736,7 @@ UE::PoseSearch::FSearchResult UPoseSearchDatabase::SearchContinuingPose(UE::Pose
 	return SearchResults.GetBestResult();
 }
 
-void UPoseSearchDatabase::SearchContinuingPose(UE::PoseSearch::FSearchContext& SearchContext, UE::PoseSearch::ISearchResults& SearchResults) const
+void UPoseSearchDatabase::SearchContinuingPose(UE::PoseSearch::FSearchContext& SearchContext, UE::PoseSearch::FSearchResults& SearchResults) const
 {
 	QUICK_SCOPE_CYCLE_COUNTER(STAT_PoseSearch_ContinuingPose);
 
@@ -1842,7 +1847,7 @@ void UPoseSearchDatabase::SearchContinuingPose(UE::PoseSearch::FSearchContext& S
 	// NoTe: no need to SearchResults.FinalizeResults(), since SearchContinuingPose already calls SearchResults.UpdateWith with a Result with AssetTime already initialized
 }
 
-void UPoseSearchDatabase::SearchPCAKDTree(UE::PoseSearch::FSearchContext& SearchContext, UE::PoseSearch::ISearchResults& SearchResults) const
+void UPoseSearchDatabase::SearchPCAKDTree(UE::PoseSearch::FSearchContext& SearchContext, UE::PoseSearch::FSearchResults& SearchResults) const
 {
 	SCOPE_CYCLE_COUNTER(STAT_PoseSearch_PCAKNN);
 
@@ -2138,7 +2143,7 @@ void UPoseSearchDatabase::SearchPCAKDTree(UE::PoseSearch::FSearchContext& Search
 #endif // UE_POSE_SEARCH_TRACE_ENABLED
 }
 
-void UPoseSearchDatabase::SearchVPTree(UE::PoseSearch::FSearchContext& SearchContext, UE::PoseSearch::ISearchResults& SearchResults) const
+void UPoseSearchDatabase::SearchVPTree(UE::PoseSearch::FSearchContext& SearchContext, UE::PoseSearch::FSearchResults& SearchResults) const
 {
 	SCOPE_CYCLE_COUNTER(STAT_PoseSearch_VPTree);
 
@@ -2255,7 +2260,7 @@ void UPoseSearchDatabase::SearchVPTree(UE::PoseSearch::FSearchContext& SearchCon
 #endif // UE_POSE_SEARCH_TRACE_ENABLED
 }
 
-void UPoseSearchDatabase::SearchBruteForce(UE::PoseSearch::FSearchContext& SearchContext, UE::PoseSearch::ISearchResults& SearchResults) const
+void UPoseSearchDatabase::SearchBruteForce(UE::PoseSearch::FSearchContext& SearchContext, UE::PoseSearch::FSearchResults& SearchResults) const
 {
 	SCOPE_CYCLE_COUNTER(STAT_PoseSearch_BruteForce);
 	
@@ -2439,7 +2444,7 @@ void UPoseSearchDatabase::SearchBruteForce(UE::PoseSearch::FSearchContext& Searc
 #endif // UE_POSE_SEARCH_TRACE_ENABLED
 }
 
-void UPoseSearchDatabase::SearchEvent(UE::PoseSearch::FSearchContext& SearchContext, UE::PoseSearch::ISearchResults& SearchResults) const
+void UPoseSearchDatabase::SearchEvent(UE::PoseSearch::FSearchContext& SearchContext, UE::PoseSearch::FSearchResults& SearchResults) const
 {
 	SCOPE_CYCLE_COUNTER(STAT_PoseSearch_Event);
 
diff --git a/Engine/Plugins/Animation/PoseSearch/Source/Runtime/Private/PoseSearchFeatureChannel_Curve.cpp b/Engine/Plugins/Animation/PoseSearch/Source/Runtime/Private/PoseSearchFeatureChannel_Curve.cpp
--- a/Engine/Plugins/Animation/PoseSearch/Source/Runtime/Private/PoseSearchFeatureChannel_Curve.cpp
+++ b/Engine/Plugins/Animation/PoseSearch/Source/Runtime/Private/PoseSearchFeatureChannel_Curve.cpp
@@ -86,16 +86,12 @@ void UPoseSearchFeatureChannel_Curve::BuildQuery(UE::PoseSearch::FSearchContext&
 
 	const bool bCanUseContinuingPoseValues = SearchContext.CanUseContinuingPoseValues();
 	const bool bSkip = InputQueryPose != EInputQueryPose::UseCharacterPose && bCanUseContinuingPoseValues;
-	if (bSkip || (!SearchContext.ArePoseHistoriesValid()))
+	if (bSkip)
 	{
 		if (bCanUseContinuingPoseValues)
 		{
 			FFeatureVectorHelper::Copy(SearchContext.EditFeatureVector(), ChannelDataOffset, ChannelCardinality, SearchContext.GetContinuingPoseValues());
-			return;
 		}
-
-		// we leave the SearchContext.EditFeatureVector() set to zero since the SearchContext.PoseHistory is invalid and it'll fail if we continue
-		UE_LOG(LogPoseSearch, Error, TEXT("UPoseSearchFeatureChannel_Curve::BuildQuery - Failed because Pose History Node is missing."));
 		return;
 	}
 	
diff --git a/Engine/Plugins/Animation/PoseSearch/Source/Runtime/Private/PoseSearchFeatureChannel_Distance.cpp b/Engine/Plugins/Animation/PoseSearch/Source/Runtime/Private/PoseSearchFeatureChannel_Distance.cpp
--- a/Engine/Plugins/Animation/PoseSearch/Source/Runtime/Private/PoseSearchFeatureChannel_Distance.cpp
+++ b/Engine/Plugins/Animation/PoseSearch/Source/Runtime/Private/PoseSearchFeatureChannel_Distance.cpp
@@ -129,17 +129,12 @@ void UPoseSearchFeatureChannel_Distance::BuildQuery(UE::PoseSearch::FSearchConte
 
 	const bool bCanUseContinuingPoseValues = SearchContext.CanUseContinuingPoseValues();
 	const bool bSkip = InputQueryPose != EInputQueryPose::UseCharacterPose && bCanUseContinuingPoseValues && SampleRole == OriginRole;
-	const bool bIsRootBone = SchemaBoneIdx == RootSchemaBoneIdx;
-	if (bSkip || (!SearchContext.ArePoseHistoriesValid() && !bIsRootBone))
+	if (bSkip)
 	{
 		if (bCanUseContinuingPoseValues)
 		{
 			FFeatureVectorHelper::Copy(SearchContext.EditFeatureVector(), ChannelDataOffset, ChannelCardinality, SearchContext.GetContinuingPoseValues());
-			return;
 		}
-
-		// we leave the SearchContext.EditFeatureVector() set to zero since the SearchContext.PoseHistory is invalid and it'll fail if we continue
-		UE_LOG(LogPoseSearch, Error, TEXT("UPoseSearchFeatureChannel_Distance::BuildQuery - Failed because Pose History Node is missing."));
 		return;
 	}
 	
diff --git a/Engine/Plugins/Animation/PoseSearch/Source/Runtime/Private/PoseSearchFeatureChannel_FilterCrashingLegs.cpp b/Engine/Plugins/Animation/PoseSearch/Source/Runtime/Private/PoseSearchFeatureChannel_FilterCrashingLegs.cpp
--- a/Engine/Plugins/Animation/PoseSearch/Source/Runtime/Private/PoseSearchFeatureChannel_FilterCrashingLegs.cpp
+++ b/Engine/Plugins/Animation/PoseSearch/Source/Runtime/Private/PoseSearchFeatureChannel_FilterCrashingLegs.cpp
@@ -99,16 +99,12 @@ void UPoseSearchFeatureChannel_FilterCrashingLegs::BuildQuery(UE::PoseSearch::FS
 	// trying to get the BuildQuery data from the continuing pose
 	const bool bCanUseContinuingPoseValues = SearchContext.CanUseContinuingPoseValues();
 	const bool bSkip = InputQueryPose != EInputQueryPose::UseCharacterPose && bCanUseContinuingPoseValues;
-	if (bSkip || !SearchContext.ArePoseHistoriesValid())
+	if (bSkip)
 	{
 		if (bCanUseContinuingPoseValues)
 		{
 			FFeatureVectorHelper::Copy(SearchContext.EditFeatureVector(), ChannelDataOffset, ChannelCardinality, SearchContext.GetContinuingPoseValues());
-			return;
 		}
-
-		// we leave the SearchContext.EditFeatureVector() set to zero since the SearchContext.PoseHistory is invalid and it'll fail if we continue
-		UE_LOG(LogPoseSearch, Error, TEXT("UPoseSearchFeatureChannel_FilterCrashingLegs::BuildQuery - Failed because Pose History Node is missing."));
 		return;
 	}
 
diff --git a/Engine/Plugins/Animation/PoseSearch/Source/Runtime/Private/PoseSearchFeatureChannel_Heading.cpp b/Engine/Plugins/Animation/PoseSearch/Source/Runtime/Private/PoseSearchFeatureChannel_Heading.cpp
--- a/Engine/Plugins/Animation/PoseSearch/Source/Runtime/Private/PoseSearchFeatureChannel_Heading.cpp
+++ b/Engine/Plugins/Animation/PoseSearch/Source/Runtime/Private/PoseSearchFeatureChannel_Heading.cpp
@@ -106,7 +106,6 @@ void UPoseSearchFeatureChannel_Heading::BuildQuery(UE::PoseSearch::FSearchContex
 {
 	using namespace UE::PoseSearch;
 
-	const bool bIsRootBone = SchemaBoneIdx == RootSchemaBoneIdx;
 	if (bUseBlueprintQueryOverride)
 	{
 		if (const UAnimInstance* AnimInstance = Cast<UAnimInstance>(SearchContext.GetContext(SampleRole)->GetFirstObjectParam()))
@@ -171,16 +170,12 @@ void UPoseSearchFeatureChannel_Heading::BuildQuery(UE::PoseSearch::FSearchContex
 
 	const bool bCanUseContinuingPoseValues = SearchContext.CanUseContinuingPoseValues();
 	const bool bSkip = InputQueryPose != EInputQueryPose::UseCharacterPose && bCanUseContinuingPoseValues && SampleRole == OriginRole;
-	if (bSkip || (!SearchContext.ArePoseHistoriesValid() && !bIsRootBone))
+	if (bSkip)
 	{
 		if (bCanUseContinuingPoseValues)
 		{
 			FFeatureVectorHelper::Copy(SearchContext.EditFeatureVector(), ChannelDataOffset, ChannelCardinality, SearchContext.GetContinuingPoseValues());
-			return;
 		}
-
-		// we leave the SearchContext.EditFeatureVector() set to zero since the SearchContext.PoseHistory is invalid and it'll fail if we continue
-		UE_LOG(LogPoseSearch, Error, TEXT("UPoseSearchFeatureChannel_Heading::BuildQuery - Failed because Pose History Node is missing."));
 		return;
 	}
 	
diff --git a/Engine/Plugins/Animation/PoseSearch/Source/Runtime/Private/PoseSearchFeatureChannel_Phase.cpp b/Engine/Plugins/Animation/PoseSearch/Source/Runtime/Private/PoseSearchFeatureChannel_Phase.cpp
--- a/Engine/Plugins/Animation/PoseSearch/Source/Runtime/Private/PoseSearchFeatureChannel_Phase.cpp
+++ b/Engine/Plugins/Animation/PoseSearch/Source/Runtime/Private/PoseSearchFeatureChannel_Phase.cpp
@@ -341,17 +341,12 @@ void UPoseSearchFeatureChannel_Phase::BuildQuery(UE::PoseSearch::FSearchContext&
 
 	const bool bCanUseContinuingPoseValues = SearchContext.CanUseContinuingPoseValues();
 	const bool bSkip = InputQueryPose != EInputQueryPose::UseCharacterPose && bCanUseContinuingPoseValues;
-	if (bSkip || !SearchContext.ArePoseHistoriesValid())
+	if (bSkip)
 	{
 		if (bCanUseContinuingPoseValues)
 		{
 			FFeatureVectorHelper::Copy(SearchContext.EditFeatureVector(), ChannelDataOffset, ChannelCardinality, SearchContext.GetContinuingPoseValues());
-			return;
 		}
-		
-		// we leave the SearchContext.EditFeatureVector() set to zero since the SearchContext.PoseHistory is invalid and it'll fail if we continue
-		UE_LOG(LogPoseSearch, Error, TEXT("UPoseSearchFeatureChannel_Phase::BuildQuery - Failed because Pose History Node is missing."));
-		return;
 	}
 	
 	// @todo: Support phase in BuildQuery
diff --git a/Engine/Plugins/Animation/PoseSearch/Source/Runtime/Private/PoseSearchFeatureChannel_Position.cpp b/Engine/Plugins/Animation/PoseSearch/Source/Runtime/Private/PoseSearchFeatureChannel_Position.cpp
--- a/Engine/Plugins/Animation/PoseSearch/Source/Runtime/Private/PoseSearchFeatureChannel_Position.cpp
+++ b/Engine/Plugins/Animation/PoseSearch/Source/Runtime/Private/PoseSearchFeatureChannel_Position.cpp
@@ -173,17 +173,12 @@ void UPoseSearchFeatureChannel_Position::BuildQuery(UE::PoseSearch::FSearchConte
 
 	const bool bCanUseContinuingPoseValues = SearchContext.CanUseContinuingPoseValues();
 	const bool bSkip = InputQueryPose != EInputQueryPose::UseCharacterPose && bCanUseContinuingPoseValues && SampleRole == OriginRole;
-	const bool bIsRootBone = SchemaBoneIdx == RootSchemaBoneIdx;
-	if (bSkip || (!SearchContext.ArePoseHistoriesValid() && !bIsRootBone))
+	if (bSkip )
 	{
 		if (bCanUseContinuingPoseValues)
 		{
 			FFeatureVectorHelper::Copy(SearchContext.EditFeatureVector(), ChannelDataOffset, ChannelCardinality, SearchContext.GetContinuingPoseValues());
-			return;
 		}
-
-		// we leave the SearchContext.EditFeatureVector() set to zero since the SearchContext.PoseHistory is invalid and it'll fail if we continue
-		UE_LOG(LogPoseSearch, Error, TEXT("UPoseSearchFeatureChannel_Position::BuildQuery - Failed because Pose History Node is missing."));
 		return;
 	}
 	
diff --git a/Engine/Plugins/Animation/PoseSearch/Source/Runtime/Private/PoseSearchFeatureChannel_Velocity.cpp b/Engine/Plugins/Animation/PoseSearch/Source/Runtime/Private/PoseSearchFeatureChannel_Velocity.cpp
--- a/Engine/Plugins/Animation/PoseSearch/Source/Runtime/Private/PoseSearchFeatureChannel_Velocity.cpp
+++ b/Engine/Plugins/Animation/PoseSearch/Source/Runtime/Private/PoseSearchFeatureChannel_Velocity.cpp
@@ -48,7 +48,6 @@ void UPoseSearchFeatureChannel_Velocity::BuildQuery(UE::PoseSearch::FSearchConte
 {
 	using namespace UE::PoseSearch;
 
-	const bool bIsRootBone = SchemaBoneIdx == RootSchemaBoneIdx;
 	if (bUseBlueprintQueryOverride)
 	{
 		if (const UAnimInstance* AnimInstance = Cast<UAnimInstance>(SearchContext.GetContext(SampleRole)->GetFirstObjectParam()))
@@ -119,16 +118,12 @@ void UPoseSearchFeatureChannel_Velocity::BuildQuery(UE::PoseSearch::FSearchConte
 
 	const bool bCanUseContinuingPoseValues = SearchContext.CanUseContinuingPoseValues();
 	const bool bSkip = InputQueryPose != EInputQueryPose::UseCharacterPose && bCanUseContinuingPoseValues && SampleRole == OriginRole;
-	if (bSkip || (!SearchContext.ArePoseHistoriesValid() && !bIsRootBone))
+	if (bSkip)
 	{
 		if (bCanUseContinuingPoseValues)
 		{
 			FFeatureVectorHelper::Copy(SearchContext.EditFeatureVector(), ChannelDataOffset, ChannelCardinality, SearchContext.GetContinuingPoseValues());
-			return;
 		}
-
-		// we leave the SearchContext.EditFeatureVector() set to zero since the SearchContext.PoseHistory is invalid and it'll fail if we continue
-		UE_LOG(LogPoseSearch, Error, TEXT("UPoseSearchFeatureChannel_Velocity::BuildQuery - Failed because Pose History Node is missing."));
 		return;
 	}
 	
diff --git a/Engine/Plugins/Animation/PoseSearch/Source/Runtime/Private/PoseSearchHistory.cpp b/Engine/Plugins/Animation/PoseSearch/Source/Runtime/Private/PoseSearchHistory.cpp
--- a/Engine/Plugins/Animation/PoseSearch/Source/Runtime/Private/PoseSearchHistory.cpp
+++ b/Engine/Plugins/Animation/PoseSearch/Source/Runtime/Private/PoseSearchHistory.cpp
@@ -475,24 +475,23 @@ void IPoseHistory::DebugDraw(FAnimInstanceProxy& AnimInstanceProxy, FColor Color
 // FArchivedPoseHistory
 void FArchivedPoseHistory::InitFrom(const IPoseHistory* PoseHistory)
 {
+	check(PoseHistory);
+
 	Trajectory.Samples.Reset();
 	BoneToTransformMap.Reset();
 	Entries.Reset();
+	
+	Trajectory = PoseHistory->GetTrajectory();
+	BoneToTransformMap = PoseHistory->GetBoneToTransformMap();
+	CollectedCurves = PoseHistory->GetCollectedCurves();
+	const int32 NumEntries = PoseHistory->GetNumEntries();
+	Entries.SetNum(NumEntries);
 
-	if (PoseHistory)
+	for (int32 EntryIndex = 0; EntryIndex < NumEntries; ++EntryIndex)
 	{
-		Trajectory = PoseHistory->GetTrajectory();
-		BoneToTransformMap = PoseHistory->GetBoneToTransformMap();
-		CollectedCurves = PoseHistory->GetCollectedCurves();
-		const int32 NumEntries = PoseHistory->GetNumEntries();
-		Entries.SetNum(NumEntries);
-
-		for (int32 EntryIndex = 0; EntryIndex < NumEntries; ++EntryIndex)
-		{
-			Entries[EntryIndex] = PoseHistory->GetEntry(EntryIndex);
-			// validating input PoseHistory to have Entries properly sorted by time
-			check(EntryIndex == 0 || (Entries[EntryIndex - 1].AccumulatedSeconds <= Entries[EntryIndex].AccumulatedSeconds));
-		}
+		Entries[EntryIndex] = PoseHistory->GetEntry(EntryIndex);
+		// validating input PoseHistory to have Entries properly sorted by time
+		check(EntryIndex == 0 || (Entries[EntryIndex - 1].AccumulatedSeconds <= Entries[EntryIndex].AccumulatedSeconds));
 	}
 }
 
@@ -687,7 +686,12 @@ FPoseHistory& FPoseHistory::operator=(const FPoseHistory& Other)
 		TrajectoryDataState = Other.TrajectoryDataState;
 		TrajectorySpeedMultiplier = Other.TrajectorySpeedMultiplier;
 
-		PoseData = Other.PoseData;
+		LastUpdateSkeleton = Other.LastUpdateSkeleton;
+		BoneToTransformMap = Other.BoneToTransformMap;
+		CollectedCurves = Other.CollectedCurves;
+		BoneToTransformMapTypeHash  = Other.BoneToTransformMapTypeHash;
+		Entries = Other.Entries;
+		PoseIndicesHistory = Other.PoseIndicesHistory;
 	}
 	return *this;
 }
@@ -708,7 +712,12 @@ FPoseHistory& FPoseHistory::operator=(FPoseHistory&& Other)
 		TrajectoryDataState = MoveTemp(Other.TrajectoryDataState);
 		TrajectorySpeedMultiplier = MoveTemp(Other.TrajectorySpeedMultiplier);
 
-		PoseData = MoveTemp(Other.PoseData);
+		LastUpdateSkeleton = MoveTemp(Other.LastUpdateSkeleton);
+		BoneToTransformMap = MoveTemp(Other.BoneToTransformMap);
+		CollectedCurves = MoveTemp(Other.CollectedCurves);
+		BoneToTransformMapTypeHash  = MoveTemp(Other.BoneToTransformMapTypeHash);
+		Entries = MoveTemp(Other.Entries);
+		PoseIndicesHistory = MoveTemp(Other.PoseIndicesHistory);
 	}
 	return *this;
 }
@@ -725,7 +734,12 @@ void FPoseHistory::Initialize_AnyThread(int32 InNumPoses, float InSamplingInterv
 	TrajectoryDataState = FPoseSearchTrajectoryData::FState();
 	TrajectorySpeedMultiplier = 1.f;
 
-	PoseData = FPoseData();
+	LastUpdateSkeleton = nullptr;
+	BoneToTransformMap.Reset();
+	CollectedCurves.Reset();
+	BoneToTransformMapTypeHash = 0;
+	Entries.Reset();
+	PoseIndicesHistory.Reset();
 }
 
 bool FPoseHistory::GetTransformAtTime(float Time, FTransform& OutBoneTransform, const USkeleton* BoneIndexSkeleton, FBoneIndexType BoneIndexType, FBoneIndexType ReferenceBoneIndexType, bool bExtrapolate) const
@@ -763,7 +777,7 @@ bool FPoseHistory::GetTransformAtTime(float Time, FTransform& OutBoneTransform,
 	{
 		check(BoneIndexType != ComponentSpaceIndexType && BoneIndexType != WorldSpaceIndexType && ReferenceBoneIndexType != WorldSpaceIndexType);
 
-		const int32 NumEntries = PoseData.Entries.Num();
+		const int32 NumEntries = Entries.Num();
 		if (NumEntries > 0)
 		{
 			int32 NextIdx = 0;
@@ -771,15 +785,15 @@ bool FPoseHistory::GetTransformAtTime(float Time, FTransform& OutBoneTransform,
 
 			if (NumEntries > 1)
 			{
-				const int32 LowerBoundIdx = LowerBound(PoseData.Entries.begin(), PoseData.Entries.end(), Time, [](const FPoseHistoryEntry& Entry, float Value) { return Value > Entry.AccumulatedSeconds; });
+				const int32 LowerBoundIdx = LowerBound(Entries.begin(), Entries.end(), Time, [](const FPoseHistoryEntry& Entry, float Value) { return Value > Entry.AccumulatedSeconds; });
 				NextIdx = FMath::Clamp(LowerBoundIdx, 1, NumEntries - 1);
 				PrevIdx = NextIdx - 1;
 			}
 
-			const FPoseHistoryEntry& PrevEntry = PoseData.Entries[PrevIdx];
-			const FPoseHistoryEntry& NextEntry = PoseData.Entries[NextIdx];
+			const FPoseHistoryEntry& PrevEntry = Entries[PrevIdx];
+			const FPoseHistoryEntry& NextEntry = Entries[NextIdx];
 
-			bSuccess = LerpEntries(Time, bExtrapolate, PrevEntry, NextEntry, BoneIndexSkeleton, PoseData.LastUpdateSkeleton.Get(), PoseData.BoneToTransformMap, BoneIndexType, ReferenceBoneIndexType, OutBoneTransform);
+			bSuccess = LerpEntries(Time, bExtrapolate, PrevEntry, NextEntry, BoneIndexSkeleton, LastUpdateSkeleton.Get(), BoneToTransformMap, BoneIndexType, ReferenceBoneIndexType, OutBoneTransform);
 		}
 		else
 		{
@@ -805,7 +819,7 @@ bool FPoseHistory::GetCurveValueAtTime(float Time, const FName& CurveName, float
 {
 	UE_MT_SCOPED_READ_ACCESS(PoseDataThreadSafeCounter);
 
-	const int32 NumEntries = PoseData.Entries.Num();
+	const int32 NumEntries = Entries.Num();
 	if (NumEntries > 0)
 	{
 		int32 NextIdx = 0;
@@ -813,15 +827,15 @@ bool FPoseHistory::GetCurveValueAtTime(float Time, const FName& CurveName, float
 
 		if (NumEntries > 1)
 		{
-			const int32 LowerBoundIdx = LowerBound(PoseData.Entries.begin(), PoseData.Entries.end(), Time, [](const FPoseHistoryEntry& Entry, float Value) { return Value > Entry.AccumulatedSeconds; });
+			const int32 LowerBoundIdx = LowerBound(Entries.begin(), Entries.end(), Time, [](const FPoseHistoryEntry& Entry, float Value) { return Value > Entry.AccumulatedSeconds; });
 			NextIdx = FMath::Clamp(LowerBoundIdx, 1, NumEntries - 1);
 			PrevIdx = NextIdx - 1;
 		}
 
-		const FPoseHistoryEntry& PrevEntry = PoseData.Entries[PrevIdx];
-		const FPoseHistoryEntry& NextEntry = PoseData.Entries[NextIdx];
+		const FPoseHistoryEntry& PrevEntry = Entries[PrevIdx];
+		const FPoseHistoryEntry& NextEntry = Entries[NextIdx];
 
-		return LerpEntries(Time, bExtrapolate, PrevEntry, NextEntry, CurveName, PoseData.CollectedCurves, OutCurveValue);
+		return LerpEntries(Time, bExtrapolate, PrevEntry, NextEntry, CurveName, CollectedCurves, OutCurveValue);
 	}
 	
 	OutCurveValue = 0.0f;
@@ -843,19 +857,19 @@ float FPoseHistory::GetTrajectorySpeedMultiplier() const
 bool FPoseHistory::IsEmpty() const
 {
 	UE_MT_SCOPED_READ_ACCESS(PoseDataThreadSafeCounter);
-	return PoseData.Entries.IsEmpty();
+	return Entries.IsEmpty();
 }
 
 const FBoneToTransformMap& FPoseHistory::GetBoneToTransformMap() const
 {
 	UE_MT_SCOPED_READ_ACCESS(PoseDataThreadSafeCounter);
-	return PoseData.BoneToTransformMap;
+	return BoneToTransformMap;
 }
 
 const TConstArrayView<FName> FPoseHistory::GetCollectedCurves() const
 {
 	UE_MT_SCOPED_READ_ACCESS(PoseDataThreadSafeCounter);
-	return MakeConstArrayView(PoseData.CollectedCurves);
+	return MakeConstArrayView(CollectedCurves);
 }
 
 PRAGMA_DISABLE_DEPRECATION_WARNINGS
@@ -868,13 +882,19 @@ PRAGMA_ENABLE_DEPRECATION_WARNINGS
 int32 FPoseHistory::GetNumEntries() const
 {
 	UE_MT_SCOPED_READ_ACCESS(PoseDataThreadSafeCounter);
-	return PoseData.Entries.Num();
+	return Entries.Num();
 }
 
 const FPoseHistoryEntry& FPoseHistory::GetEntry(int32 EntryIndex) const
 {
 	UE_MT_SCOPED_READ_ACCESS(PoseDataThreadSafeCounter);
-	return PoseData.Entries[EntryIndex];
+	return Entries[EntryIndex];
+}
+
+const FPoseIndicesHistory* FPoseHistory::GetPoseIndicesHistory() const
+{
+	UE_MT_SCOPED_READ_ACCESS(PoseDataThreadSafeCounter);
+	return &PoseIndicesHistory;
 }
 
 void FPoseHistory::GenerateTrajectory(const UObject* AnimContext, float DeltaTime, const FPoseSearchTrajectoryData& TrajectoryData, const FPoseSearchTrajectoryData::FSampling& TrajectoryDataSampling)
@@ -922,7 +942,7 @@ void FPoseHistory::SetTrajectory(const FTransformTrajectory& InTrajectory, float
 void FPoseHistory::EvaluateComponentSpace_AnyThread(float DeltaTime, IComponentSpacePoseProvider& ComponentSpacePoseProvider, bool bStoreScales,
 	float RootBoneRecoveryTime, float RootBoneTranslationRecoveryRatio, float RootBoneRotationRecoveryRatio,
 	bool bNeedsReset, bool bCacheBones, const TArray<FBoneIndexType>& RequiredBones,
-	const FBlendedCurve& Curves, const TConstArrayView<FName>& CollectedCurves)
+	const FBlendedCurve& Curves, const TConstArrayView<FName>& InCollectedCurves)
 {
 	UE_MT_SCOPED_WRITE_ACCESS(PoseDataThreadSafeCounter);
 
@@ -932,72 +952,72 @@ void FPoseHistory::EvaluateComponentSpace_AnyThread(float DeltaTime, IComponentS
 
 	if (bCacheBones)
 	{
-		const uint32 OldBoneToTransformMapTypeHash = PoseData.BoneToTransformMapTypeHash;
+		const uint32 OldBoneToTransformMapTypeHash = BoneToTransformMapTypeHash;
 
-		PoseData.BoneToTransformMap.Reset();
-		PoseData.CollectedCurves = CollectedCurves;
+		BoneToTransformMap.Reset();
+		CollectedCurves = InCollectedCurves;
 		if (!RequiredBones.IsEmpty())
 		{
 			// making sure we always collect the root bone transform (by construction BoneToTransformMap[0] = 0)
 			const FComponentSpaceTransformIndex ComponentSpaceTransformRootBoneIndex = 0;
-			PoseData.BoneToTransformMap.Add(RootBoneIndexType) = ComponentSpaceTransformRootBoneIndex;
+			BoneToTransformMap.Add(RootBoneIndexType) = ComponentSpaceTransformRootBoneIndex;
 
 			for (int32 i = 0; i < RequiredBones.Num(); ++i)
 			{
 				// adding only unique RequiredBones to avoid oversizing Entries::ComponentSpaceTransforms
-				if (!PoseData.BoneToTransformMap.Find(RequiredBones[i]))
+				if (!BoneToTransformMap.Find(RequiredBones[i]))
 				{
-					const FComponentSpaceTransformIndex ComponentSpaceTransformIndex = PoseData.BoneToTransformMap.Num();
-					PoseData.BoneToTransformMap.Add(RequiredBones[i]) = ComponentSpaceTransformIndex;
+					const FComponentSpaceTransformIndex ComponentSpaceTransformIndex = BoneToTransformMap.Num();
+					BoneToTransformMap.Add(RequiredBones[i]) = ComponentSpaceTransformIndex;
 				}
 			}
 		}
 
-		PoseData.BoneToTransformMapTypeHash = GetTypeHash(PoseData.BoneToTransformMap);
-		bNeedsReset |= (OldBoneToTransformMapTypeHash != PoseData.BoneToTransformMapTypeHash);
+		BoneToTransformMapTypeHash = GetTypeHash(BoneToTransformMap);
+		bNeedsReset |= (OldBoneToTransformMapTypeHash != BoneToTransformMapTypeHash);
 	}
 
-	if (PoseData.LastUpdateSkeleton != Skeleton)
+	if (LastUpdateSkeleton != Skeleton)
 	{
 		bNeedsReset = true;
-		PoseData.LastUpdateSkeleton = Skeleton;
+		LastUpdateSkeleton = Skeleton;
 	}
 
 	if (bNeedsReset)
 	{
-		PoseData.Entries.Reset();
-		PoseData.Entries.Reserve(MaxNumPoses);
+		Entries.Reset();
+		Entries.Reserve(MaxNumPoses);
 	}
 
 	FPoseHistoryEntry FutureEntryTemp;
-	if (!PoseData.Entries.IsEmpty() && PoseData.Entries.Last().AccumulatedSeconds > 0.f)
+	if (!Entries.IsEmpty() && Entries.Last().AccumulatedSeconds > 0.f)
 	{
 		// removing the "future" root bone Entry
-		FutureEntryTemp = MoveTemp(PoseData.Entries.Last());
-		PoseData.Entries.Pop();
+		FutureEntryTemp = MoveTemp(Entries.Last());
+		Entries.Pop();
 	}
 
 	// Age our elapsed times
-	for (FPoseHistoryEntry& Entry : PoseData.Entries)
+	for (FPoseHistoryEntry& Entry : Entries)
 	{
 		Entry.AccumulatedSeconds -= DeltaTime;
 	}
 
-	if (PoseData.Entries.Num() != MaxNumPoses)
+	if (Entries.Num() != MaxNumPoses)
 	{
 		// Consume every pose until the queue is full
-		PoseData.Entries.Emplace();
+		Entries.Emplace();
 	}
-	else if (SamplingInterval <= 0.f || PoseData.Entries[PoseData.Entries.Num() - 2].AccumulatedSeconds <= -SamplingInterval)
+	else if (SamplingInterval <= 0.f || Entries[Entries.Num() - 2].AccumulatedSeconds <= -SamplingInterval)
 	{
-		FPoseHistoryEntry EntryTemp = MoveTemp(PoseData.Entries.First());
-		PoseData.Entries.PopFront();
-		PoseData.Entries.Emplace(MoveTemp(EntryTemp));
+		FPoseHistoryEntry EntryTemp = MoveTemp(Entries.First());
+		Entries.PopFront();
+		Entries.Emplace(MoveTemp(EntryTemp));
 	}
 
 	// Regardless of the retention policy, we always update the most recent Entry
-	FPoseHistoryEntry& MostRecentEntry = PoseData.Entries.Last();
-	MostRecentEntry.Update(0.f, ComponentSpacePoseProvider, PoseData.BoneToTransformMap, bStoreScales, Curves, PoseData.CollectedCurves);
+	FPoseHistoryEntry& MostRecentEntry = Entries.Last();
+	MostRecentEntry.Update(0.f, ComponentSpacePoseProvider, BoneToTransformMap, bStoreScales, Curves, CollectedCurves);
 
 	if (RootBoneRecoveryTime > 0.f && !Trajectory.Samples.IsEmpty())
 	{
@@ -1019,18 +1039,18 @@ void FPoseHistory::EvaluateComponentSpace_AnyThread(float DeltaTime, IComponentS
 		FutureEntryTemp.SetNum(1, bStoreScales);
 		FutureEntryTemp.SetComponentSpaceTransform(RootBoneIndexType, RootBoneTransformAtRecoveryTime);
 		FutureEntryTemp.AccumulatedSeconds = RootBoneRecoveryTime;
-		PoseData.Entries.Emplace(MoveTemp(FutureEntryTemp));
+		Entries.Emplace(MoveTemp(FutureEntryTemp));
 	}
 }
 
 void FPoseHistory::EvaluateComponentSpace_AnyThread(float DeltaTime, FCSPose<FCompactPose>& ComponentSpacePose, bool bStoreScales,
 	float RootBoneRecoveryTime, float RootBoneTranslationRecoveryRatio, float RootBoneRotationRecoveryRatio,
 	bool bNeedsReset, bool bCacheBones, const TArray<FBoneIndexType>& RequiredBones,
-	const FBlendedCurve& Curves, const TConstArrayView<FName>& CollectedCurves)
+	const FBlendedCurve& Curves, const TConstArrayView<FName>& InCollectedCurves)
 {
 	FComponentSpacePoseProvider ComponentSpacePoseProvider(ComponentSpacePose);
 	EvaluateComponentSpace_AnyThread(DeltaTime, ComponentSpacePoseProvider, bStoreScales, RootBoneRecoveryTime, RootBoneTranslationRecoveryRatio,
-		RootBoneRotationRecoveryRatio, bNeedsReset, bCacheBones, RequiredBones, Curves, CollectedCurves);
+		RootBoneRotationRecoveryRatio, bNeedsReset, bCacheBones, RequiredBones, Curves, InCollectedCurves);
 }
 
 #if ENABLE_DRAW_DEBUG && ENABLE_ANIM_DEBUG
@@ -1051,9 +1071,9 @@ void FPoseHistory::DebugDraw(FAnimInstanceProxy& AnimInstanceProxy, FColor Color
 		const bool bValidTrajectory = !Trajectory.Samples.IsEmpty();
 		TArray<FTransform, TInlineAllocator<128>> PrevGlobalTransforms;
 
-		for (int32 EntryIndex = 0; EntryIndex < PoseData.Entries.Num(); ++EntryIndex)
+		for (int32 EntryIndex = 0; EntryIndex < Entries.Num(); ++EntryIndex)
 		{
-			const FPoseHistoryEntry& Entry = PoseData.Entries[EntryIndex];
+			const FPoseHistoryEntry& Entry = Entries[EntryIndex];
 
 			const int32 PrevGlobalTransformsNum = PrevGlobalTransforms.Num();
 			const int32 Max = FMath::Max(PrevGlobalTransformsNum, Entry.Num());
@@ -1293,6 +1313,12 @@ const FPoseHistoryEntry& FMemStackPoseHistory::GetEntry(int32 EntryIndex) const
 	return FutureEntries[EntryIndex - PoseHistoryNumEntries];
 }
 
+const FPoseIndicesHistory* FMemStackPoseHistory::GetPoseIndicesHistory() const
+{
+	check(PoseHistory);
+	return PoseHistory->GetPoseIndicesHistory();
+}
+
 #if ENABLE_DRAW_DEBUG && ENABLE_ANIM_DEBUG
 void FMemStackPoseHistory::DebugDraw(FAnimInstanceProxy& AnimInstanceProxy, FColor Color) const
 {
diff --git a/Engine/Plugins/Animation/PoseSearch/Source/Runtime/Private/PoseSearchInteractionIsland.cpp b/Engine/Plugins/Animation/PoseSearch/Source/Runtime/Private/PoseSearchInteractionIsland.cpp
--- a/Engine/Plugins/Animation/PoseSearch/Source/Runtime/Private/PoseSearchInteractionIsland.cpp
+++ b/Engine/Plugins/Animation/PoseSearch/Source/Runtime/Private/PoseSearchInteractionIsland.cpp
@@ -131,18 +131,13 @@ static FInteractionSearchResult InitSearchResult(const FSearchResult& SearchResu
 	for (int32 AnimContextIndex = 0; AnimContextIndex < AnimContextsNum; ++AnimContextIndex)
 	{
 		const UObject* AnimContext = SearchContext.GetAnimContext(AnimContextIndex);
-		const USkeleton* Skeleton = AnimContext ? GetContextSkeleton(AnimContext) : nullptr;
+		check(AnimContext);
+		const USkeleton* Skeleton = GetContextSkeleton(AnimContext);
 		const IPoseHistory* PoseHistory = SearchContext.GetPoseHistory(AnimContextIndex);
-		if (Skeleton && PoseHistory)
-		{
-			PoseHistory->GetTransformAtTime(0.f, InteractionSearchResult.ActorRootTransforms[AnimContextIndex], Skeleton, ComponentSpaceIndexType, WorldSpaceIndexType);
-			PoseHistory->GetTransformAtTime(0.f, InteractionSearchResult.ActorRootBoneTransforms[AnimContextIndex], Skeleton, RootBoneIndexType, ComponentSpaceIndexType);
-		}
-		else
-		{
-			InteractionSearchResult.ActorRootTransforms[AnimContextIndex] = FTransform::Identity;
-			InteractionSearchResult.ActorRootBoneTransforms[AnimContextIndex] = FTransform::Identity;
-		}
+		check(PoseHistory && Skeleton);
+		
+		PoseHistory->GetTransformAtTime(0.f, InteractionSearchResult.ActorRootTransforms[AnimContextIndex], Skeleton, ComponentSpaceIndexType, WorldSpaceIndexType);
+		PoseHistory->GetTransformAtTime(0.f, InteractionSearchResult.ActorRootBoneTransforms[AnimContextIndex], Skeleton, RootBoneIndexType, ComponentSpaceIndexType);
 	}
 
 	return InteractionSearchResult;
@@ -305,11 +300,8 @@ static void AddPrerequisite(FTickFunction& TickFunction, UObject* TargetObject,
 const UE::PoseSearch::IPoseHistory* FInteractionSearchContextBase::GetPoseHistory(int32 Index) const
 {
 	check(PoseHistories.IsValidIndex(Index));
-	if (PoseHistories[Index].IsValid())
-	{
-		return PoseHistories[Index].Pin().Get();
-	}
-	return nullptr;
+	check(PoseHistories[Index].IsValid());
+	return PoseHistories[Index].Pin().Get();
 }
 
 bool FInteractionSearchContextBase::IsEquivalent(const FInteractionSearchContextBase& Other) const
@@ -443,14 +435,14 @@ void FInteractionSearchContextBase::TestHistoryCollectorsThreadingAccess() const
 	{
 		if (const UAnimInstance* AnimInstance = Cast<UAnimInstance>(AnimContexts[Index].Get()))
 		{
-			if (const IPoseHistory* PoseHistory = GetPoseHistory(Index))
-			{		
-				const USkeleton* Skeleton = AnimInstance->GetRequiredBonesOnAnyThread().GetSkeletonAsset();
-				check(Skeleton);
+			const IPoseHistory* PoseHistory = GetPoseHistory(Index);
+			check(PoseHistory);
+			
+			const USkeleton* Skeleton = AnimInstance->GetRequiredBonesOnAnyThread().GetSkeletonAsset();
+			check(Skeleton);
 
-				FTransform ActorRootBoneTransform;
-				PoseHistory->GetTransformAtTime(0.f, ActorRootBoneTransform, Skeleton, RootBoneIndexType, ComponentSpaceIndexType);
-			}
+			FTransform ActorRootBoneTransform;
+			PoseHistory->GetTransformAtTime(0.f, ActorRootBoneTransform, Skeleton, RootBoneIndexType, ComponentSpaceIndexType);
 		}
 	}
 }
@@ -482,11 +474,11 @@ void FInteractionIsland::FPreTickFunction::ExecuteTick(float DeltaTime, enum ELe
 			{
 				if (const UObject* AnimContext = SearchContext.GetAnimContext(Index))
 				{
-					if (const IPoseHistory* PoseHistory = SearchContext.GetPoseHistory(Index))
-					{
-						// since FInteractionIsland has a tick dependency with the USkeletalMeshComponent it's safe modify the IPoseHistory
-						const_cast<IPoseHistory*>(PoseHistory)->GenerateTrajectory(AnimContext, DeltaTime);
-					}
+					const IPoseHistory* PoseHistory = SearchContext.GetPoseHistory(Index);
+					check(PoseHistory);
+					
+					// since FInteractionIsland has a tick dependency with the USkeletalMeshComponent it's safe modify the IPoseHistory
+					const_cast<IPoseHistory*>(PoseHistory)->GenerateTrajectory(AnimContext, DeltaTime);
 				}
 			}
 		}
@@ -1022,7 +1014,7 @@ bool FInteractionIsland::DoSearch_AnyThread(const UObject* AnimContext, const TC
 		const int32 NumRoles = InteractionSearchContext.Num();
 		Contexts.SetNum(NumRoles);
 
-		FSearchContext SearchContext(0.f, nullptr, FFloatInterval(0.f, 0.f), FPoseSearchEvent());
+		FSearchContext SearchContext(0.f, FFloatInterval(0.f, 0.f), FPoseSearchEvent());
 		for (int32 RoleIndex = 0; RoleIndex < NumRoles; ++RoleIndex)
 		{
 			const UObject* SearchContextAnimContext = InteractionSearchContext.GetAnimContext(RoleIndex);
@@ -1035,12 +1027,7 @@ bool FInteractionIsland::DoSearch_AnyThread(const UObject* AnimContext, const TC
 			Contexts[RoleIndex].AddObjectParam(const_cast<UObject*>(SearchContextAnimContext));
 
 			const UE::PoseSearch::IPoseHistory* PoseHistory = InteractionSearchContext.GetPoseHistory(RoleIndex);
-			if (!PoseHistory)
-			{
-				UE_LOG(LogPoseSearch, Error, TEXT("FInteractionIsland::DoSearch_AnyThread null PoseHistory"));
-				return false;
-			}
-			
+			check(PoseHistory);
 			SearchContext.AddRole(InteractionSearchContext.GetRoles()[RoleIndex], &Contexts[RoleIndex], PoseHistory);
 		}
 
diff --git a/Engine/Plugins/Animation/PoseSearch/Source/Runtime/Private/PoseSearchInteractionSubsystem.cpp b/Engine/Plugins/Animation/PoseSearch/Source/Runtime/Private/PoseSearchInteractionSubsystem.cpp
--- a/Engine/Plugins/Animation/PoseSearch/Source/Runtime/Private/PoseSearchInteractionSubsystem.cpp
+++ b/Engine/Plugins/Animation/PoseSearch/Source/Runtime/Private/PoseSearchInteractionSubsystem.cpp
@@ -779,43 +779,39 @@ void UPoseSearchInteractionSubsystem::GenerateSearchContexts(float DeltaSeconds,
 			check(AnimContextInfo.AnimContextAvailabilities);
 			for (const FPoseSearchInteractionAvailabilityEx& Availability : AnimContextInfo.AnimContextAvailabilities->Availabilities)
 			{
-				if (const IPoseHistory* PoseHistory = Availability.GetPoseHistory(AnimContextInfo.AnimContextAvailabilities->AnimContext))
+				const IPoseHistory* PoseHistory = Availability.GetPoseHistory(AnimContextInfo.AnimContextAvailabilities->AnimContext);
+				check(PoseHistory);
+				
+				if (const UPoseSearchDatabase* Database = Availability.Database.Get())
 				{
-					if (const UPoseSearchDatabase* Database = Availability.Database.Get())
-					{
-						check(Database->Schema);
-						FRoledAnimContextInfos& RoledAnimContextInfos = DatabaseToRoledAnimContextInfos.FindOrAdd(Database);
-						RoledAnimContextInfos.AddRoledAnimContextInfos(Availability, AnimContextInfo, PoseHistory, *Database);
-					}
-					else if (Availability.IsTagValid())
+					check(Database->Schema);
+					FRoledAnimContextInfos& RoledAnimContextInfos = DatabaseToRoledAnimContextInfos.FindOrAdd(Database);
+					RoledAnimContextInfos.AddRoledAnimContextInfos(Availability, AnimContextInfo, PoseHistory, *Database);
+				}
+				else if (Availability.IsTagValid())
+				{
+					// since Database is null, but this availability has a valid Tag, we're looking for valid databases by Availability.Tag
+					if (const FDatabasesPerTag* DatabasesPerTag = TagToDatabases.Find(Availability.Tag))
 					{
-						// since Database is null, but this availability has a valid Tag, we're looking for valid databases by Availability.Tag
-						if (const FDatabasesPerTag* DatabasesPerTag = TagToDatabases.Find(Availability.Tag))
+						check(!DatabasesPerTag->IsEmpty());
+						for (const UPoseSearchDatabase* DatabaseFromTag : *DatabasesPerTag)
 						{
-							check(!DatabasesPerTag->IsEmpty());
-							for (const UPoseSearchDatabase* DatabaseFromTag : *DatabasesPerTag)
-							{
-								check(DatabaseFromTag && DatabaseFromTag->Schema);
-								FRoledAnimContextInfos& RoledAnimContextInfos = DatabaseToRoledAnimContextInfos.FindOrAdd(DatabaseFromTag);
-								RoledAnimContextInfos.AddRoledAnimContextInfos(Availability, AnimContextInfo, PoseHistory, *DatabaseFromTag);
-							}
-						}
-						else
-						{
-							//@todo: should we add a verbose LOG here? not sure since it'd be very spammy...
-							
-							// this is a valid condition we shouldn't log: for example when the "main character" is loaded and publishing availabilities with a valid tag and null database,
-							// looking for other NPC / seconday characters to interact with, but they are not present of didn't publish any availability
+							check(DatabaseFromTag && DatabaseFromTag->Schema);
+							FRoledAnimContextInfos& RoledAnimContextInfos = DatabaseToRoledAnimContextInfos.FindOrAdd(DatabaseFromTag);
+							RoledAnimContextInfos.AddRoledAnimContextInfos(Availability, AnimContextInfo, PoseHistory, *DatabaseFromTag);
 						}
 					}
 					else
 					{
-						UE_LOG(LogPoseSearch, Log, TEXT("UPoseSearchInteractionSubsystem::GenerateSearchContexts null Availability.Database (with invalid Availability.Tag)"));
+						//@todo: should we add a verbose LOG here? not sure since it'd be very spammy...
+							
+						// this is a valid condition we shouldn't log: for example when the "main character" is loaded and publishing availabilities with a valid tag and null database,
+						// looking for other NPC / seconday characters to interact with, but they are not present of didn't publish any availability
 					}
 				}
 				else
 				{
-					UE_LOG(LogPoseSearch, Error, TEXT("UPoseSearchInteractionSubsystem::GenerateSearchContexts couldn't find PoseHistory %s for AnimContext %s"), *Availability.GetPoseHistoryName(), *AnimContextInfo.AnimContextAvailabilities->AnimContext->GetName());
+					UE_LOG(LogPoseSearch, Log, TEXT("UPoseSearchInteractionSubsystem::GenerateSearchContexts null Availability.Database (with invalid Availability.Tag)"));
 				}
 			}
 		},
diff --git a/Engine/Plugins/Animation/PoseSearch/Source/Runtime/Private/PoseSearchLibrary.cpp b/Engine/Plugins/Animation/PoseSearch/Source/Runtime/Private/PoseSearchLibrary.cpp
--- a/Engine/Plugins/Animation/PoseSearch/Source/Runtime/Private/PoseSearchLibrary.cpp
+++ b/Engine/Plugins/Animation/PoseSearch/Source/Runtime/Private/PoseSearchLibrary.cpp
@@ -505,7 +505,6 @@ PRAGMA_ENABLE_DEPRECATION_WARNINGS
 	SearchResult = FPoseSearchBlueprintResult();
 	// Set the elapsed time to INFINITY to trigger a search right away
 	ElapsedPoseSearchTime = std::numeric_limits<float>::infinity();
-	PoseIndicesHistory.Reset();
 }
 
 void FMotionMatchingState::AdjustAssetTime(float AssetTime)
@@ -580,7 +579,7 @@ void UPoseSearchLibrary::TraceMotionMatching(
 
 void UPoseSearchLibrary::TraceMotionMatching(
 	UE::PoseSearch::FSearchContext& SearchContext,
-	const UE::PoseSearch::ISearchResults& SearchResults,
+	const UE::PoseSearch::FSearchResults& SearchResults,
 	float ElapsedPoseSearchTime,
 	float WantedPlayRate,
 	EPoseSearchInterruptMode InterruptMode)
@@ -725,7 +724,7 @@ void UPoseSearchLibrary::TraceMotionMatching(
 
 	// @todo: integrate DeltaTime into SearchContext, and implement it for UAF as well
 	float DeltaTime = FiniteDelta;
-	if (SearchContext.GetContexts()[0]->ObjectParams.Num() > 0)
+	if (!SearchContext.GetContexts().IsEmpty())
 	{
 		if (const UAnimInstance* AnimInstance = Cast<UAnimInstance>(SearchContext.GetContexts()[0]->GetFirstObjectParam()))
 		{
@@ -908,7 +907,7 @@ void UPoseSearchLibrary::UpdateMotionMatchingState(
 	FMemMark Mark(FMemStack::Get());
 
 	const FPoseSearchEvent PlayRateOverriddenEvent = EventToSearch.GetPlayRateOverriddenEvent(PlayRate);
-	FSearchContext SearchContext(0.f, &InOutMotionMatchingState.PoseIndicesHistory, PoseJumpThresholdTime, PlayRateOverriddenEvent);
+	FSearchContext SearchContext(0.f, PoseJumpThresholdTime, PlayRateOverriddenEvent);
 	SearchContext.AddRole(GetCommonDefaultRole(Databases), AnimContext, PoseHistory);
 
 	const UPoseSearchDatabase* CurrentResultDatabase = InOutMotionMatchingState.SearchResult.SelectedDatabase.Get();
@@ -1016,7 +1015,15 @@ void UPoseSearchLibrary::UpdateMotionMatchingState(
 	}
 
 	const float WantedPlayRate = CalculateWantedPlayRate(SearchResult, SearchContext, PlayRate, PoseHistory ? PoseHistory->GetTrajectorySpeedMultiplier() : 1.f, EventToSearch);
-	InOutMotionMatchingState.PoseIndicesHistory.Update(SearchResult, DeltaTime, PoseReselectHistory);
+	if (PoseHistory)
+	{
+		if (const FPoseIndicesHistory* PoseIndicesHistory = PoseHistory->GetPoseIndicesHistory())
+		{
+			// const casting here is safe since we're in the thread owning the pose history, and it's the correct place to update the previously selected poses
+			const_cast<FPoseIndicesHistory*>(PoseIndicesHistory)->Update(SearchResult, DeltaTime, PoseReselectHistory);
+		}
+	}
+
 	InOutMotionMatchingState.SearchResult.InitFrom(SearchResult, WantedPlayRate);
 
 PRAGMA_DISABLE_DEPRECATION_WARNINGS
@@ -1212,7 +1219,7 @@ void UPoseSearchLibrary::MotionMatch(
 	const FPoseSearchContinuingProperties& ContinuingProperties,
 	const FPoseSearchFutureProperties& Future,
 	const FPoseSearchEvent& EventToSearch,
-	UE::PoseSearch::ISearchResults& SearchResults)
+	UE::PoseSearch::FSearchResults& SearchResults)
 {
 	TArray<FChooserEvaluationContext, TInlineAllocator<UE::PoseSearch::PreallocatedRolesNum>> Contexts;
 	const int NumContexts = AnimContexts.Num(); 
@@ -1233,7 +1240,7 @@ void UPoseSearchLibrary::MotionMatch(
 	const FPoseSearchContinuingProperties& ContinuingProperties,
 	float DesiredPermutationTimeOffset,
 	const FPoseSearchEvent& EventToSearch,
-	UE::PoseSearch::ISearchResults& SearchResults)
+	UE::PoseSearch::FSearchResults& SearchResults)
 {
 	TArray<FChooserEvaluationContext, TInlineAllocator<UE::PoseSearch::PreallocatedRolesNum>> Contexts;
 	const int NumContexts = AnimContexts.Num(); 
@@ -1254,7 +1261,7 @@ void UPoseSearchLibrary::MotionMatch(
 	const FPoseSearchContinuingProperties& ContinuingProperties,
 	const FPoseSearchFutureProperties& Future,
 	const FPoseSearchEvent& EventToSearch,
-	UE::PoseSearch::ISearchResults& SearchResults)
+	UE::PoseSearch::FSearchResults& SearchResults)
 {
 	check(!Contexts.IsEmpty() && Contexts.Num() == Roles.Num() && Contexts.Num() == PoseHistories.Num());
 
@@ -1288,26 +1295,25 @@ void UPoseSearchLibrary::MotionMatch(
 		{
 			if (const IPoseHistory* PoseHistory = InternalPoseHistories[RoleIndex])
 			{
-				if (const USkeleton* Skeleton = GetContextSkeleton(Contexts[RoleIndex]))
+				const USkeleton* Skeleton = GetContextSkeleton(Contexts[RoleIndex]);
+				check(Skeleton);
+				// @todo: add input BlendParameters to support sampling FutureAnimation blendspaces and support for multi character
+				const UAnimationAsset* AnimationAsset = Cast<UAnimationAsset>(Future.Animation);
+				if (!AnimationAsset)
 				{
-					// @todo: add input BlendParameters to support sampling FutureAnimation blendspaces and support for multi character
-					const UAnimationAsset* AnimationAsset = Cast<UAnimationAsset>(Future.Animation);
-					if (!AnimationAsset)
+					if (const UMultiAnimAsset* MultiAnimAsset = Cast<UMultiAnimAsset>(Future.Animation))
 					{
-						if (const UMultiAnimAsset* MultiAnimAsset = Cast<UMultiAnimAsset>(Future.Animation))
-						{
-							AnimationAsset = MultiAnimAsset->GetAnimationAsset(Roles[RoleIndex]);
-						}
-						else
-						{
-							checkNoEntry();
-						}
+						AnimationAsset = MultiAnimAsset->GetAnimationAsset(Roles[RoleIndex]);
+					}
+					else
+					{
+						checkNoEntry();
 					}
-
-					MemStackPoseHistories[RoleIndex].Init(InternalPoseHistories[RoleIndex]);
-					MemStackPoseHistories[RoleIndex].ExtractAndAddFuturePoses(AnimationAsset, FutureAnimationTime, FiniteDelta, FVector::ZeroVector, FutureIntervalTime, Skeleton);
-					InternalPoseHistories[RoleIndex] = MemStackPoseHistories[RoleIndex].GetThisOrPoseHistory();
 				}
+
+				MemStackPoseHistories[RoleIndex].Init(InternalPoseHistories[RoleIndex]);
+				MemStackPoseHistories[RoleIndex].ExtractAndAddFuturePoses(AnimationAsset, FutureAnimationTime, FiniteDelta, FVector::ZeroVector, FutureIntervalTime, Skeleton);
+				InternalPoseHistories[RoleIndex] = MemStackPoseHistories[RoleIndex].GetThisOrPoseHistory();
 			}
 		}
 	}		
@@ -1323,11 +1329,11 @@ void UPoseSearchLibrary::MotionMatch(
 	const FPoseSearchContinuingProperties& ContinuingProperties,
 	const float DesiredPermutationTimeOffset,
 	const FPoseSearchEvent& EventToSearch,
-	UE::PoseSearch::ISearchResults& SearchResults)
+	UE::PoseSearch::FSearchResults& SearchResults)
 {
 	using namespace UE::PoseSearch;
 	
-	FSearchContext SearchContext(DesiredPermutationTimeOffset, nullptr, FFloatInterval(0.f, 0.f), EventToSearch);
+	FSearchContext SearchContext(DesiredPermutationTimeOffset, FFloatInterval(0.f, 0.f), EventToSearch);
 	for (int32 RoleIndex = 0; RoleIndex < Roles.Num(); ++RoleIndex)
 	{
 		SearchContext.AddRole(Roles[RoleIndex], &Contexts[RoleIndex], PoseHistories[RoleIndex]);
@@ -1340,7 +1346,7 @@ void UPoseSearchLibrary::MotionMatch(
 	UE::PoseSearch::FSearchContext& SearchContext,
 	const TArrayView<const UObject*> AssetsToSearch,
 	const FPoseSearchContinuingProperties& ContinuingProperties,
-	UE::PoseSearch::ISearchResults& SearchResults)
+	UE::PoseSearch::FSearchResults& SearchResults)
 {
 	using namespace UE::PoseSearch;
 
diff --git a/Engine/Plugins/Animation/PoseSearch/Source/Runtime/Private/PoseSearchModule.cpp b/Engine/Plugins/Animation/PoseSearch/Source/Runtime/Private/PoseSearchModule.cpp
--- a/Engine/Plugins/Animation/PoseSearch/Source/Runtime/Private/PoseSearchModule.cpp
+++ b/Engine/Plugins/Animation/PoseSearch/Source/Runtime/Private/PoseSearchModule.cpp
@@ -72,42 +72,48 @@ public:
 			PoseHistory = &PoseHistoryProvider->GetPoseHistory();
 		}
 
-		// @todo: Add event support
-		const UObject* AnimContext = GraphContext.AnimInstanceProxy->GetAnimInstanceObject();
-		const UE::PoseSearch::FSearchResult SearchResult = UPoseSearchLibrary::MotionMatch(MakeArrayView(&AnimContext, 1),
-			MakeArrayView(&DefaultRole, 1), MakeArrayView(&PoseHistory, 1), AssetsToSearch, ContinuingProperties, Future, FPoseSearchEvent());
-
 		UE::Anim::IPoseSearchProvider::FSearchResult ProviderResult;
-		if (const FSearchIndexAsset* SearchIndexAsset = SearchResult.GetSearchIndexAsset())
+		if (!PoseHistory)
 		{
-			const UPoseSearchDatabase* Database = SearchResult.Database.Get();
-			if (const FPoseSearchDatabaseAnimationAssetBase* DatabaseAnimationAssetBase = Database->GetDatabaseAnimationAsset(*SearchIndexAsset))
-			{
-				ProviderResult.SelectedAsset = DatabaseAnimationAssetBase->GetAnimationAsset();
-				ProviderResult.Dissimilarity = SearchResult.PoseCost;
-				ProviderResult.TimeOffsetSeconds = SearchResult.GetAssetTime();
-				ProviderResult.bIsFromContinuingPlaying = SearchResult.bIsContinuingPoseSearch;
-				ProviderResult.bMirrored = SearchIndexAsset->IsMirrored();
-				ProviderResult.BlendParameters = SearchIndexAsset->GetBlendParameters();
+			UE_LOG(LogPoseSearch, Error, TEXT("FPoseSearchModule::Search, missing IPoseHistory"));
+		}
+		else
+		{
+			// @todo: Add event support
+			const UObject* AnimContext = GraphContext.AnimInstanceProxy->GetAnimInstanceObject();
+			const UE::PoseSearch::FSearchResult SearchResult = UPoseSearchLibrary::MotionMatch(MakeArrayView(&AnimContext, 1),
+				MakeArrayView(&DefaultRole, 1), MakeArrayView(&PoseHistory, 1), AssetsToSearch, ContinuingProperties, Future, FPoseSearchEvent());
 
-				// figuring out the WantedPlayRate
-				ProviderResult.WantedPlayRate = 1.f;
-				if (Future.Animation && Future.IntervalTime > 0.f)
+			if (const FSearchIndexAsset* SearchIndexAsset = SearchResult.GetSearchIndexAsset())
+			{
+				const UPoseSearchDatabase* Database = SearchResult.Database.Get();
+				if (const FPoseSearchDatabaseAnimationAssetBase* DatabaseAnimationAssetBase = Database->GetDatabaseAnimationAsset(*SearchIndexAsset))
 				{
-					if (const UPoseSearchFeatureChannel_PermutationTime* PermutationTimeChannel = Database->Schema->FindFirstChannelOfType<UPoseSearchFeatureChannel_PermutationTime>())
+					ProviderResult.SelectedAsset = DatabaseAnimationAssetBase->GetAnimationAsset();
+					ProviderResult.Dissimilarity = SearchResult.PoseCost;
+					ProviderResult.TimeOffsetSeconds = SearchResult.GetAssetTime();
+					ProviderResult.bIsFromContinuingPlaying = SearchResult.bIsContinuingPoseSearch;
+					ProviderResult.bMirrored = SearchIndexAsset->IsMirrored();
+					ProviderResult.BlendParameters = SearchIndexAsset->GetBlendParameters();
+
+					// figuring out the WantedPlayRate
+					ProviderResult.WantedPlayRate = 1.f;
+					if (Future.Animation && Future.IntervalTime > 0.f)
 					{
-						const FSearchIndex& SearchIndex = Database->GetSearchIndex();
-						if (!SearchIndex.IsValuesEmpty())
+						if (const UPoseSearchFeatureChannel_PermutationTime* PermutationTimeChannel = Database->Schema->FindFirstChannelOfType<UPoseSearchFeatureChannel_PermutationTime>())
 						{
-							TConstArrayView<float> ResultData = Database->GetSearchIndex().GetPoseValues(SearchResult.PoseIdx);
-							const float ActualIntervalTime = PermutationTimeChannel->GetPermutationTime(ResultData);
-							ProviderResult.WantedPlayRate = ActualIntervalTime / Future.IntervalTime;
+							const FSearchIndex& SearchIndex = Database->GetSearchIndex();
+							if (!SearchIndex.IsValuesEmpty())
+							{
+								TConstArrayView<float> ResultData = Database->GetSearchIndex().GetPoseValues(SearchResult.PoseIdx);
+								const float ActualIntervalTime = PermutationTimeChannel->GetPermutationTime(ResultData);
+								ProviderResult.WantedPlayRate = ActualIntervalTime / Future.IntervalTime;
+							}
 						}
 					}
 				}
 			}
 		}
-
 		return ProviderResult;
 	}
 
diff --git a/Engine/Plugins/Animation/PoseSearch/Source/Runtime/Private/PoseSearchResult.cpp b/Engine/Plugins/Animation/PoseSearch/Source/Runtime/Private/PoseSearchResult.cpp
--- a/Engine/Plugins/Animation/PoseSearch/Source/Runtime/Private/PoseSearchResult.cpp
+++ b/Engine/Plugins/Animation/PoseSearch/Source/Runtime/Private/PoseSearchResult.cpp
@@ -225,37 +225,40 @@ void FSearchResult::InitFrom(const FPoseSearchBlueprintResult& BlueprintResult)
 
 /////////////////////////////////////////////////////
 // FSearchResults_Single
-void FSearchResults_Single::UpdateWith(const FSearchResult& SearchResult)
+FSearchResults_Single::FSearchResults_Single()
 {
-	check(SearchResult.IsValid());
-	if (SearchResult.PoseCost < SingleSearchResult.PoseCost)
-	{
-		SingleSearchResult = SearchResult;
-	}
-}
-
-void FSearchResults_Single::FinalizeResults()
-{
-	if (SingleSearchResult.IsValid() && !SingleSearchResult.IsAssetTimeValid())
-	{
-		SingleSearchResult.SetAssetTime(SingleSearchResult.Database->GetNormalizedAssetTime(SingleSearchResult.PoseIdx));
-		check(SingleSearchResult.DebugValidate());
-	}
-}
+	UpdateWith = [this](const FSearchResult& SearchResult)
+		{
+			check(SearchResult.IsValid());
+			if (SearchResult.PoseCost < SingleSearchResult.PoseCost)
+			{
+				SingleSearchResult = SearchResult;
+			}
+		};
 
-FSearchResult FSearchResults_Single::GetBestResult() const
-{
-	return SingleSearchResult;
-}
+	FinalizeResults = [this]()
+		{
+			if (SingleSearchResult.IsValid() && !SingleSearchResult.IsAssetTimeValid())
+			{
+				SingleSearchResult.SetAssetTime(SingleSearchResult.Database->GetNormalizedAssetTime(SingleSearchResult.PoseIdx));
+				check(SingleSearchResult.DebugValidate());
+			}
+		};
 
-bool FSearchResults_Single::ShouldPerformSearch(float SearchMinimumCost) const
-{
-	return SingleSearchResult.PoseCost > SearchMinimumCost;
-}
+	GetBestResult = [this]()
+		{
+			return SingleSearchResult;
+		};
+	
+	ShouldPerformSearch = [this](float SearchMinimumCost)
+		{
+			return SingleSearchResult.PoseCost > SearchMinimumCost;
+		};
 
-bool FSearchResults_Single::IterateOverSearchResults(const TFunctionRef<bool(const FSearchResult& SearchResult)> IterateOverSearchResultsFunction) const
-{
-	return IterateOverSearchResultsFunction(SingleSearchResult);
+	IterateOverSearchResults = [this](const TFunctionRef<bool(const FSearchResult& SearchResult)> IterateOverSearchResultsFunction)
+		{
+			return IterateOverSearchResultsFunction(SingleSearchResult);
+		};
 }
 
 /////////////////////////////////////////////////////
@@ -265,139 +268,207 @@ FSearchResults_Multi::FSearchResults_Multi(int32 InMaxNumberOfWantedResults)
 {
 	check(MaxNumberOfWantedResults > 0); // or else we need to refine the logic in UpdateWith
 	SearchResults.Reserve(MaxNumberOfWantedResults);
-}
 
-void FSearchResults_Multi::UpdateWith(const FSearchResult& SearchResult)
-{
-	check(SearchResult.IsValid());
-	if (SearchResults.Num() < MaxNumberOfWantedResults)
-	{
-		SearchResults.Add(SearchResult);
+	UpdateWith = [this](const FSearchResult& SearchResult)
+		{
+			check(SearchResult.IsValid());
+			if (SearchResults.Num() < MaxNumberOfWantedResults)
+			{
+				SearchResults.Add(SearchResult);
 
-		// @todo: speed up this sort, since SearchResults started alreay sorted!
-		SearchResults.Sort([](const FSearchResult& A, const FSearchResult& B)
+				// @todo: speed up this sort, since SearchResults started alreay sorted!
+				SearchResults.Sort([](const FSearchResult& A, const FSearchResult& B)
+					{
+						return A.PoseCost < B.PoseCost;
+					});
+			}
+			else
 			{
-				return A.PoseCost < B.PoseCost;
-			});
-	}
-	else
-	{
-		FSearchResult& WorstResult = SearchResults.Last();
-		if (SearchResult.PoseCost < WorstResult.PoseCost)
-		{
-			WorstResult = SearchResult;
+				FSearchResult& WorstResult = SearchResults.Last();
+				if (SearchResult.PoseCost < WorstResult.PoseCost)
+				{
+					WorstResult = SearchResult;
+
+					// @todo: speed up this sort, since SearchResults started alreay sorted!
+					SearchResults.Sort([](const FSearchResult& A, const FSearchResult& B)
+						{
+							return A.PoseCost < B.PoseCost;
+						});
+				}
+			}
+		};
 
-			// @todo: speed up this sort, since SearchResults started alreay sorted!
-			SearchResults.Sort([](const FSearchResult& A, const FSearchResult& B)
+	FinalizeResults = [this]()
+		{
+			for (FSearchResult& SearchResult : SearchResults)
+			{
+				check(SearchResult.IsValid());
+				if (!SearchResult.IsAssetTimeValid())
 				{
-					return A.PoseCost < B.PoseCost;
-				});
-		}
-	}
-}
+					SearchResult.SetAssetTime(SearchResult.Database->GetNormalizedAssetTime(SearchResult.PoseIdx));
+					check(SearchResult.DebugValidate());
+				}
+			}
+		};
 
-void FSearchResults_Multi::FinalizeResults()
-{
-	for (FSearchResult& SearchResult : SearchResults)
-	{
-		check(SearchResult.IsValid());
-		if (!SearchResult.IsAssetTimeValid())
+	GetBestResult = [this]()
 		{
-			SearchResult.SetAssetTime(SearchResult.Database->GetNormalizedAssetTime(SearchResult.PoseIdx));
-			check(SearchResult.DebugValidate());
-		}
-	}
-}
+			if (!SearchResults.IsEmpty())
+			{
+				return SearchResults[0];
+			}
+			return FSearchResult();
+		};
+	
+	ShouldPerformSearch = [this](float SearchMinimumCost)
+		{
+			if (SearchResults.IsEmpty())
+			{
+				return true;
+			}
 
-FSearchResult FSearchResults_Multi::GetBestResult() const
-{
-	if (!SearchResults.IsEmpty())
-	{
-		return SearchResults[0];
-	}
-	return FSearchResult();
+			return SearchResults.Last().PoseCost > SearchMinimumCost;
+		};
+
+	IterateOverSearchResults = [this](const TFunctionRef<bool(const FSearchResult& SearchResult)> IterateOverSearchResultsFunction)
+		{
+			for (const FSearchResult& SearchResult : SearchResults)
+			{
+				if (IterateOverSearchResultsFunction(SearchResult))
+				{
+					return true;
+				}
+			}
+			return false;
+		};
 }
 
-bool FSearchResults_Multi::ShouldPerformSearch(float SearchMinimumCost) const
+/////////////////////////////////////////////////////
+// FSearchResults_AssetBests
+FSearchResults_AssetBests::FSearchResults_AssetBests()
 {
-	if (SearchResults.IsEmpty())
-	{
-		return true;
-	}
+	UpdateWith = [this](const FSearchResult& SearchResult)
+		{
+			check(SearchResult.IsValid());
+			if (const FSearchIndexAsset* SearchIndexAsset = SearchResult.GetSearchIndexAsset())
+			{
+				FSourceAssetIdxToBestSearchResult& SourceAssetIdxToBestSearchResult = PerDatabaseSourceAssetIdxToBestSearchResult.FindOrAdd(SearchResult.Database.Get());
+				FSearchResult& BestSearchResultForSourceAssetIdx = SourceAssetIdxToBestSearchResult.FindOrAdd(SearchIndexAsset->GetSourceAssetIdx());
+				if (SearchResult.PoseCost < BestSearchResultForSourceAssetIdx.PoseCost)
+				{
+					BestSearchResultForSourceAssetIdx = SearchResult;
+				}
+			}
+		};
 
-	return SearchResults.Last().PoseCost > SearchMinimumCost;
-}
+	FinalizeResults = [this]()
+		{
+			for (TPair<const UPoseSearchDatabase*, FSourceAssetIdxToBestSearchResult>& PerDatabaseSourceAssetIdxToBestSearchResultPair : PerDatabaseSourceAssetIdxToBestSearchResult)
+			{
+				for (TPair<int32, FSearchResult>& SourceAssetIdxToBestSearchResultPair : PerDatabaseSourceAssetIdxToBestSearchResultPair.Value)
+				{
+					FSearchResult& BestSearchResultForSourceAssetIdx = SourceAssetIdxToBestSearchResultPair.Value;
+					check(BestSearchResultForSourceAssetIdx.IsValid());
+					if (!BestSearchResultForSourceAssetIdx.IsAssetTimeValid())
+					{
+						BestSearchResultForSourceAssetIdx.SetAssetTime(BestSearchResultForSourceAssetIdx.Database->GetNormalizedAssetTime(BestSearchResultForSourceAssetIdx.PoseIdx));
+						check(BestSearchResultForSourceAssetIdx.DebugValidate());
+					}
+				}
+			}
+		};
 
-bool FSearchResults_Multi::IterateOverSearchResults(const TFunctionRef<bool(const FSearchResult& SearchResult)> IterateOverSearchResultsFunction) const
-{
-	for (const FSearchResult& SearchResult : SearchResults)
-	{
-		if (IterateOverSearchResultsFunction(SearchResult))
+	GetBestResult = [this]()
+		{
+			FSearchResult BestSearchResult;
+			for (const TPair<const UPoseSearchDatabase*, FSourceAssetIdxToBestSearchResult>& PerDatabaseSourceAssetIdxToBestSearchResultPair : PerDatabaseSourceAssetIdxToBestSearchResult)
+			{
+				for (const TPair<int32, FSearchResult>& SourceAssetIdxToBestSearchResultPair : PerDatabaseSourceAssetIdxToBestSearchResultPair.Value)
+				{
+					const FSearchResult& BestSearchResultForSourceAssetIdx = SourceAssetIdxToBestSearchResultPair.Value;
+					if (BestSearchResultForSourceAssetIdx.PoseCost < BestSearchResult.PoseCost)
+					{
+						BestSearchResult = BestSearchResultForSourceAssetIdx;
+					}
+				}
+			}
+			return BestSearchResult;
+		};
+	
+	ShouldPerformSearch = [this](float SearchMinimumCost)
 		{
 			return true;
-		}
-	}
-	return false;
-}
+		};
 
-/////////////////////////////////////////////////////
-// FSearchResults_AssetBests
-void FSearchResults_AssetBests::UpdateWith(const FSearchResult& SearchResult)
-{
-	check(SearchResult.IsValid());
-	if (const FSearchIndexAsset* SearchIndexAsset = SearchResult.GetSearchIndexAsset())
-	{
-		FSearchResult& BestSearchResultForSourceAssetIdx = SourceAssetIdxToBestSearchResult.FindOrAdd(SearchIndexAsset->GetSourceAssetIdx());
-		if (SearchResult.PoseCost < BestSearchResultForSourceAssetIdx.PoseCost)
+	IterateOverSearchResults = [this](const TFunctionRef<bool(const FSearchResult& SearchResult)> IterateOverSearchResultsFunction)
 		{
-			BestSearchResultForSourceAssetIdx = SearchResult;
-		}
-	}
+			for (const TPair<const UPoseSearchDatabase*, FSourceAssetIdxToBestSearchResult>& PerDatabaseSourceAssetIdxToBestSearchResultPair : PerDatabaseSourceAssetIdxToBestSearchResult)
+			{
+				for (const TPair<int32, FSearchResult>& SourceAssetIdxToBestSearchResultPair : PerDatabaseSourceAssetIdxToBestSearchResultPair.Value)
+				{
+					if (IterateOverSearchResultsFunction(SourceAssetIdxToBestSearchResultPair.Value))
+					{
+						return true;
+					}
+				}
+			}
+			return false;
+		};
 }
 
-void FSearchResults_AssetBests::FinalizeResults()
+const FSearchResult* FSearchResults_AssetBests::FindSearchResultFor(const UPoseSearchDatabase* Database, int32 SourceAssetIdx) const
 {
-	for (TPair<int32, FSearchResult>& SourceAssetIdxToBestSearchResultPair : SourceAssetIdxToBestSearchResult)
+	if (const FSourceAssetIdxToBestSearchResult* SourceAssetIdxToBestSearchResult = PerDatabaseSourceAssetIdxToBestSearchResult.Find(Database))
 	{
-		FSearchResult& BestSearchResultForSourceAssetIdx = SourceAssetIdxToBestSearchResultPair.Value;
-		check(BestSearchResultForSourceAssetIdx.IsValid());
-		if (!BestSearchResultForSourceAssetIdx.IsAssetTimeValid())
-		{
-			BestSearchResultForSourceAssetIdx.SetAssetTime(BestSearchResultForSourceAssetIdx.Database->GetNormalizedAssetTime(BestSearchResultForSourceAssetIdx.PoseIdx));
-			check(BestSearchResultForSourceAssetIdx.DebugValidate());
-		}
+		return SourceAssetIdxToBestSearchResult->Find(SourceAssetIdx);
 	}
+	return nullptr;
 }
 
-FSearchResult FSearchResults_AssetBests::GetBestResult() const
+void FSearchResults_AssetBests::Shrink(int32 MaxNumResults)
 {
-	FSearchResult BestSearchResult;
-	for (const TPair<int32, FSearchResult>& SourceAssetIdxToBestSearchResultPair : SourceAssetIdxToBestSearchResult)
+	int32 NumResults = 0;
+	for (const TPair<const UPoseSearchDatabase*, FSourceAssetIdxToBestSearchResult>& PerDatabaseSourceAssetIdxToBestSearchResultPair : PerDatabaseSourceAssetIdxToBestSearchResult)
 	{
-		const FSearchResult& BestSearchResultForSourceAssetIdx = SourceAssetIdxToBestSearchResultPair.Value;
-		if (BestSearchResultForSourceAssetIdx.PoseCost < BestSearchResult.PoseCost)
-		{
-			BestSearchResult = BestSearchResultForSourceAssetIdx;
-		}
+		NumResults += PerDatabaseSourceAssetIdxToBestSearchResultPair.Value.Num();
 	}
-	return BestSearchResult;
-}
 
-bool FSearchResults_AssetBests::ShouldPerformSearch(float SearchMinimumCost) const
-{
-	return true;
-}
-
-bool FSearchResults_AssetBests::IterateOverSearchResults(const TFunctionRef<bool(const FSearchResult& SearchResult)> IterateOverSearchResultsFunction) const
-{
-	for (const TPair<int32, FSearchResult>& SourceAssetIdxToBestSearchResultPair : SourceAssetIdxToBestSearchResult)
+	if (NumResults > MaxNumResults)
 	{
-		if (IterateOverSearchResultsFunction(SourceAssetIdxToBestSearchResultPair.Value))
+		FMemMark Mark(FMemStack::Get());
+
+		struct FSearchResultHeapCompare
 		{
-			return true;
+			bool operator()(const FSearchResult& A, const FSearchResult& B) const
+			{
+				// using > to create a max heap where BestSearchResults.HeapTop().Cost is the greatest Cost (to behave like the std::priority_queue)
+				return A.PoseCost > B.PoseCost;
+			}
+		};
+		
+		// max heap of FSearchResult(s).
+		TArray<FSearchResult, TInlineAllocator<8, TMemStackAllocator<>>> BestSearchResults;
+		BestSearchResults.Reserve(MaxNumResults);
+
+		for (const TPair<const UPoseSearchDatabase*, FSourceAssetIdxToBestSearchResult>& PerDatabaseSourceAssetIdxToBestSearchResultPair : PerDatabaseSourceAssetIdxToBestSearchResult)
+		{
+			for (const TPair<int32, FSearchResult>& SourceAssetIdxToBestSearchResultPair : PerDatabaseSourceAssetIdxToBestSearchResultPair.Value)
+			{
+				const FSearchResult& SearchResult = SourceAssetIdxToBestSearchResultPair.Value;
+				if (BestSearchResults.Num() < MaxNumResults)
+				{
+					BestSearchResults.HeapPush(SearchResult, FSearchResultHeapCompare());
+				}
+				else if (SearchResult.PoseCost < BestSearchResults.HeapTop().PoseCost)
+				{
+					// popping the max heap (BestSearchResults) head since it costs more than the FSearchResult we're about to insert
+					FSearchResult PoppedSearchResult;
+					BestSearchResults.HeapPop(PoppedSearchResult, FSearchResultHeapCompare(), EAllowShrinking::No);
+					BestSearchResults.HeapPush(SearchResult, FSearchResultHeapCompare());
+				}
+			}
 		}
 	}
-	return false;
 }
 
 } // namespace UE::PoseSearch
diff --git a/Engine/Plugins/Animation/PoseSearch/Source/Runtime/Public/PoseSearch/Chooser/PoseSearchChooserColumn.h b/Engine/Plugins/Animation/PoseSearch/Source/Runtime/Public/PoseSearch/Chooser/PoseSearchChooserColumn.h
--- a/Engine/Plugins/Animation/PoseSearch/Source/Runtime/Public/PoseSearch/Chooser/PoseSearchChooserColumn.h
+++ b/Engine/Plugins/Animation/PoseSearch/Source/Runtime/Public/PoseSearch/Chooser/PoseSearchChooserColumn.h
@@ -4,14 +4,13 @@
 #include <IHasContext.h>
 
 #include "CoreMinimal.h"
-#include "IChooserColumn.h"
 #include "ChooserPropertyAccess.h"
-#include "PoseSearch/PoseSearchDatabase.h"
+#include "IChooserColumn.h"
 #include "PoseSearch/Chooser/ChooserParameterPoseHistoryBase.h"
-#include "StructUtils/InstancedStruct.h"
+#include "PoseSearch/PoseSearchLibrary.h"
 #include "Serialization/MemoryReader.h"
+#include "StructUtils/InstancedStruct.h"
 #include "Templates/SubclassOf.h"
-#include "PoseSearch/PoseSearchHistory.h"
 #include "PoseSearchChooserColumn.generated.h"
 
 #define UE_API POSESEARCH_API
@@ -22,7 +21,6 @@ USTRUCT(Experimental, DisplayName = "Pose History Property Binding")
 struct FPoseHistoryContextProperty :  public FChooserParameterPoseHistoryBase
 {
 	GENERATED_BODY()
-public:
 
 	UPROPERTY(EditAnywhere, Meta = (BindingType = "FPoseHistoryReference", BindingAllowFunctions = "false", BindingColor = "StructPinTypeColor"), Category = "Binding")
 	FChooserPropertyBinding Binding;
@@ -46,31 +44,50 @@ struct FPoseSearchColumn : public FChooserColumnBase
 	UPROPERTY(VisibleAnywhere, Category = "Data")
 	TObjectPtr<UPoseSearchDatabase> InternalDatabase;
 	
+	// @todo: expose it as pinnable parameter
+	// Prevent re-selection of poses that have been selected previously within this much time (in seconds) in the past. This is across all animation segments that have been selected within this time range.
+	UPROPERTY(EditAnywhere, Category = Settings, meta = (PinHiddenByDefault, ClampMin = "0"))
+	float PoseReselectHistory = 0.3f;
+
+	// @todo: expose it as pinnable parameter
+	// MaxNumberOfResults represent the maximum number of results from the motion matching search. if MaxNumberOfResults <= 0 the column will add ALL the results to the output
+	// if bEnableMultiSelection is true, this FPoseSearchColumn will be set as cost column, retrieving multiple results from the MM search (one per chooser table entry)
+	// and setting their relative costs to be used with subsequent column such as the FRandomizeColumn for additional results refinements.
+	// if bEnableMultiSelection is false, this FPoseSearchColumn will return ONLY the best result, the one with the lowest cost
+	UPROPERTY(EditAnywhere, Category = "Data")
+	int32 MaxNumberOfResults = 1;
+
+	// @todo: expose it as pinnable parameter
+	UPROPERTY(EditAnywhere, Category = "Data")
+	EPoseSearchInterruptMode InterruptMode = EPoseSearchInterruptMode::DoNotInterrupt;
+
 	// Pose History
 	UPROPERTY(EditAnywhere, DisplayName = "Pose History", NoClear, Meta = (ExcludeBaseStruct, BaseStruct = "/Script/PoseSearch.ChooserParameterPoseHistoryBase"), Category = "Data")
 	FInstancedStruct InputValue;
 	
 	// Float output for the start time with the best matching pose
-	UPROPERTY(EditAnywhere, NoClear, Meta = (ExcludeBaseStruct, BaseStruct = "/Script/Chooser.ChooserParameterFloatBase"), Category = "Data")
+	UPROPERTY(EditAnywhere, NoClear, Meta = (ExcludeBaseStruct, BaseStruct = "/Script/Chooser.ChooserParameterFloatBase"), Category = "Output")
 	FInstancedStruct OutputStartTime;
 	
 	// Bool output for if asset should be mirrored
-	UPROPERTY(EditAnywhere, NoClear, Meta = (ExcludeBaseStruct, BaseStruct = "/Script/Chooser.ChooserParameterBoolBase"), Category = "Data")
+	UPROPERTY(EditAnywhere, NoClear, Meta = (ExcludeBaseStruct, BaseStruct = "/Script/Chooser.ChooserParameterBoolBase"), Category = "Output")
 	FInstancedStruct OutputMirror;
 
 	// Float output for the cost of the selected pose
-	UPROPERTY(EditAnywhere, NoClear, Meta = (ExcludeBaseStruct, BaseStruct = "/Script/Chooser.ChooserParameterFloatBase"), Category = "Data")
+	UPROPERTY(EditAnywhere, NoClear, Meta = (ExcludeBaseStruct, BaseStruct = "/Script/Chooser.ChooserParameterFloatBase"), Category = "Output")
 	FInstancedStruct OutputCost;
 
-	UE_API virtual bool HasFilters() const override;
-	UE_API virtual void Filter(FChooserEvaluationContext& Context, const FChooserIndexArray& IndexListIn, FChooserIndexArray& IndexListOut) const override;
+	virtual bool HasFilters() const override;
+	virtual void Filter(FChooserEvaluationContext& Context, const FChooserIndexArray& IndexListIn, FChooserIndexArray& IndexListOut) const override;
+	virtual bool HasCosts() const override { return true; }
+	virtual void SetOutputs(FChooserEvaluationContext& Context, int RowIndex) const override;
 
 #if WITH_EDITOR
 	virtual void Initialize(UChooserTable* OuterChooser) override;
 	
 	mutable double TestValue = 0.0;
-	UE_API virtual bool EditorTestFilter(int32 RowIndex) const override;
-	UE_API virtual float EditorTestCost(int32 RowIndex) const override;
+	virtual bool EditorTestFilter(int32 RowIndex) const override;
+	virtual float EditorTestCost(int32 RowIndex) const override;
 	
 	virtual void SetTestValue(TArrayView<const uint8> Value) override
 	{
@@ -78,8 +95,8 @@ struct FPoseSearchColumn : public FChooserColumnBase
 		Reader << TestValue;
 	}
 
-	UE_API virtual void AddToDetails(FInstancedPropertyBag& PropertyBag, int32 ColumnIndex, int32 RowIndex) override;
-	UE_API virtual void SetFromDetails(FInstancedPropertyBag& PropertyBag, int32 ColumnIndex, int32 RowIndex) override;
+	virtual void AddToDetails(FInstancedPropertyBag& PropertyBag, int32 ColumnIndex, int32 RowIndex) override;
+	virtual void SetFromDetails(FInstancedPropertyBag& PropertyBag, int32 ColumnIndex, int32 RowIndex) override;
 
 	virtual FName RowValuesPropertyName() override;
 	virtual void SetNumRows(int32 NumRows) override;
diff --git a/Engine/Plugins/Animation/PoseSearch/Source/Runtime/Public/PoseSearch/PoseSearchContext.h b/Engine/Plugins/Animation/PoseSearch/Source/Runtime/Public/PoseSearch/PoseSearchContext.h
--- a/Engine/Plugins/Animation/PoseSearch/Source/Runtime/Public/PoseSearch/PoseSearchContext.h
+++ b/Engine/Plugins/Animation/PoseSearch/Source/Runtime/Public/PoseSearch/PoseSearchContext.h
@@ -162,8 +162,7 @@ struct FSearchContext
 	UE_API FSearchContext(float InDesiredPermutationTimeOffset, const FPoseIndicesHistory* InPoseIndicesHistory,
 		const FSearchResult& InCurrentResult, const FFloatInterval& InPoseJumpThresholdTime, const FPoseSearchEvent& InEventToSearch);
 
-	UE_API FSearchContext(float InDesiredPermutationTimeOffset, const FPoseIndicesHistory* InPoseIndicesHistory,
-		const FFloatInterval& InPoseJumpThresholdTime, const FPoseSearchEvent& InEventToSearch);
+	UE_API FSearchContext(float InDesiredPermutationTimeOffset, const FFloatInterval& InPoseJumpThresholdTime, const FPoseSearchEvent& InEventToSearch);
 
 	// deleting copies and moves since members could reference other members (e.g.: ContinuingPoseValues could point to ContinuingPoseValuesData, so it'll require proper copies and movers implementations)
 	FSearchContext(const FSearchContext& Other) = delete;
@@ -226,9 +225,13 @@ struct FSearchContext
 	const FSearchResult& GetContinuingPoseSearchResult() const { return ContinuingPoseSearchResult; }
 
 	const FFloatInterval& GetPoseJumpThresholdTime() const { return PoseJumpThresholdTime; }
-	const FPoseIndicesHistory* GetPoseIndicesHistory() const { return PoseIndicesHistory; }
 	
+	UE_DEPRECATED(5.7, "use GetPoseHistory(Role)->GetPoseIndicesHistory() with the appropriate Role instead")
+	const FPoseIndicesHistory* GetPoseIndicesHistory() const;
+	
+	UE_DEPRECATED(5.7, "ArePoseHistoriesValid is no longer necessary, since PoseHistories MUST be valid since it contains necessary information for MM to function like trajectory, historical poses, etc")
 	UE_API bool ArePoseHistoriesValid() const;
+
 	TConstArrayView<const IPoseHistory*> GetPoseHistories() const { return PoseHistories; }
 	UE_API const IPoseHistory* GetPoseHistory(const FRole& Role) const;
 
@@ -310,7 +313,6 @@ private:
 	const FPoseSearchEvent EventToSearch;
 
 	const float DesiredPermutationTimeOffset = 0.f;
-	const FPoseIndicesHistory* PoseIndicesHistory = nullptr;
 	
 	// search result used to extract the continuing pose properties, except for gathering the continuing pose values, where we use ContinuingPoseValuesDatabasePoseIdx.
 	// More details in ContinuingPoseValuesDatabasePoseIdx comments
@@ -381,8 +383,6 @@ public:
 		
 	private:
 
-		// @todo: evaluate memstack usage here!
-		
 		struct FPoseCandidateHeapCompare
 		{
 			bool operator()(const FPoseCandidate& A, const FPoseCandidate& B) const
diff --git a/Engine/Plugins/Animation/PoseSearch/Source/Runtime/Public/PoseSearch/PoseSearchDatabase.h b/Engine/Plugins/Animation/PoseSearch/Source/Runtime/Public/PoseSearch/PoseSearchDatabase.h
--- a/Engine/Plugins/Animation/PoseSearch/Source/Runtime/Public/PoseSearch/PoseSearchDatabase.h
+++ b/Engine/Plugins/Animation/PoseSearch/Source/Runtime/Public/PoseSearch/PoseSearchDatabase.h
@@ -682,15 +682,15 @@ public:
 	UE_API virtual void Serialize(FArchive& Ar) override;
 	// End UObject
 	
-	// deprecate in favor of Search(UE::PoseSearch::FSearchContext& SearchContext, UE::PoseSearch::ISearchResults& SearchResults) once we settle on ISearchResults API signatures
+	// deprecate in favor of Search(UE::PoseSearch::FSearchContext& SearchContext, UE::PoseSearch::FSearchResults& SearchResults) once we settle on FSearchResults API signatures
 	UE_API UE::PoseSearch::FSearchResult Search(UE::PoseSearch::FSearchContext& SearchContext) const;
-	// deprecate in favor of SearchContinuingPose(UE::PoseSearch::FSearchContext& SearchContext, UE::PoseSearch::ISearchResults& SearchResults) once we settle on ISearchResults API signatures
+	// deprecate in favor of SearchContinuingPose(UE::PoseSearch::FSearchContext& SearchContext, UE::PoseSearch::FSearchResults& SearchResults) once we settle on FSearchResults API signatures
 	UE_API UE::PoseSearch::FSearchResult SearchContinuingPose(UE::PoseSearch::FSearchContext& SearchContext) const;
 	
 	// Experimental, this feature might be removed without warning, not for production use
-	UE_API void Search(UE::PoseSearch::FSearchContext& SearchContext, UE::PoseSearch::ISearchResults& SearchResults) const;
+	UE_API void Search(UE::PoseSearch::FSearchContext& SearchContext, UE::PoseSearch::FSearchResults& SearchResults) const;
 	// Experimental, this feature might be removed without warning, not for production use
-	UE_API void SearchContinuingPose(UE::PoseSearch::FSearchContext& SearchContext, UE::PoseSearch::ISearchResults& SearchResults) const;
+	UE_API void SearchContinuingPose(UE::PoseSearch::FSearchContext& SearchContext, UE::PoseSearch::FSearchResults& SearchResults) const;
 
 	UE_API bool Contains(const UObject* Object) const;
 
@@ -733,10 +733,10 @@ public:
 #endif // WITH_EDITORONLY_DATA
 
 private:
-	UE_API void SearchPCAKDTree(UE::PoseSearch::FSearchContext& SearchContext, UE::PoseSearch::ISearchResults& SearchResults) const;
-	UE_API void SearchVPTree(UE::PoseSearch::FSearchContext& SearchContext, UE::PoseSearch::ISearchResults& SearchResults) const;
-	UE_API void SearchBruteForce(UE::PoseSearch::FSearchContext& SearchContext, UE::PoseSearch::ISearchResults& SearchResults) const;
-	UE_API void SearchEvent(UE::PoseSearch::FSearchContext& SearchContext, UE::PoseSearch::ISearchResults& SearchResults) const;
+	UE_API void SearchPCAKDTree(UE::PoseSearch::FSearchContext& SearchContext, UE::PoseSearch::FSearchResults& SearchResults) const;
+	UE_API void SearchVPTree(UE::PoseSearch::FSearchContext& SearchContext, UE::PoseSearch::FSearchResults& SearchResults) const;
+	UE_API void SearchBruteForce(UE::PoseSearch::FSearchContext& SearchContext, UE::PoseSearch::FSearchResults& SearchResults) const;
+	UE_API void SearchEvent(UE::PoseSearch::FSearchContext& SearchContext, UE::PoseSearch::FSearchResults& SearchResults) const;
 
 	typedef TArray<int32, TInlineAllocator<256, TMemStackAllocator<>>> FSelectableAssetIdx;
 	UE_API void PopulateSelectableAssetIdx(FSelectableAssetIdx& SelectableAssetIdx, TConstArrayView<const UObject*> AssetsToConsider) const;
diff --git a/Engine/Plugins/Animation/PoseSearch/Source/Runtime/Public/PoseSearch/PoseSearchHistory.h b/Engine/Plugins/Animation/PoseSearch/Source/Runtime/Public/PoseSearch/PoseSearchHistory.h
--- a/Engine/Plugins/Animation/PoseSearch/Source/Runtime/Public/PoseSearch/PoseSearchHistory.h
+++ b/Engine/Plugins/Animation/PoseSearch/Source/Runtime/Public/PoseSearch/PoseSearchHistory.h
@@ -31,6 +31,31 @@ typedef uint16 FComponentSpaceTransformIndex;
 typedef TPair<FBoneIndexType, FComponentSpaceTransformIndex> FBoneToTransformPair;
 typedef TMap<FBoneIndexType, FComponentSpaceTransformIndex> FBoneToTransformMap;
 
+struct FHistoricalPoseIndex
+{
+	bool operator==(const FHistoricalPoseIndex& Index) const
+	{
+		return PoseIndex == Index.PoseIndex && DatabaseKey == Index.DatabaseKey;
+	}
+
+	friend inline uint32 GetTypeHash(const FHistoricalPoseIndex& Index)
+	{
+		return HashCombineFast(::GetTypeHash(Index.PoseIndex), GetTypeHash(Index.DatabaseKey));
+	}
+
+	int32 PoseIndex = INDEX_NONE;
+	FObjectKey DatabaseKey;
+};
+
+struct FPoseIndicesHistory
+{
+	UE_INTERNAL UE_API void Update(const FSearchResult& SearchResult, float DeltaTime, float MaxTime);
+	void Reset() { IndexToTime.Reset(); }
+	bool operator==(const FPoseIndicesHistory& Other) const;
+		
+	TMap<FHistoricalPoseIndex, float> IndexToTime;
+};
+
 struct IComponentSpacePoseProvider
 {
 	virtual ~IComponentSpacePoseProvider() {}
@@ -109,6 +134,7 @@ struct IPoseHistory : public TSharedFromThis<IPoseHistory, ESPMode::ThreadSafe>
 	virtual const TConstArrayView<FName> GetCollectedCurves() const = 0;
 	virtual int32 GetNumEntries() const = 0;
 	virtual const FPoseHistoryEntry& GetEntry(int32 EntryIndex) const = 0;
+	virtual const FPoseIndicesHistory* GetPoseIndicesHistory() const { return nullptr; }
 
 #if ENABLE_DRAW_DEBUG && ENABLE_ANIM_DEBUG
 	virtual void DebugDraw(const UWorld* World, FColor Color) const = 0;
@@ -173,13 +199,13 @@ struct FPoseHistory : public IPoseHistory
 	UE_API void EvaluateComponentSpace_AnyThread(float DeltaTime, IComponentSpacePoseProvider& ComponentSpacePoseProvider, bool bStoreScales,
 		float RootBoneRecoveryTime, float RootBoneTranslationRecoveryRatio, float RootBoneRotationRecoveryRatio,
 		bool bNeedsReset, bool bCacheBones, const TArray<FBoneIndexType>& RequiredBones, 
-		const FBlendedCurve& Curves = FBlendedCurve(), const TConstArrayView<FName>& CollectedCurves = TConstArrayView<FName>());
+		const FBlendedCurve& Curves = FBlendedCurve(), const TConstArrayView<FName>& InCollectedCurves = TConstArrayView<FName>());
 
 	UE_DEPRECATED(5.6, "Use other EvaluateComponentSpace_AnyThread signatures instead")
 	UE_API void EvaluateComponentSpace_AnyThread(float DeltaTime, FCSPose<FCompactPose>& ComponentSpacePose, bool bStoreScales,
 		float RootBoneRecoveryTime, float RootBoneTranslationRecoveryRatio, float RootBoneRotationRecoveryRatio,
 		bool bNeedsReset, bool bCacheBones, const TArray<FBoneIndexType>& RequiredBones,
-		const FBlendedCurve& Curves = FBlendedCurve(), const TConstArrayView<FName>& CollectedCurves = TConstArrayView<FName>());
+		const FBlendedCurve& Curves = FBlendedCurve(), const TConstArrayView<FName>& InCollectedCurves = TConstArrayView<FName>());
 		
 	// IPoseHistory interface
 	UE_API virtual bool GetTransformAtTime(float Time, FTransform& OutBoneTransform, const USkeleton* BoneIndexSkeleton = nullptr, FBoneIndexType BoneIndexType = RootBoneIndexType, FBoneIndexType ReferenceBoneIndexType = ComponentSpaceIndexType, bool bExtrapolate = false) const override;
@@ -202,6 +228,7 @@ struct FPoseHistory : public IPoseHistory
 	
 	virtual int32 GetNumEntries() const override;
 	UE_API virtual const FPoseHistoryEntry& GetEntry(int32 EntryIndex) const override;
+	UE_API virtual const FPoseIndicesHistory* GetPoseIndicesHistory() const override;
 
 #if ENABLE_DRAW_DEBUG && ENABLE_ANIM_DEBUG
 	virtual void DebugDraw(const UWorld* World, FColor Color) const override { unimplemented(); }
@@ -225,24 +252,24 @@ private:
 	// @todo: deprecate this member and expose it via blue print logic or as global query scaling multiplier
 	float TrajectorySpeedMultiplier = 1.f;
 
-	struct FPoseData
-	{
-		// skeleton from the last Update, to keep tracking skeleton changes, and support compatible skeletons
-		TWeakObjectPtr<const USkeleton> LastUpdateSkeleton;
+	// skeleton from the last Update, to keep tracking skeleton changes, and support compatible skeletons
+	TWeakObjectPtr<const USkeleton> LastUpdateSkeleton;
 
-		// Map from skeleton bone indexes (encoded as FBoneIndexType) to internal pose history transform index (FComponentSpaceTransformIndex).
-		// If Empty all the bones get collected
-		FBoneToTransformMap BoneToTransformMap;
+	// Map from skeleton bone indexes (encoded as FBoneIndexType) to internal pose history transform index (FComponentSpaceTransformIndex).
+	// If Empty all the bones get collected
+	FBoneToTransformMap BoneToTransformMap;
 		
-		// list of curves that we want to collect into our history.
-		TArray<FName> CollectedCurves;
+	// list of curves that we want to collect into our history.
+	TArray<FName> CollectedCurves;
+
+	// GetTypeHash for BoneToTransformMap
+	uint32 BoneToTransformMapTypeHash = 0;
 
-		// GetTypeHash for BoneToTransformMap
-		uint32 BoneToTransformMapTypeHash = 0;
+	// ring buffer of collected bones
+	TRingBuffer<FPoseHistoryEntry> Entries;
 
-		// ring buffer of collected bones
-		TRingBuffer<FPoseHistoryEntry> Entries;
-	} PoseData;
+	// datra structure containing the previously played pose indexes with their respective datababases
+	FPoseIndicesHistory PoseIndicesHistory;
 
 	UE_MT_DECLARE_TS_RW_ACCESS_DETECTOR(PoseDataThreadSafeCounter);
 };
@@ -277,6 +304,7 @@ struct FMemStackPoseHistory : public IPoseHistory
 	
 	UE_API virtual int32 GetNumEntries() const override;
 	UE_API virtual const FPoseHistoryEntry& GetEntry(int32 EntryIndex) const override;
+	UE_API virtual const FPoseIndicesHistory* GetPoseIndicesHistory() const override;
 	
 #if ENABLE_DRAW_DEBUG && ENABLE_ANIM_DEBUG
 	virtual void DebugDraw(const UWorld* World, FColor Color) const override { unimplemented(); }
@@ -302,31 +330,6 @@ private:
 	TArray<FPoseHistoryEntry, TInlineAllocator<4, TMemStackAllocator<>>> FutureEntries;
 };
 
-struct FHistoricalPoseIndex
-{
-	bool operator==(const FHistoricalPoseIndex& Index) const
-	{
-		return PoseIndex == Index.PoseIndex && DatabaseKey == Index.DatabaseKey;
-	}
-
-	friend inline uint32 GetTypeHash(const FHistoricalPoseIndex& Index)
-	{
-		return HashCombineFast(::GetTypeHash(Index.PoseIndex), GetTypeHash(Index.DatabaseKey));
-	}
-
-	int32 PoseIndex = INDEX_NONE;
-	FObjectKey DatabaseKey;
-};
-
-struct FPoseIndicesHistory
-{
-	UE_INTERNAL UE_API void Update(const FSearchResult& SearchResult, float DeltaTime, float MaxTime);
-	void Reset() { IndexToTime.Reset(); }
-	bool operator==(const FPoseIndicesHistory& Other) const;
-		
-	TMap<FHistoricalPoseIndex, float> IndexToTime;
-};
-
 } // namespace UE::PoseSearch
 
 USTRUCT(BlueprintType)
diff --git a/Engine/Plugins/Animation/PoseSearch/Source/Runtime/Public/PoseSearch/PoseSearchLibrary.h b/Engine/Plugins/Animation/PoseSearch/Source/Runtime/Public/PoseSearch/PoseSearchLibrary.h
--- a/Engine/Plugins/Animation/PoseSearch/Source/Runtime/Public/PoseSearch/PoseSearchLibrary.h
+++ b/Engine/Plugins/Animation/PoseSearch/Source/Runtime/Public/PoseSearch/PoseSearchLibrary.h
@@ -98,6 +98,7 @@ PRAGMA_ENABLE_DEPRECATION_WARNINGS
 	UE_DEPRECATED(5.6, "Derive this value via 'CurrentSearchResult.IsValid() && !CurrentSearchResult.bIsContinuingPoseSearch' if needed")
 	bool bJumpedToPose = false;
 
+	UE_DEPRECATED(5.7, "this property has been moved into the pose history. Use IPoseHistory::GetPoseIndicesHistory to access it")
 	UE::PoseSearch::FPoseIndicesHistory PoseIndicesHistory;
 };
 
@@ -161,7 +162,7 @@ class UPoseSearchLibrary : public UBlueprintFunctionLibrary
 
 #if UE_POSE_SEARCH_TRACE_ENABLED
 
-	// deprecate in favor of TraceMotionMatching using ISearchResults once we settle on ISearchResults API signatures
+	// deprecate in favor of TraceMotionMatching using FSearchResults once we settle on FSearchResults API signatures
 	static void TraceMotionMatching(
 		UE::PoseSearch::FSearchContext& SearchContext,
 		const UE::PoseSearch::FSearchResult& SearchResult,
@@ -174,7 +175,7 @@ class UPoseSearchLibrary : public UBlueprintFunctionLibrary
 	// Experimental, this feature might be removed without warning, not for production use
 	static void TraceMotionMatching(
 		UE::PoseSearch::FSearchContext& SearchContext,
-		const UE::PoseSearch::ISearchResults& SearchResults,
+		const UE::PoseSearch::FSearchResults& SearchResults,
 		float ElapsedPoseSearchTime,
 		float WantedPlayRate,
 		EPoseSearchInterruptMode InterruptMode);
@@ -304,7 +305,7 @@ public:
 		const FPoseSearchFutureProperties& Future,
 		FPoseSearchBlueprintResult& Result);
 
-	// deprecate in favor of MotionMatch with the ISearchResults parameter once we settle on ISearchResults API signatures
+	// deprecate in favor of MotionMatch with the FSearchResults parameter once we settle on FSearchResults API signatures
 	static UE_API UE::PoseSearch::FSearchResult MotionMatch(
 		const TArrayView<const UObject*> AnimContexts,
 		const TArrayView<const UE::PoseSearch::FRole> Roles,
@@ -323,7 +324,7 @@ public:
 		const FPoseSearchContinuingProperties& ContinuingProperties,
 		const FPoseSearchFutureProperties& Future,
 		const FPoseSearchEvent& EventToSearch,
-		UE::PoseSearch::ISearchResults& SearchResults);
+		UE::PoseSearch::FSearchResults& SearchResults);
 
 	// Experimental, this feature might be removed without warning, not for production use
 	static UE_API void MotionMatch(
@@ -334,7 +335,7 @@ public:
 		const FPoseSearchContinuingProperties& ContinuingProperties,
 		const FPoseSearchFutureProperties& Future,
 		const FPoseSearchEvent& EventToSearch,
-		UE::PoseSearch::ISearchResults& SearchResults);
+		UE::PoseSearch::FSearchResults& SearchResults);
 			
 	// Experimental, this feature might be removed without warning, not for production use
 	static UE_API void MotionMatch(
@@ -345,7 +346,7 @@ public:
 		const FPoseSearchContinuingProperties& ContinuingProperties,
 		const float DesiredPermutationTimeOffset,
 		const FPoseSearchEvent& EventToSearch,
-		UE::PoseSearch::ISearchResults& SearchResults);
+		UE::PoseSearch::FSearchResults& SearchResults);
 
 	// Experimental, this feature might be removed without warning, not for production use
 	static UE_API void MotionMatch(
@@ -356,14 +357,14 @@ public:
 		const FPoseSearchContinuingProperties& ContinuingProperties,
 		const float DesiredPermutationTimeOffset,
 		const FPoseSearchEvent& EventToSearch,
-		UE::PoseSearch::ISearchResults& SearchResults);
+		UE::PoseSearch::FSearchResults& SearchResults);
 
 	// Experimental, this feature might be removed without warning, not for production use
 	static UE_API void MotionMatch(
 		UE::PoseSearch::FSearchContext& SearchContext,
 		const TArrayView<const UObject*> AssetsToSearch,
 		const FPoseSearchContinuingProperties& ContinuingProperties,
-		UE::PoseSearch::ISearchResults& SearchResults);
+		UE::PoseSearch::FSearchResults& SearchResults);
 
 	static UE_API const FAnimNode_PoseSearchHistoryCollector_Base* FindPoseHistoryNode(
 		const FName PoseHistoryName,
diff --git a/Engine/Plugins/Animation/PoseSearch/Source/Runtime/Public/PoseSearch/PoseSearchResult.h b/Engine/Plugins/Animation/PoseSearch/Source/Runtime/Public/PoseSearch/PoseSearchResult.h
--- a/Engine/Plugins/Animation/PoseSearch/Source/Runtime/Public/PoseSearch/PoseSearchResult.h
+++ b/Engine/Plugins/Animation/PoseSearch/Source/Runtime/Public/PoseSearch/PoseSearchResult.h
@@ -31,11 +31,11 @@ struct FDatabasePoseIdx
 	bool IsValid() const { return PoseIdx != INDEX_NONE && Database != nullptr; }
 	bool operator==(const FDatabasePoseIdx& Other) const { return PoseIdx == Other.PoseIdx && Database == Other.Database; }
 
-	POSESEARCH_API const FSearchIndexAsset* GetSearchIndexAsset(bool bMandatory = false) const;
+	UE_API const FSearchIndexAsset* GetSearchIndexAsset(bool bMandatory = false) const;
 	// Experimental, this feature might be removed without warning, not for production use
-	POSESEARCH_API const UAnimationAsset* GetCurrentResultAnimationAsset() const;
+	UE_API const UAnimationAsset* GetCurrentResultAnimationAsset() const;
 	// Experimental, this feature might be removed without warning, not for production use
-	POSESEARCH_API const UAnimationAsset* GetCurrentResultAnimationAsset(const FRole& Role) const;
+	UE_API const UAnimationAsset* GetCurrentResultAnimationAsset(const FRole& Role) const;
 };
 
 struct FSearchResult : FDatabasePoseIdx
@@ -119,44 +119,50 @@ struct FSearchResult : FDatabasePoseIdx
 };
 
 // Experimental, this feature might be removed without warning, not for production use
-struct ISearchResults
+struct FSearchResults
 {
-	virtual ~ISearchResults() {}
-	virtual void UpdateWith(const FSearchResult& SearchResult) = 0;
-	virtual void FinalizeResults() = 0;
-	virtual FSearchResult GetBestResult() const = 0;
+	FSearchResults()
+		: UpdateWith([](const FSearchResult& SearchResult) {})
+		, FinalizeResults([]() {})
+		, GetBestResult([]() { return FSearchResult(); })
+		, ShouldPerformSearch([](float) { return true; })
+		, IterateOverSearchResults([](const TFunctionRef<bool(const FSearchResult& SearchResult)>) { return true; })
+	{
+	}
 
+	virtual ~FSearchResults() {}
+
+	FSearchResults(const FSearchResults& Other) = delete;
+	FSearchResults(FSearchResult&& Other) = delete;
+	FSearchResults& operator=(const FSearchResults& Other) = delete;
+	FSearchResults& operator=(FSearchResults&& Other) = delete;
+	
+	TFunction<void(const FSearchResult& SearchResult)> UpdateWith;
+	TFunction<void()> FinalizeResults;
+	TFunction<FSearchResult()> GetBestResult;
+	
 	// every search will ask FSearchResults if it should perform the search, knowing that its minimum cost will be SearchMinimumCost:
 	// it's up to this method to perform early out optimizations if the holded results are already good enough fo skip the search
-	virtual bool ShouldPerformSearch(float SearchMinimumCost) const = 0;
+	TFunction<bool(float)> ShouldPerformSearch;
 
-	// calls IterateOverSearchResultsFunction for each FSearchResult this ISearchResults is holding
+	// calls IterateOverSearchResultsFunction for each FSearchResult this FSearchResults is holding
 	// is IterateOverSearchResultsFunction returns true, the iteration will stop and IterateOverSearchResults will return true otherwise will return false (useful to implement contains / find methods)
-	virtual bool IterateOverSearchResults(const TFunctionRef<bool(const FSearchResult& SearchResult)> IterateOverSearchResultsFunction) const = 0;
+	TFunction<bool(const TFunctionRef<bool(const FSearchResult& SearchResult)>)> IterateOverSearchResults;
 };
 
 // Experimental, this feature might be removed without warning, not for production use
-struct FSearchResults_Single : public ISearchResults
+struct FSearchResults_Single : public FSearchResults
 {
-	POSESEARCH_API virtual void UpdateWith(const FSearchResult& SearchResult) override;
-	POSESEARCH_API virtual void FinalizeResults() override;
-	POSESEARCH_API virtual FSearchResult GetBestResult() const override;
-	POSESEARCH_API virtual bool ShouldPerformSearch(float SearchMinimumCost) const override;
-	POSESEARCH_API virtual bool IterateOverSearchResults(const TFunctionRef<bool(const FSearchResult& SearchResult)> IterateOverSearchResultsFunction) const override;
+	UE_API FSearchResults_Single();
 
 private:
 	FSearchResult SingleSearchResult;
 };
 
 // Experimental, this feature might be removed without warning, not for production use
-struct FSearchResults_Multi : public ISearchResults
+struct FSearchResults_Multi : public FSearchResults
 {
-	POSESEARCH_API explicit FSearchResults_Multi(int32 InMaxNumberOfWantedResults = 1);
-	POSESEARCH_API virtual void UpdateWith(const FSearchResult& SearchResult) override;
-	POSESEARCH_API virtual void FinalizeResults() override;
-	POSESEARCH_API virtual FSearchResult GetBestResult() const override;
-	POSESEARCH_API virtual bool ShouldPerformSearch(float SearchMinimumCost) const override;
-	POSESEARCH_API virtual bool IterateOverSearchResults(const TFunctionRef<bool(const FSearchResult& SearchResult)> IterateOverSearchResultsFunction) const override;
+	explicit FSearchResults_Multi(int32 InMaxNumberOfWantedResults = 1);
 
 private:
 	// @todo: evaluate a max heap instead (using AlgoImpl::HeapSiftUp / AlgoImpl::HeapSiftDown)
@@ -166,18 +172,19 @@ private:
 };
 
 // Experimental, this feature might be removed without warning, not for production use
-struct FSearchResults_AssetBests : public ISearchResults
+struct FSearchResults_AssetBests : public FSearchResults
 {
-	POSESEARCH_API virtual void UpdateWith(const FSearchResult& SearchResult) override;
-	POSESEARCH_API virtual void FinalizeResults() override;
-	POSESEARCH_API virtual FSearchResult GetBestResult() const override;
-	POSESEARCH_API virtual bool ShouldPerformSearch(float SearchMinimumCost) const override;
-	POSESEARCH_API virtual bool IterateOverSearchResults(const TFunctionRef<bool(const FSearchResult& SearchResult)> IterateOverSearchResultsFunction) const override;
-
-	typedef TMap<int32, FSearchResult, TInlineSetAllocator<16, TMemStackSetAllocator<>>> FMap;
-	const FMap& GetSourceAssetIdxToBestSearchResult() const { return SourceAssetIdxToBestSearchResult; }
+	FSearchResults_AssetBests();
+	const FSearchResult* FindSearchResultFor(const UPoseSearchDatabase* Database, int32 SourceAssetIdx) const;
+	
+	// shrinks FSearchResults_AssetBests down to MaxNumResults preserving the best results (the ones with less cost)
+	void Shrink(int32 MaxNumResults);
+
+	typedef TMap<int32, FSearchResult, TInlineSetAllocator<16, TMemStackSetAllocator<>>> FSourceAssetIdxToBestSearchResult;
+	typedef TMap<const UPoseSearchDatabase*, FSourceAssetIdxToBestSearchResult, TInlineSetAllocator<4, TMemStackSetAllocator<>>> FPerDatabaseSourceAssetIdxToBestSearchResult;
+
 private:
-	FMap SourceAssetIdxToBestSearchResult;
+	FPerDatabaseSourceAssetIdxToBestSearchResult PerDatabaseSourceAssetIdxToBestSearchResult;
 };
 
 } // namespace UE::PoseSearch
diff --git a/Engine/Plugins/Experimental/UAF/UAFPoseSearch/Source/UAFPoseSearch/Private/MotionMatchingTrait.cpp b/Engine/Plugins/Experimental/UAF/UAFPoseSearch/Source/UAFPoseSearch/Private/MotionMatchingTrait.cpp
--- a/Engine/Plugins/Experimental/UAF/UAFPoseSearch/Source/UAFPoseSearch/Private/MotionMatchingTrait.cpp
+++ b/Engine/Plugins/Experimental/UAF/UAFPoseSearch/Source/UAFPoseSearch/Private/MotionMatchingTrait.cpp
@@ -53,6 +53,7 @@ namespace UE::UAF
 {
 	AUTO_REGISTER_ANIM_TRAIT(FMotionMatchingTrait)
 
+PRAGMA_DISABLE_DEPRECATION_WARNINGS
 #define TRAIT_INTERFACE_ENUMERATOR(GeneratorMacro) \
 		GeneratorMacro(IUpdate) \
 		GeneratorMacro(IEvaluate) \
@@ -60,6 +61,7 @@ namespace UE::UAF
 	// Trait implementation boilerplate
 	GENERATE_ANIM_TRAIT_IMPLEMENTATION(FMotionMatchingTrait, TRAIT_INTERFACE_ENUMERATOR, NULL_ANIM_TRAIT_INTERFACE_ENUMERATOR, NULL_ANIM_TRAIT_EVENT_ENUMERATOR)
 #undef TRAIT_INTERFACE_ENUMERATOR
+PRAGMA_ENABLE_DEPRECATION_WARNINGS
 
 	void FMotionMatchingTrait::FInstanceData::Construct(const FExecutionContext& Context, const FTraitBinding& Binding)
 	{
@@ -108,7 +110,7 @@ namespace UE::UAF
 		TTraitBinding<IPoseHistory> PoseHistoryTrait;
 		if (!Context.GetScopedInterface<IPoseHistory>(PoseHistoryTrait))
 		{
-			UE_LOG(LogPoseSearch, Error, TEXT("FMotionMatchingTrait::PreUpdate, missing IPoseHistory"));
+			UE_LOG(LogPoseSearch, Error, TEXT("FMotionMatchingTrait::PreUpdate, missing UE::UAF::IPoseHistory"));
 			return;
 		}
 
@@ -122,7 +124,7 @@ namespace UE::UAF
 		const UE::PoseSearch::IPoseHistory* PoseHistory = PoseHistoryTrait.GetPoseHistory(Context);
 		if (!PoseHistory)
 		{
-			UE_LOG(LogPoseSearch, Error, TEXT("FMotionMatchingTrait::PreUpdate, missing IPoseHistory"));
+			UE_LOG(LogPoseSearch, Error, TEXT("FMotionMatchingTrait::PreUpdate, missing UE::PoseSearch::IPoseHistory"));
 			return;
 		}
 
@@ -274,7 +276,7 @@ namespace UE::UAF
 				const float SearchThrottleTime = SharedData->GetbShouldSearch(Binding) ? SharedData->GetSearchThrottleTime(Binding) : UE_BIG_NUMBER;
 
 				FPoseSearchEvent EventToSearch = SharedData->GetEventToSearch(Binding);
-				FSearchContext SearchContext(0.f, &MotionMatchingState.PoseIndicesHistory, SharedData->GetPoseJumpThresholdTime(Binding), EventToSearch.GetPlayRateOverriddenEvent(PoseSearchPlayRate));
+				FSearchContext SearchContext(0.f, SharedData->GetPoseJumpThresholdTime(Binding), EventToSearch.GetPlayRateOverriddenEvent(PoseSearchPlayRate));
 
 				FChooserEvaluationContext ChooserEvaluationContext(const_cast<UObject*>(AnimContext));
 				const FUAFAssetInstance& Instance = Binding.GetTraitPtr().GetNodeInstance()->GetOwner();
@@ -322,7 +324,12 @@ namespace UE::UAF
 
 				const float WantedPlayRate = CalculateWantedPlayRate(SearchResult, SearchContext, PoseSearchPlayRate, PoseHistory ? PoseHistory->GetTrajectorySpeedMultiplier() : 1.f, EventToSearch);
 				MotionMatchingState.SearchResult.InitFrom(SearchResult, WantedPlayRate);
-				MotionMatchingState.PoseIndicesHistory.Update(SearchResult, DeltaTime, SharedData->GetPoseReselectHistory(Binding));
+
+				if (const FPoseIndicesHistory* PoseIndicesHistory = PoseHistory->GetPoseIndicesHistory())
+				{
+					// const casting here is safe since we're in the thread owning the pose history, and it's the correct place to update the previously selected poses
+					const_cast<FPoseIndicesHistory*>(PoseIndicesHistory)->Update(SearchResult, DeltaTime, SharedData->GetPoseReselectHistory(Binding));
+				}
 
 				PRAGMA_DISABLE_DEPRECATION_WARNINGS
 				MotionMatchingState.CurrentSearchResult = SearchResult;
