PoseSearch + Chooser - fix for PoseSearchColumn copy/paste/move operations

diff --git a/Engine/Plugins/Animation/PoseSearch/Source/Runtime/Private/Chooser/PoseSearchChooserColumn.cpp b/Engine/Plugins/Animation/PoseSearch/Source/Runtime/Private/Chooser/PoseSearchChooserColumn.cpp
--- a/Engine/Plugins/Animation/PoseSearch/Source/Runtime/Private/Chooser/PoseSearchChooserColumn.cpp
+++ b/Engine/Plugins/Animation/PoseSearch/Source/Runtime/Private/Chooser/PoseSearchChooserColumn.cpp
@@ -137,8 +137,8 @@ void FPoseSearchColumn::InsertRows(int Index, int Count)
 		for (int i = 0; i < Count; i++)
 		{
 			FPoseSearchDatabaseAnimationAsset NewAsset;
-			NewAsset.AnimAsset = GetReferencedObject(Index);
-			InternalDatabase->InsertAnimationAssetAt(NewAsset, Index);
+			NewAsset.AnimAsset = GetReferencedObject(Index + i);
+			InternalDatabase->InsertAnimationAssetAt(NewAsset, Index + i);
 		}
 	}
 
@@ -186,18 +186,18 @@ void FPoseSearchColumn::MoveRow(int SourceRowIndex, int TargetRowIndex)
 void FPoseSearchColumn::CopyRow(FChooserColumnBase& SourceColumn, int SourceRowIndex, int TargetRowIndex)
 {
 	check(InternalDatabase);
-	if (SourceRowIndex != TargetRowIndex)
+
+	const FPoseSearchColumn& SourcePoseSearchColumn = static_cast<FPoseSearchColumn&>(SourceColumn);
+
+	if (const FPoseSearchDatabaseAnimationAsset* SourceAsset = SourcePoseSearchColumn.InternalDatabase->GetDatabaseAnimationAsset(SourceRowIndex))
 	{
-		if (const FPoseSearchDatabaseAnimationAsset* SourceAsset = InternalDatabase->GetDatabaseAnimationAsset(SourceRowIndex))
+		if (FPoseSearchDatabaseAnimationAsset* TargetAsset = InternalDatabase->GetMutableDatabaseAnimationAsset(TargetRowIndex))
 		{
-			if (FPoseSearchDatabaseAnimationAsset* TargetAsset = InternalDatabase->GetMutableDatabaseAnimationAsset(TargetRowIndex))
-			{
-				FScopedTransaction ScopedTransaction(LOCTEXT("CopyRow","CopyRow FPoseSearchColumn"));
-				InternalDatabase->Modify();
-				*TargetAsset = *SourceAsset;
+			FScopedTransaction ScopedTransaction(LOCTEXT("CopyRow","CopyRow FPoseSearchColumn"));
+			InternalDatabase->Modify();
+			*TargetAsset = *SourceAsset;
 
-				InternalDatabase->NotifySynchronizeWithExternalDependencies();
-			}
+			InternalDatabase->NotifySynchronizeWithExternalDependencies();
 		}
 	}
 
@@ -235,7 +235,12 @@ float FPoseSearchColumn::EditorTestCost(int32 RowIndex) const
 
 FChooserParameterBase* FPoseSearchColumn::GetInputValue()
 {
-	return InputValue.GetMutablePtr<FChooserParameterBase>();
+	if (FChooserParameterBase* ChooserParameterBase = InputValue.GetMutablePtr<FChooserParameterBase>())
+	{
+		return ChooserParameterBase;
+	}
+
+	return &NullInputValue;
 }
 
 bool FPoseSearchColumn::HasFilters() const
@@ -389,7 +394,6 @@ void FPoseSearchColumn::SetOutputs(FChooserEvaluationContext& Context, int RowIn
 
 	for (const FSearchResult& SearchResult : ScratchArea.SearchResults)
 	{
-		check(SearchResult.IsValid() && SearchResult.Database == InternalDatabase);
 		if (const FSearchIndexAsset* SearchIndexAsset = SearchResult.GetSearchIndexAsset())
 		{
 			const int32 SearchResultSourceAssetIdx = SearchIndexAsset->GetSourceAssetIdx();
diff --git a/Engine/Plugins/Animation/PoseSearch/Source/Runtime/Public/PoseSearch/Chooser/PoseSearchChooserColumn.h b/Engine/Plugins/Animation/PoseSearch/Source/Runtime/Public/PoseSearch/Chooser/PoseSearchChooserColumn.h
--- a/Engine/Plugins/Animation/PoseSearch/Source/Runtime/Public/PoseSearch/Chooser/PoseSearchChooserColumn.h
+++ b/Engine/Plugins/Animation/PoseSearch/Source/Runtime/Public/PoseSearch/Chooser/PoseSearchChooserColumn.h
@@ -34,6 +34,17 @@ struct FPoseHistoryContextProperty :  public FChooserParameterPoseHistoryBase
 	CHOOSER_PARAMETER_BOILERPLATE();
 };
 
+USTRUCT(Experimental)
+struct FPoseSearchColumnNullInputValue : public FChooserParameterBase
+{
+	GENERATED_BODY()
+
+	virtual void GetDisplayName(FText& OutName) const override
+	{
+		OutName = FText::FromString("FPoseSearchColumn");
+	}
+};
+
 USTRUCT(DisplayName = "Pose Match", Meta = (Experimental, Category = "Experimental", Tooltip = "This column filters out all assets except the one which is selected by motion matching query.  Results must be AnimationAssets with a PoseSearchBranchIn notify state.  It also outputs OutputStartTime to specify the frame which matched pose best.  To work as intended it must be placed last (furthest right) in the Chooser so that other filters are applied first."))
 struct FPoseSearchColumn : public FChooserColumnBase
 {
@@ -65,6 +76,10 @@ struct FPoseSearchColumn : public FChooserColumnBase
 	UPROPERTY(EditAnywhere, DisplayName = "Pose History", NoClear, Meta = (ExcludeBaseStruct, BaseStruct = "/Script/PoseSearch.ChooserParameterPoseHistoryBase"), Category = "Data")
 	FInstancedStruct InputValue;
 	
+	// @todo: do we really need this? Currently used if InputValue is not set to be able to make FChooserTableEditor::PasteInternal work
+	UPROPERTY()
+	FPoseSearchColumnNullInputValue NullInputValue;
+
 	// Float output for the start time with the best matching pose
 	UPROPERTY(EditAnywhere, NoClear, Meta = (ExcludeBaseStruct, BaseStruct = "/Script/Chooser.ChooserParameterFloatBase"), Category = "Output")
 	FInstancedStruct OutputStartTime;
diff --git a/Engine/Plugins/Chooser/Source/ChooserEditor/Private/ChooserTableEditor.cpp b/Engine/Plugins/Chooser/Source/ChooserEditor/Private/ChooserTableEditor.cpp
--- a/Engine/Plugins/Chooser/Source/ChooserEditor/Private/ChooserTableEditor.cpp
+++ b/Engine/Plugins/Chooser/Source/ChooserEditor/Private/ChooserTableEditor.cpp
@@ -2314,17 +2314,6 @@ UChooserTable* FChooserTableEditor::CopySelectionInternal()
 		TArray<UChooserRowDetails*> SelectedRowsCopy = SelectedRows;
 		Algo::Sort(SelectedRowsCopy, [](UChooserRowDetails *A, UChooserRowDetails *B) { return A->Row<B->Row;});
 		
-		// copy all columns
-		CopyData->ColumnsStructs = Chooser->ColumnsStructs;
-		
-		// clear all column's cell data
-		for (FInstancedStruct& ColumnData : CopyData->ColumnsStructs)
-		{
-			FChooserColumnBase& Column = ColumnData.GetMutable<FChooserColumnBase>();
-			Column.SetNumRows(0);
-			Column.SetNumRows(SelectedRowsCopy.Num());
-		}
-		
 		if (SelectedRowsCopy.Num() > 0 && SelectedRowsCopy[0]->Row == ColumnWidget_SpecialIndex_Fallback)
 		{
 			SelectedRowsCopy.RemoveAt(0);
@@ -2353,9 +2342,28 @@ UChooserTable* FChooserTableEditor::CopySelectionInternal()
 					CopiedNestedChooser->Chooser = DuplicateNestedChooser(CopiedNestedChooser->Chooser, CopyData);
 				}
 			}
-			
+
 			CopyData->DisabledRows[RowIndex] = Chooser->DisabledRows[SelectedRowsCopy[RowIndex]->Row];
-			
+		}
+
+		// copy all columns
+		CopyData->ColumnsStructs = Chooser->ColumnsStructs;
+		for (FInstancedStruct& ColumnData : CopyData->ColumnsStructs)
+		{
+			FChooserColumnBase& Column = ColumnData.GetMutable<FChooserColumnBase>();
+			Column.Initialize(CopyData);
+		}
+
+		// clear all column's cell data
+		for (FInstancedStruct& ColumnData : CopyData->ColumnsStructs)
+		{
+			FChooserColumnBase& Column = ColumnData.GetMutable<FChooserColumnBase>();
+			Column.SetNumRows(0);
+			Column.SetNumRows(SelectedRowsCopy.Num());
+		}
+
+		for (int RowIndex = 0; RowIndex < SelectedRowsCopy.Num(); RowIndex++)
+		{
 			for (int ColumnIndex = 0; ColumnIndex < CopyData->ColumnsStructs.Num(); ColumnIndex++)
 			{
 				FChooserColumnBase& SourceColumn = Chooser->ColumnsStructs[ColumnIndex].GetMutable<FChooserColumnBase>();
