PoseSearch - fix for crash when opening the PoseSearchInteractionAsset editor

diff --git a/Engine/Plugins/Animation/PoseSearch/Source/Editor/Public/PoseSearchInteractionAssetEditor.h b/Engine/Plugins/Animation/PoseSearch/Source/Editor/Public/PoseSearchInteractionAssetEditor.h
--- a/Engine/Plugins/Animation/PoseSearch/Source/Editor/Public/PoseSearchInteractionAssetEditor.h
+++ b/Engine/Plugins/Animation/PoseSearch/Source/Editor/Public/PoseSearchInteractionAssetEditor.h
@@ -39,8 +39,8 @@ public:
 
 private:
 	enum { PreviewActor, DebugActor, NumActors };
-	TArray<TWeakObjectPtr<AActor>> ActorPtrs;
-	TArray<UE::PoseSearch::FAnimationAssetSampler> Samplers; // used to keep track of the animation position / orientation
+	TWeakObjectPtr<AActor> ActorPtrs[NumActors];
+	UE::PoseSearch::FAnimationAssetSampler Samplers[NumActors]; // used to keep track of the animation position / orientation
 
 	float CurrentTime = 0.f;
 	UE::PoseSearch::FRole ActorRole = UE::PoseSearch::DefaultRole;
