PoseSearch + Chooser - keeping PoseSearchColumn InternalDatabase in synch when editing the chooser table asset result column

diff --git a/Engine/Plugins/Animation/PoseSearch/Source/Runtime/Private/Chooser/PoseSearchChooserColumn.cpp b/Engine/Plugins/Animation/PoseSearch/Source/Runtime/Private/Chooser/PoseSearchChooserColumn.cpp
--- a/Engine/Plugins/Animation/PoseSearch/Source/Runtime/Private/Chooser/PoseSearchChooserColumn.cpp
+++ b/Engine/Plugins/Animation/PoseSearch/Source/Runtime/Private/Chooser/PoseSearchChooserColumn.cpp
@@ -14,7 +14,6 @@
 #include "PoseSearch/PoseSearchLibrary.h"
 #include "PoseSearch/PoseSearchSchema.h"
 #if WITH_EDITOR
-#include "ScopedTransaction.h"
 #include "StructUtils/PropertyBag.h"
 #endif
 
@@ -99,7 +98,6 @@ void FPoseSearchColumn::SetNumRows(int32 NumRows)
 	}
 	else if (NumRows < InternalDatabase->GetNumAnimationAssets())
 	{
-		FScopedTransaction ScopedTransaction(LOCTEXT("SetNumRows", "SetNumRows FPoseSearchColumn"));
 		InternalDatabase->Modify();
 		while (InternalDatabase->GetNumAnimationAssets() > 0 && InternalDatabase->GetNumAnimationAssets() > NumRows)
 		{
@@ -108,7 +106,6 @@ void FPoseSearchColumn::SetNumRows(int32 NumRows)
 	}
 	else
 	{
-		FScopedTransaction ScopedTransaction(LOCTEXT("SetNumRows", "SetNumRows FPoseSearchColumn"));
 		InternalDatabase->Modify();
 		while (InternalDatabase->GetNumAnimationAssets() < NumRows)
 		{
@@ -132,7 +129,6 @@ void FPoseSearchColumn::InsertRows(int Index, int Count)
 	check(InternalDatabase);
 	if (Count > 0)
 	{
-		FScopedTransaction ScopedTransaction(LOCTEXT("InsertRows","InsertRows FPoseSearchColumn"));
 		InternalDatabase->Modify();
 		for (int i = 0; i < Count; i++)
 		{
@@ -149,11 +145,10 @@ void FPoseSearchColumn::InsertRows(int Index, int Count)
 void FPoseSearchColumn::DeleteRows(const TArray<uint32>& RowIndices)
 {
 	check(InternalDatabase);
-	FScopedTransaction ScopedTransaction(LOCTEXT("DeleteRows","DeleteRows FPoseSearchColumn"));
 	InternalDatabase->Modify();
-	for(uint32 Index : RowIndices)
+	for(uint32 RowIndex : RowIndices)
 	{
-		InternalDatabase->RemoveAnimationAssetAt(Index);
+		InternalDatabase->RemoveAnimationAssetAt(RowIndex);
 	}
 
 	InternalDatabase->NotifySynchronizeWithExternalDependencies();
@@ -165,7 +160,6 @@ void FPoseSearchColumn::MoveRow(int SourceRowIndex, int TargetRowIndex)
 	check(InternalDatabase);
 	if (const FPoseSearchDatabaseAnimationAsset* SourceAsset = InternalDatabase->GetDatabaseAnimationAsset(SourceRowIndex))
 	{
-		FScopedTransaction ScopedTransaction(LOCTEXT("MoveRow","MoveRow FPoseSearchColumn"));
 		InternalDatabase->Modify();
 
 		FPoseSearchDatabaseAnimationAsset CopySourceAsset = *SourceAsset;
@@ -193,7 +187,6 @@ void FPoseSearchColumn::CopyRow(FChooserColumnBase& SourceColumn, int SourceRowI
 	{
 		if (FPoseSearchDatabaseAnimationAsset* TargetAsset = InternalDatabase->GetMutableDatabaseAnimationAsset(TargetRowIndex))
 		{
-			FScopedTransaction ScopedTransaction(LOCTEXT("CopyRow","CopyRow FPoseSearchColumn"));
 			InternalDatabase->Modify();
 			*TargetAsset = *SourceAsset;
 
@@ -219,6 +212,20 @@ void FPoseSearchColumn::SetInputType(const UScriptStruct* Type)
 	InputValue.InitializeAs(Type);
 }
 
+void FPoseSearchColumn::AutoPopulate(int32 RowIndex, UObject* OutputObject)
+{
+	if (FPoseSearchDatabaseAnimationAsset* Asset = InternalDatabase->GetMutableDatabaseAnimationAsset(RowIndex))
+	{
+		if (Asset->AnimAsset != OutputObject)
+		{
+			InternalDatabase->Modify();
+			Asset->AnimAsset = OutputObject;
+		}
+	}
+	
+	CheckConsistency();
+}
+
 bool FPoseSearchColumn::EditorTestFilter(int32 RowIndex) const
 {
 	// todo: not sure how to do editor display for pose search weights (might need to make a copy of the pose history struct)
@@ -235,12 +242,7 @@ float FPoseSearchColumn::EditorTestCost(int32 RowIndex) const
 
 FChooserParameterBase* FPoseSearchColumn::GetInputValue()
 {
-	if (FChooserParameterBase* ChooserParameterBase = InputValue.GetMutablePtr<FChooserParameterBase>())
-	{
-		return ChooserParameterBase;
-	}
-
-	return &NullInputValue;
+	return InputValue.GetMutablePtr<FChooserParameterBase>();
 }
 
 bool FPoseSearchColumn::HasFilters() const
diff --git a/Engine/Plugins/Animation/PoseSearch/Source/Runtime/Public/PoseSearch/Chooser/PoseSearchChooserColumn.h b/Engine/Plugins/Animation/PoseSearch/Source/Runtime/Public/PoseSearch/Chooser/PoseSearchChooserColumn.h
--- a/Engine/Plugins/Animation/PoseSearch/Source/Runtime/Public/PoseSearch/Chooser/PoseSearchChooserColumn.h
+++ b/Engine/Plugins/Animation/PoseSearch/Source/Runtime/Public/PoseSearch/Chooser/PoseSearchChooserColumn.h
@@ -34,17 +34,6 @@ struct FPoseHistoryContextProperty :  public FChooserParameterPoseHistoryBase
 	CHOOSER_PARAMETER_BOILERPLATE();
 };
 
-USTRUCT(Experimental)
-struct FPoseSearchColumnNullInputValue : public FChooserParameterBase
-{
-	GENERATED_BODY()
-
-	virtual void GetDisplayName(FText& OutName) const override
-	{
-		OutName = FText::FromString("FPoseSearchColumn");
-	}
-};
-
 USTRUCT(DisplayName = "Pose Match", Meta = (Experimental, Category = "Experimental", Tooltip = "This column filters out all assets except the one which is selected by motion matching query.  Results must be AnimationAssets with a PoseSearchBranchIn notify state.  It also outputs OutputStartTime to specify the frame which matched pose best.  To work as intended it must be placed last (furthest right) in the Chooser so that other filters are applied first."))
 struct FPoseSearchColumn : public FChooserColumnBase
 {
@@ -76,10 +65,6 @@ struct FPoseSearchColumn : public FChooserColumnBase
 	UPROPERTY(EditAnywhere, DisplayName = "Pose History", NoClear, Meta = (ExcludeBaseStruct, BaseStruct = "/Script/PoseSearch.ChooserParameterPoseHistoryBase"), Category = "Data")
 	FInstancedStruct InputValue;
 	
-	// @todo: do we really need this? Currently used if InputValue is not set to be able to make FChooserTableEditor::PasteInternal work
-	UPROPERTY()
-	FPoseSearchColumnNullInputValue NullInputValue;
-
 	// Float output for the start time with the best matching pose
 	UPROPERTY(EditAnywhere, NoClear, Meta = (ExcludeBaseStruct, BaseStruct = "/Script/Chooser.ChooserParameterFloatBase"), Category = "Output")
 	FInstancedStruct OutputStartTime;
@@ -123,6 +108,9 @@ struct FPoseSearchColumn : public FChooserColumnBase
 	virtual const UScriptStruct* GetInputType() const override;
 	virtual void SetInputType(const UScriptStruct* Type) override;
 	
+	virtual bool AutoPopulates() const override { return true; }
+	virtual void AutoPopulate(int32 RowIndex, UObject* OutputObject) override;
+
 #endif // WITH_EDITOR
 	virtual FChooserParameterBase* GetInputValue() override;
 
diff --git a/Engine/Plugins/Chooser/Source/ChooserEditor/Private/ChooserEditorWidgets.cpp b/Engine/Plugins/Chooser/Source/ChooserEditor/Private/ChooserEditorWidgets.cpp
--- a/Engine/Plugins/Chooser/Source/ChooserEditor/Private/ChooserEditorWidgets.cpp
+++ b/Engine/Plugins/Chooser/Source/ChooserEditor/Private/ChooserEditorWidgets.cpp
@@ -36,11 +36,12 @@ TSharedRef<SWidget> CreateAssetWidget(bool bReadOnly, UObject* TransactionObject
 		.AllowedClass(ResultBaseClass ? ResultBaseClass : UObject::StaticClass())
 		.ObjectPath_Lambda([DIAsset](){ return DIAsset->Asset ? DIAsset->Asset.GetPath() : "";})
 		.bOnlyRecognizeOnDragEnter(true)
-		.OnObjectChanged_Lambda([TransactionObject, DIAsset](const FAssetData& AssetData)
+		.OnObjectChanged_Lambda([TransactionObject, DIAsset, ValueChanged](const FAssetData& AssetData)
 		{
 			const FScopedTransaction Transaction(LOCTEXT("Edit Asset", "Edit Asset"));
 			TransactionObject->Modify(true);
 			DIAsset->Asset = AssetData.GetAsset();
+			ValueChanged.ExecuteIfBound();
 		});
 }
 	
@@ -52,11 +53,12 @@ TSharedRef<SWidget> CreateSoftAssetWidget(bool bReadOnly, UObject* TransactionOb
 		.IsEnabled(!bReadOnly)
 		.AllowedClass(ResultBaseClass ? ResultBaseClass : UObject::StaticClass())
 		.ObjectPath_Lambda([DIAsset](){ return DIAsset->Asset.ToSoftObjectPath().ToString(); })
-		.OnObjectChanged_Lambda([TransactionObject, DIAsset](const FAssetData& AssetData)
+		.OnObjectChanged_Lambda([TransactionObject, DIAsset, ValueChanged](const FAssetData& AssetData)
 		{
 			const FScopedTransaction Transaction(LOCTEXT("Edit Asset", "Edit Asset"));
 			TransactionObject->Modify(true);
 			DIAsset->Asset = AssetData.GetAsset();
+			ValueChanged.ExecuteIfBound();
 		});
 }
 	
@@ -106,11 +108,12 @@ TSharedRef<SWidget> CreateEvaluateChooserWidget(bool bReadOnly, UObject* Transac
 		 	}
 			return true;
 		})
-		.OnObjectChanged_Lambda([TransactionObject, EvaluateChooser](const FAssetData& AssetData)
+		.OnObjectChanged_Lambda([TransactionObject, EvaluateChooser, ValueChanged](const FAssetData& AssetData)
 		{
 			const FScopedTransaction Transaction(LOCTEXT("Edit Chooser", "Edit Chooser"));
 			TransactionObject->Modify(true);
 			EvaluateChooser->Chooser = Cast<UChooserTable>(AssetData.GetAsset());
+			ValueChanged.ExecuteIfBound();
 		});
 }
 
diff --git a/Engine/Plugins/Chooser/Source/ChooserEditor/Private/ChooserTableEditor.cpp b/Engine/Plugins/Chooser/Source/ChooserEditor/Private/ChooserTableEditor.cpp
--- a/Engine/Plugins/Chooser/Source/ChooserEditor/Private/ChooserTableEditor.cpp
+++ b/Engine/Plugins/Chooser/Source/ChooserEditor/Private/ChooserTableEditor.cpp
@@ -2589,36 +2589,33 @@ void FChooserTableEditor::PasteInternal(UChooserTable* PastedContent, int PasteR
 			FInstancedStruct& PastedColumnData = PastedContent->ColumnsStructs[SourceColumnIndex];
 			FChooserColumnBase& PastedColumn = PastedColumnData.GetMutable<FChooserColumnBase>();
 			FString PastedColumnName = GetColumnName(PastedColumn);
-			if (FChooserParameterBase* PastedColumnInput = PastedColumn.GetInputValue())
+			for(int TargetColumnIndex = 0; TargetColumnIndex < Chooser->ColumnsStructs.Num(); TargetColumnIndex++)
 			{
-				for(int TargetColumnIndex = 0; TargetColumnIndex < Chooser->ColumnsStructs.Num(); TargetColumnIndex++)
+				if (!MatchedTargetColumns[TargetColumnIndex])
 				{
-					if (!MatchedTargetColumns[TargetColumnIndex])
+					FInstancedStruct& ColumnData = Chooser->ColumnsStructs[TargetColumnIndex];
+					if (ColumnData.GetScriptStruct() == PastedColumnData.GetScriptStruct())
 					{
-						FInstancedStruct& ColumnData = Chooser->ColumnsStructs[TargetColumnIndex];
-						if (ColumnData.GetScriptStruct() == PastedColumnData.GetScriptStruct())
+						FChooserColumnBase& Column = ColumnData.GetMutable<FChooserColumnBase>();
+						FString ColumnName = GetColumnName(Column);
+
+						if (ColumnName == PastedColumnName)
 						{
-							FChooserColumnBase& Column = ColumnData.GetMutable<FChooserColumnBase>();
-							FString ColumnName = GetColumnName(Column);
+							MatchedTargetColumns[TargetColumnIndex] = true;
+							MatchedSourceColumns[SourceColumnIndex] = true;
 
-							if (ColumnName == PastedColumnName)
+							// found a match, copy the data over
+							for (int i = 0; i < RowsToPaste; i++)
 							{
-								MatchedTargetColumns[TargetColumnIndex] = true;
-								MatchedSourceColumns[SourceColumnIndex] = true;
-
-								// found a match, copy the data over
-								for (int i = 0; i < RowsToPaste; i++)
-								{
-									Column.CopyRow(PastedColumn, i, InsertIndex + i);
-								}
+								Column.CopyRow(PastedColumn, i, InsertIndex + i);
+							}
 
-								if (PastedContent->FallbackResult.IsValid())
-								{
-									// if the fallback row was copied, paste the fallback data for columns
-									Column.CopyFallback(PastedColumn);
-								}
-								break;
+							if (PastedContent->FallbackResult.IsValid())
+							{
+								// if the fallback row was copied, paste the fallback data for columns
+								Column.CopyFallback(PastedColumn);
 							}
+							break;
 						}
 					}
 				}
diff --git a/Engine/Plugins/Chooser/Source/ChooserEditor/Private/SChooserTableRow.cpp b/Engine/Plugins/Chooser/Source/ChooserEditor/Private/SChooserTableRow.cpp
--- a/Engine/Plugins/Chooser/Source/ChooserEditor/Private/SChooserTableRow.cpp
+++ b/Engine/Plugins/Chooser/Source/ChooserEditor/Private/SChooserTableRow.cpp
@@ -113,17 +113,40 @@ namespace UE::ChooserEditor
 			}
 			else if (ColumnName == Result) 
 			{
+				FChooserWidgetValueChanged ChooserWidgetValueChanged = FChooserWidgetValueChanged::CreateLambda([this]()
+					{
+						check(Editor);
+						Editor->AutoPopulateRow(RowIndex->RowIndex);
+					});
+
 				UChooserTable* ContextOwner = Chooser->GetRootChooser();
-				TSharedPtr<SWidget> ResultWidget = FObjectChooserWidgetFactories::CreateWidget(false, Chooser, FObjectChooserBase::StaticStruct(), Chooser->ResultsStructs[RowIndex->RowIndex].GetMutableMemory(), Chooser->ResultsStructs[RowIndex->RowIndex].GetScriptStruct(), ContextOwner->OutputObjectType,
-				FOnStructPicked::CreateLambda([this, InRowIndex=RowIndex->RowIndex](const UScriptStruct* ChosenStruct)
+				TSharedPtr<SWidget> ResultWidget = FObjectChooserWidgetFactories::CreateWidget(
+					false, 
+					Chooser,
+					FObjectChooserBase::StaticStruct(),
+					Chooser->ResultsStructs[RowIndex->RowIndex].GetMutableMemory(),
+					Chooser->ResultsStructs[RowIndex->RowIndex].GetScriptStruct(),
+					ContextOwner->OutputObjectType,
+
+					FOnStructPicked::CreateLambda([this, InRowIndex = RowIndex->RowIndex, ChooserWidgetValueChanged](const UScriptStruct* ChosenStruct)
 				{
 					UChooserTable* ContextOwner = Chooser->GetRootChooser();
 					const FScopedTransaction Transaction(LOCTEXT("Change Row Result Type", "Change Row Result Type"));
 					Chooser->Modify(true);
 					Chooser->ResultsStructs[InRowIndex].InitializeAs(ChosenStruct);
-					FObjectChooserWidgetFactories::CreateWidget(false, Chooser, FObjectChooserBase::StaticStruct(), Chooser->ResultsStructs[InRowIndex].GetMutableMemory(), ChosenStruct, ContextOwner->OutputObjectType, FOnStructPicked(), &CacheBorder);
+							FObjectChooserWidgetFactories::CreateWidget(
+								false,
+								Chooser,
+								FObjectChooserBase::StaticStruct(),
+								Chooser->ResultsStructs[InRowIndex].GetMutableMemory(),
+								ChosenStruct,
+								ContextOwner->OutputObjectType,
+								FOnStructPicked(),
+								&CacheBorder,
+								ChooserWidgetValueChanged);
 				}),
-				&CacheBorder
+					&CacheBorder,
+					ChooserWidgetValueChanged
 				);
 
 				return SNew(SOverlay)
