Fixed dllstorage to make merged modules link errors

diff --git a/Engine/Plugins/Animation/PoseSearch/Source/Runtime/Private/PoseSearchAssetSampler.cpp b/Engine/Plugins/Animation/PoseSearch/Source/Runtime/Private/PoseSearchAssetSampler.cpp
--- a/Engine/Plugins/Animation/PoseSearch/Source/Runtime/Private/PoseSearchAssetSampler.cpp
+++ b/Engine/Plugins/Animation/PoseSearch/Source/Runtime/Private/PoseSearchAssetSampler.cpp
@@ -239,11 +239,19 @@ static int32 GetHighestWeightSample(const TArray<struct FBlendSampleData>& Sampl
 
 //////////////////////////////////////////////////////////////////////////
 // FAssetSamplerBase
+
+FAnimationAssetSampler::FAnimationAssetSampler()
+{
+	Init(nullptr, FTransform::Identity, FVector::ZeroVector, DefaultRootTransformSamplingRate, true, true);
+}
+
 FAnimationAssetSampler::FAnimationAssetSampler(const UAnimationAsset* InAnimationAsset, const FTransform& InRootTransformOrigin, const FVector& InBlendParameters, int32 InRootTransformSamplingRate, bool bPreProcessRootTransform, bool bInEnforceCompressedDataSampling)
 {
 	Init(InAnimationAsset, InRootTransformOrigin, InBlendParameters, InRootTransformSamplingRate, bPreProcessRootTransform, bInEnforceCompressedDataSampling);
 }
 
+FAnimationAssetSampler::~FAnimationAssetSampler() = default;
+
 void FAnimationAssetSampler::Init(const UAnimationAsset* InAnimationAsset, const FTransform& InRootTransformOrigin, const FVector& InBlendParameters, int32 InRootTransformSamplingRate, bool bPreProcessRootTransform, bool bInEnforceCompressedDataSampling)
 {
 	AnimationAssetPtr = InAnimationAsset;
diff --git a/Engine/Plugins/Animation/PoseSearch/Source/Runtime/Public/PoseSearch/PoseSearchAssetSampler.h b/Engine/Plugins/Animation/PoseSearch/Source/Runtime/Public/PoseSearch/PoseSearchAssetSampler.h
--- a/Engine/Plugins/Animation/PoseSearch/Source/Runtime/Public/PoseSearch/PoseSearchAssetSampler.h
+++ b/Engine/Plugins/Animation/PoseSearch/Source/Runtime/Public/PoseSearch/PoseSearchAssetSampler.h
@@ -21,7 +21,9 @@ namespace UE::PoseSearch
 struct FAnimationAssetSampler
 {
 	enum { DefaultRootTransformSamplingRate = 30 };
-	POSESEARCH_API FAnimationAssetSampler(const UAnimationAsset* InAnimationAsset = nullptr, const FTransform& InRootTransformOrigin = FTransform::Identity, const FVector& InBlendParameters = FVector::ZeroVector, int32 InRootTransformSamplingRate = DefaultRootTransformSamplingRate, bool bPreProcessRootTransform = true, bool bInEnforceCompressedDataSampling = true);
+	POSESEARCH_API FAnimationAssetSampler();
+	POSESEARCH_API FAnimationAssetSampler(const UAnimationAsset* InAnimationAsset, const FTransform& InRootTransformOrigin = FTransform::Identity, const FVector& InBlendParameters = FVector::ZeroVector, int32 InRootTransformSamplingRate = DefaultRootTransformSamplingRate, bool bPreProcessRootTransform = true, bool bInEnforceCompressedDataSampling = true);
+	POSESEARCH_API ~FAnimationAssetSampler();
 	POSESEARCH_API void Init(const UAnimationAsset* InAnimationAsset, const FTransform& InRootTransformOrigin = FTransform::Identity, const FVector& InBlendParameters = FVector::ZeroVector, int32 InRootTransformSamplingRate = DefaultRootTransformSamplingRate, bool bPreProcessRootTransform = true, bool bInEnforceCompressedDataSampling = true);
 
 	POSESEARCH_API bool IsInitialized() const;
diff --git a/Engine/Source/Editor/UnrealEd/Private/EditorViewportClient.cpp b/Engine/Source/Editor/UnrealEd/Private/EditorViewportClient.cpp
--- a/Engine/Source/Editor/UnrealEd/Private/EditorViewportClient.cpp
+++ b/Engine/Source/Editor/UnrealEd/Private/EditorViewportClient.cpp
@@ -7014,6 +7014,7 @@ void FEditorViewportStats::SendUsageData()
 	FMemory::Memzero(DataPoints);
 }
 
+UE_DEFINE_TCOMMANDS(FViewportNavigationCommands)
 
 FViewportNavigationCommands::FViewportNavigationCommands()
 	: TCommands<FViewportNavigationCommands>(
diff --git a/Engine/Source/Editor/UnrealEd/Public/EditorViewportClient.h b/Engine/Source/Editor/UnrealEd/Public/EditorViewportClient.h
--- a/Engine/Source/Editor/UnrealEd/Public/EditorViewportClient.h
+++ b/Engine/Source/Editor/UnrealEd/Public/EditorViewportClient.h
@@ -58,19 +58,21 @@ namespace EDragTool
 	};
 }
 
+UE_DECLARE_TCOMMANDS(class FViewportNavigationCommands, UNREALED_API)
+
 /**
 * Unreal level editor actions
 */
-class UNREALED_API FViewportNavigationCommands : public TCommands<FViewportNavigationCommands>
+class FViewportNavigationCommands : public TCommands<FViewportNavigationCommands>
 {
 
 public:
-	FViewportNavigationCommands();
+	UNREALED_API FViewportNavigationCommands();
 
 	/**
 	* Initialize commands
 	*/
-	virtual void RegisterCommands() override;
+	UNREALED_API virtual void RegisterCommands() override;
 
 
 	TSharedPtr< FUICommandInfo > Forward;
