Instead of checking in PoseSearchHistory LerpEntries we if/ensure/return early indicating failure

#rnx

diff --git a/Engine/Plugins/Animation/PoseSearch/Source/Runtime/Private/PoseSearchHistory.cpp b/Engine/Plugins/Animation/PoseSearch/Source/Runtime/Private/PoseSearchHistory.cpp
--- a/Engine/Plugins/Animation/PoseSearch/Source/Runtime/Private/PoseSearchHistory.cpp
+++ b/Engine/Plugins/Animation/PoseSearch/Source/Runtime/Private/PoseSearchHistory.cpp
@@ -162,7 +162,11 @@ static bool LerpEntries(float Time, bool bExtrapolate, const FPoseHistoryEntry&
 
 	bool bSuccess = true;
 	const FComponentSpaceTransformIndex BoneTransformIndex = GetRemappedComponentSpaceTransformIndex(BoneIndexSkeleton, LastUpdateSkeleton, BoneToTransformMap, BoneIndexType, bSuccess);
-	check(BoneTransformIndex != ComponentSpaceIndexType);
+	if (!ensureMsgf(BoneTransformIndex != ComponentSpaceIndexType, TEXT("BoneIndexSkeleton [%s], LastUpdateSkeleton [%s] mapped to the same index."), *GetNameSafe(BoneIndexSkeleton), *GetNameSafe(LastUpdateSkeleton)))
+	{
+		OutBoneTransform = FTransform::Identity;
+		return false;
+	}
 	
 	const float Denominator = NextEntry.AccumulatedSeconds - PrevEntry.AccumulatedSeconds;
 	float LerpValue = 0.f;
