PoseSearch + Chooser
- fix for pose search debugger being able to load and display utrace files correctly
- moved FAnimNode_ChooserPlayer::bForceBlendTo into FChooserPlayerSettings to be able to set it to true when the continuing pose search from the FPoseSearchColumn fails
- added FChooserPlayerSettings::MinDeltaTimeToForceBlendTo to control the chooser player blendstack blend to with some time tolerance

diff --git a/Engine/Plugins/Animation/PoseSearch/Source/Editor/Private/PoseSearchDebuggerDatabaseColumns.h b/Engine/Plugins/Animation/PoseSearch/Source/Editor/Private/PoseSearchDebuggerDatabaseColumns.h
--- a/Engine/Plugins/Animation/PoseSearch/Source/Editor/Private/PoseSearchDebuggerDatabaseColumns.h
+++ b/Engine/Plugins/Animation/PoseSearch/Source/Editor/Private/PoseSearchDebuggerDatabaseColumns.h
@@ -182,7 +182,7 @@ struct FDatabaseName : IColumn
 
 	virtual FSortPredicate GetSortPredicate() const override
 	{
-		return [](const FRowDataRef& Row0, const FRowDataRef& Row1) -> bool { return Row0->SharedData->DatabaseName < Row1->SharedData->DatabaseName; };
+		return [](const FRowDataRef& Row0, const FRowDataRef& Row1) -> bool { return GetNameSafe(Row0->SharedData->SourceDatabase.Get()) < GetNameSafe(Row1->SharedData->SourceDatabase.Get()); };
 	}
 
 	virtual TSharedRef<SWidget> GenerateWidget(const FRowDataRef& RowData) const override
@@ -190,41 +190,25 @@ struct FDatabaseName : IColumn
 		return SNew(SHyperlink)
 		.Text_Lambda([RowData]() -> FText
 		{
-			return FText::FromString(RowData->SharedData->DatabaseName);
+			return FText::FromString(GetNameSafe(RowData->SharedData->SourceDatabase.Get()));
 		})
 		.TextStyle(&FCoreStyle::Get().GetWidgetStyle<FTextBlockStyle>("SmallText"))
 		.ToolTipText_Lambda([RowData]() -> FText
 		{
 			return FText::Format(
 				LOCTEXT("DatabaseHyperlinkTooltipFormat", "Open database '{0}'"),
-				FText::FromString(RowData->SharedData->DatabasePath));
+				FText::FromString(GetPathNameSafe(RowData->SharedData->SourceDatabase.Get())));
 		})
 		.OnNavigate_Lambda([this, RowData]()
 		{
-			UObject* Asset = nullptr;
-
-			// Load asset
-			if (UPackage* Package = LoadPackage(NULL, *RowData->SharedData->DatabasePath, LOAD_NoRedirects))
-			{
-				Package->FullyLoad();
-
-				const FString AssetName = FPaths::GetBaseFilename(RowData->SharedData->DatabasePath);
-				Asset = FindObject<UObject>(Package, *AssetName);
-			}
-			else
-			{
-				// Fallback for unsaved assets
-				Asset = FindObject<UObject>(nullptr, *RowData->SharedData->DatabasePath);
-			}
-
 			// Open editor
-			if (Asset != nullptr)
+			if (UPoseSearchDatabase* Database = RowData->SharedData->SourceDatabase.Get())
 			{
 				if (UAssetEditorSubsystem* AssetEditorSS = GEditor->GetEditorSubsystem<UAssetEditorSubsystem>())
 				{
-					AssetEditorSS->OpenEditorForAsset(Asset);
+					AssetEditorSS->OpenEditorForAsset(Database);
 								
-					if (IAssetEditorInstance* Editor = AssetEditorSS->FindEditorForAsset(Asset, true))
+					if (IAssetEditorInstance* Editor = AssetEditorSS->FindEditorForAsset(Database, true))
 					{
 						if (Editor->GetEditorName() == FName("PoseSearchDatabaseEditor"))
 						{
@@ -248,7 +232,7 @@ struct FDatabaseName : IColumn
 		
 		for (const TSharedRef<FDebuggerDatabaseRowData>& Row : Rows)
 		{
-			MinWidth = FMath::Max(FontMeasure->Measure(Row->SharedData->DatabaseName, CellTextStyle).X + ColumnWidthPadding, MinWidth);
+			MinWidth = FMath::Max(FontMeasure->Measure(GetNameSafe(Row->SharedData->SourceDatabase.Get()), CellTextStyle).X + ColumnWidthPadding, MinWidth);
 		}
 		
 		return MinWidth;
diff --git a/Engine/Plugins/Animation/PoseSearch/Source/Editor/Private/PoseSearchDebuggerDatabaseRowData.h b/Engine/Plugins/Animation/PoseSearch/Source/Editor/Private/PoseSearchDebuggerDatabaseRowData.h
--- a/Engine/Plugins/Animation/PoseSearch/Source/Editor/Private/PoseSearchDebuggerDatabaseRowData.h
+++ b/Engine/Plugins/Animation/PoseSearch/Source/Editor/Private/PoseSearchDebuggerDatabaseRowData.h
@@ -12,9 +12,13 @@ namespace UE::PoseSearch
 class FDebuggerDatabaseSharedData : public TSharedFromThis<FDebuggerDatabaseSharedData>
 {
 public:
-	TWeakObjectPtr<const UPoseSearchDatabase> SourceDatabase;
+	TWeakObjectPtr<UPoseSearchDatabase> SourceDatabase;
+	
+	UE_DEPRECATED(5.7, "Use GetNameSafe(SourceDatabase.Get()) instead")
 	FString DatabaseName = "";
+	UE_DEPRECATED(5.7, "Use GetPathNameSafe(SourceDatabase.Get()) instead")
 	FString DatabasePath = "";
+	
 	TAlignedArray<float> QueryVector;
 	TAlignedArray<float> PCAQueryVector;
 };
diff --git a/Engine/Plugins/Animation/PoseSearch/Source/Editor/Private/PoseSearchDebuggerDatabaseView.cpp b/Engine/Plugins/Animation/PoseSearch/Source/Editor/Private/PoseSearchDebuggerDatabaseView.cpp
--- a/Engine/Plugins/Animation/PoseSearch/Source/Editor/Private/PoseSearchDebuggerDatabaseView.cpp
+++ b/Engine/Plugins/Animation/PoseSearch/Source/Editor/Private/PoseSearchDebuggerDatabaseView.cpp
@@ -6,7 +6,9 @@
 #include "Animation/AnimComposite.h"
 #include "Animation/AnimMontage.h"
 #include "Animation/AnimSequence.h"
+#include "IGameplayProvider.h"
 #include "Internationalization/Regex.h"
+#include "IRewindDebugger.h"
 #include "PoseSearch/PoseSearchDatabase.h"
 #include "PoseSearch/PoseSearchDerivedData.h"
 #include "PoseSearch/PoseSearchFeatureChannel.h"
@@ -28,6 +30,23 @@
 namespace UE::PoseSearch
 {
 
+UPoseSearchDatabase* ResolveDatabaseFromId(uint64 DatabaseId)
+{
+	UPoseSearchDatabase* Database = nullptr;
+	if (const TraceServices::IAnalysisSession* AnalysisSession = IRewindDebugger::Instance()->GetAnalysisSession())
+	{
+		TraceServices::FAnalysisSessionReadScope SessionReadScope(*AnalysisSession);
+		if (const IGameplayProvider* GameplayProvider = AnalysisSession->ReadProvider<IGameplayProvider>("GameplayProvider"))
+		{
+			if (const FObjectInfo* DatabaseInfo = GameplayProvider->FindObjectInfo(DatabaseId))
+			{
+				Database = TSoftObjectPtr<UPoseSearchDatabase>(FSoftObjectPath(DatabaseInfo->PathName)).LoadSynchronous();
+			}
+		}
+	}
+	return Database;
+}
+
 static FORCEINLINE float ArraySum(TConstArrayView<float> View, int32 StartIndex, int32 Offset)
 {
 	float Sum = 0.f;
@@ -49,7 +68,7 @@ public:
 		FAsyncPoseSearchDatabasesManagementCache DatabasesManagementCache;
 		for (const FTraceMotionMatchingStateDatabaseEntry& DbEntry : DatabaseEntries)
 		{
-			const UPoseSearchDatabase* Database = FTraceMotionMatchingStateMessage::GetObjectFromId<UPoseSearchDatabase>(DbEntry.DatabaseId);
+			const UPoseSearchDatabase* Database = ResolveDatabaseFromId(DbEntry.DatabaseId);
 			if (EAsyncBuildIndexResult::Success == DatabasesManagementCache(Database))
 			{
 				for (const TObjectPtr<UPoseSearchFeatureChannel>& ChannelPtr : Database->Schema->GetChannels())
@@ -150,8 +169,6 @@ private:
 		}
 	}
 
-	
-
 	struct FLabelToChannels
 	{
 		FText Label;
@@ -259,7 +276,7 @@ void SDebuggerDatabaseView::Update(const FTraceMotionMatchingStateMessage& State
 	FAsyncPoseSearchDatabasesManagementCache DatabasesManagementCache;
 	for (const FTraceMotionMatchingStateDatabaseEntry& DbEntry : State.DatabaseEntries)
 	{
-		const UPoseSearchDatabase* Database = FTraceMotionMatchingStateMessage::GetObjectFromId<UPoseSearchDatabase>(DbEntry.DatabaseId);
+		UPoseSearchDatabase* Database = ResolveDatabaseFromId(DbEntry.DatabaseId);
 		if (EAsyncBuildIndexResult::Success == DatabasesManagementCache(Database))
 		{
 			const FSearchIndex& SearchIndex = Database->GetSearchIndex();
@@ -267,8 +284,6 @@ void SDebuggerDatabaseView::Update(const FTraceMotionMatchingStateMessage& State
 
 			TSharedRef<FDebuggerDatabaseSharedData> SharedData = MakeShared<FDebuggerDatabaseSharedData>();
 			SharedData->SourceDatabase = Database;
-			SharedData->DatabaseName = Database->GetName();
-			SharedData->DatabasePath = Database->GetPathName();
 			SharedData->QueryVector = DbEntry.QueryVector;
 
 			// checking for DbEntry.QueryVector.Num() == Database->Schema->SchemaCardinality to make sure the stored DbEntry data is still aligned with the Database->Schema
@@ -658,7 +673,7 @@ void SDebuggerDatabaseView::OnFilteredDatabaseViewMouseButtonClick(const TShared
 
 void SDebuggerDatabaseView::SortDatabaseRows()
 {
-	if (!UnfilteredDatabaseRows.IsEmpty())
+	if (!UnfilteredDatabaseRows.IsEmpty() && Columns.Contains(SortColumn))
 	{
 		if (SortMode == EColumnSortMode::Ascending)
 		{
@@ -709,13 +724,14 @@ void SDebuggerDatabaseView::PopulateViewRows()
 
 				if (RowDataA.SharedData->SourceDatabase != RowDataB.SharedData->SourceDatabase)
 				{
-					const int32 DatabasePathCompare = RowDataA.SharedData->DatabasePath.Compare(RowDataB.SharedData->DatabasePath);
-					if (DatabasePathCompare < 0)
+					const int32 DatabaseNameCompare = GetNameSafe(RowDataA.SharedData->SourceDatabase.Get()).Compare(GetNameSafe(RowDataB.SharedData->SourceDatabase.Get()));
+					if (DatabaseNameCompare < 0)
 					{
 						return true;
 					}
-					const int32 DatabaseNameCompare = RowDataA.SharedData->DatabaseName.Compare(RowDataB.SharedData->DatabaseName);
-					return DatabaseNameCompare < 0;
+
+					const int32 DatabasePathCompare = GetPathNameSafe(RowDataA.SharedData->SourceDatabase.Get()).Compare(GetPathNameSafe(RowDataB.SharedData->SourceDatabase.Get()));
+					return DatabasePathCompare < 0;
 				}
 
 				if (RowDataA.DbAssetIdx == RowDataB.DbAssetIdx)
diff --git a/Engine/Plugins/Animation/PoseSearch/Source/Editor/Private/PoseSearchDebuggerDatabaseView.h b/Engine/Plugins/Animation/PoseSearch/Source/Editor/Private/PoseSearchDebuggerDatabaseView.h
--- a/Engine/Plugins/Animation/PoseSearch/Source/Editor/Private/PoseSearchDebuggerDatabaseView.h
+++ b/Engine/Plugins/Animation/PoseSearch/Source/Editor/Private/PoseSearchDebuggerDatabaseView.h
@@ -11,6 +11,8 @@ template <typename ItemType> class SListView;
 namespace UE::PoseSearch
 {
 
+UPoseSearchDatabase* ResolveDatabaseFromId(uint64 DatabaseId);
+
 namespace DebuggerDatabaseColumns { struct IColumn; }
 struct FTraceMotionMatchingStateMessage;
 
diff --git a/Engine/Plugins/Animation/PoseSearch/Source/Editor/Private/PoseSearchDebuggerView.cpp b/Engine/Plugins/Animation/PoseSearch/Source/Editor/Private/PoseSearchDebuggerView.cpp
--- a/Engine/Plugins/Animation/PoseSearch/Source/Editor/Private/PoseSearchDebuggerView.cpp
+++ b/Engine/Plugins/Animation/PoseSearch/Source/Editor/Private/PoseSearchDebuggerView.cpp
@@ -399,7 +399,7 @@ void SDebuggerView::Tick(const FGeometry& AllottedGeometry, const double InCurre
 		{
 			for (const FTraceMotionMatchingStateDatabaseEntry& DbEntry : State->DatabaseEntries)
 			{
-				const UPoseSearchDatabase* Database = FTraceMotionMatchingStateMessage::GetObjectFromId<UPoseSearchDatabase>(DbEntry.DatabaseId);
+				const UPoseSearchDatabase* Database = ResolveDatabaseFromId(DbEntry.DatabaseId);
 				if (Database &&
 					(SelectedRowsDatabases.IsEmpty() || SelectedRowsDatabases.Contains(Database)) &&
 					EAsyncBuildIndexResult::Success == DatabasesManagementCache(Database) &&
diff --git a/Engine/Plugins/Animation/PoseSearch/Source/Runtime/Private/Chooser/PoseSearchChooserColumn.cpp b/Engine/Plugins/Animation/PoseSearch/Source/Runtime/Private/Chooser/PoseSearchChooserColumn.cpp
--- a/Engine/Plugins/Animation/PoseSearch/Source/Runtime/Private/Chooser/PoseSearchChooserColumn.cpp
+++ b/Engine/Plugins/Animation/PoseSearch/Source/Runtime/Private/Chooser/PoseSearchChooserColumn.cpp
@@ -388,6 +388,7 @@ void FPoseSearchColumn::SetOutputs(FChooserEvaluationContext& Context, int RowIn
 	float CostValue = 100.f;
 	float StartTimeValue = 0.f;
 	bool bMirrorValue = false;
+	bool bForceBlendToValue = false;
 
 	FPoseSearchColumnContextScratchArea& ScratchArea = FPoseSearchColumnContextScratchArea::Get();
 #if DO_CHECK
@@ -404,6 +405,7 @@ void FPoseSearchColumn::SetOutputs(FChooserEvaluationContext& Context, int RowIn
 				StartTimeValue = SearchResult.GetAssetTime();
 				CostValue = SearchResult.PoseCost;
 				bMirrorValue = SearchIndexAsset->IsMirrored();
+				bForceBlendToValue = SearchResult.bIsContinuingPoseSearch;
 
 				if (ScratchArea.PoseHistory)
 				{
@@ -445,6 +447,11 @@ void FPoseSearchColumn::SetOutputs(FChooserEvaluationContext& Context, int RowIn
 	{
 		Mirror->SetValue(Context, bMirrorValue);
 	}
+
+	if (const FChooserParameterBoolBase* ForceBlendTo = OutputForceBlendTo.GetPtr<FChooserParameterBoolBase>())
+	{
+		ForceBlendTo->SetValue(Context, bForceBlendToValue);
+	}
 }
 
 #if WITH_EDITOR
diff --git a/Engine/Plugins/Animation/PoseSearch/Source/Runtime/Private/Trace/PoseSearchTraceLogger.cpp b/Engine/Plugins/Animation/PoseSearch/Source/Runtime/Private/Trace/PoseSearchTraceLogger.cpp
--- a/Engine/Plugins/Animation/PoseSearch/Source/Runtime/Private/Trace/PoseSearchTraceLogger.cpp
+++ b/Engine/Plugins/Animation/PoseSearch/Source/Runtime/Private/Trace/PoseSearchTraceLogger.cpp
@@ -4,11 +4,11 @@
 #include "Animation/AnimInstanceProxy.h"
 #include "Animation/AnimNodeBase.h"
 #include "Components/SkeletalMeshComponent.h"
+#include "PoseSearch/PoseSearchCustomVersion.h"
 #include "PoseSearch/PoseSearchDatabase.h"
 #include "Serialization/MemoryWriter.h"
 #include "Trace/Trace.inl"
 #include "TraceFilter.h"
-#include "PoseSearch/PoseSearchCustomVersion.h"
 
 UE_TRACE_CHANNEL_DEFINE(PoseSearchChannel);
 
@@ -30,7 +30,6 @@ UE_TRACE_EVENT_END()
 namespace UE::PoseSearch
 {
 
-// @todo: These should be deprecated. No need to expose them.
 const FName FTraceLogger::Name("PoseSearch");
 const FName FTraceMotionMatchingStateMessage::Name("MotionMatchingState3");
 
@@ -111,14 +110,8 @@ void FTraceMotionMatchingStateMessage::Output()
 
 const UPoseSearchDatabase* FTraceMotionMatchingStateMessage::GetCurrentDatabase() const
 {
-	const UPoseSearchDatabase* Database = nullptr;
-	PRAGMA_DISABLE_DEPRECATION_WARNINGS
-	if (DatabaseEntries.IsValidIndex(CurrentDbEntryIdx))
-	{
-		Database = GetObjectFromId<UPoseSearchDatabase>(DatabaseEntries[CurrentDbEntryIdx].DatabaseId);
-	}
-	PRAGMA_ENABLE_DEPRECATION_WARNINGS
-	return Database;
+	unimplemented();
+	return nullptr;
 }
 
 int32 FTraceMotionMatchingStateMessage::GetCurrentDatabasePoseIndex() const
@@ -149,6 +142,7 @@ const FTraceMotionMatchingStatePoseEntry* FTraceMotionMatchingStateMessage::GetC
 
 FText FTraceMotionMatchingStateMessage::GenerateSearchName() const
 {
+	PRAGMA_DISABLE_DEPRECATION_WARNINGS
 	TStringBuilder<256> StringBuilder;
 	bool bAddDatabasesSeparator = false;
 	for (const FTraceMotionMatchingStateDatabaseEntry& DbEntry : DatabaseEntries)
@@ -215,6 +209,7 @@ FText FTraceMotionMatchingStateMessage::GenerateSearchName() const
 			StringBuilder.Append("]");
 		}
 	}
+	PRAGMA_ENABLE_DEPRECATION_WARNINGS
 
 	return FText::FromString(StringBuilder.ToString());
 }
diff --git a/Engine/Plugins/Animation/PoseSearch/Source/Runtime/Public/PoseSearch/Chooser/PoseSearchChooserColumn.h b/Engine/Plugins/Animation/PoseSearch/Source/Runtime/Public/PoseSearch/Chooser/PoseSearchChooserColumn.h
--- a/Engine/Plugins/Animation/PoseSearch/Source/Runtime/Public/PoseSearch/Chooser/PoseSearchChooserColumn.h
+++ b/Engine/Plugins/Animation/PoseSearch/Source/Runtime/Public/PoseSearch/Chooser/PoseSearchChooserColumn.h
@@ -73,6 +73,10 @@ struct FPoseSearchColumn : public FChooserColumnBase
 	UPROPERTY(EditAnywhere, NoClear, Meta = (ExcludeBaseStruct, BaseStruct = "/Script/Chooser.ChooserParameterBoolBase"), Category = "Output")
 	FInstancedStruct OutputMirror;
 
+	// Bool output for suggesting a chooser player to force a blend into the newly selected asset
+	UPROPERTY(EditAnywhere, NoClear, Meta = (ExcludeBaseStruct, BaseStruct = "/Script/Chooser.ChooserParameterBoolBase"), Category = "Output")
+	FInstancedStruct OutputForceBlendTo;
+
 	// Float output for the cost of the selected pose
 	UPROPERTY(EditAnywhere, NoClear, Meta = (ExcludeBaseStruct, BaseStruct = "/Script/Chooser.ChooserParameterFloatBase"), Category = "Output")
 	FInstancedStruct OutputCost;
diff --git a/Engine/Plugins/Animation/PoseSearch/Source/Runtime/Public/PoseSearch/Trace/PoseSearchTraceLogger.h b/Engine/Plugins/Animation/PoseSearch/Source/Runtime/Public/PoseSearch/Trace/PoseSearchTraceLogger.h
--- a/Engine/Plugins/Animation/PoseSearch/Source/Runtime/Public/PoseSearch/Trace/PoseSearchTraceLogger.h
+++ b/Engine/Plugins/Animation/PoseSearch/Source/Runtime/Public/PoseSearch/Trace/PoseSearchTraceLogger.h
@@ -16,7 +16,7 @@ UE_TRACE_CHANNEL_EXTERN(PoseSearchChannel, POSESEARCH_API);
 namespace UE::PoseSearch
 {
 
-struct FTraceLogger
+struct UE_DEPRECATED(5.7, "FTraceLogger is no longer necessary and it'll be removed") FTraceLogger
 {
 	/** Used for reading trace data */
     static UE_API const FName Name;
@@ -116,6 +116,7 @@ struct FTraceMotionMatchingStateMessage : public FTraceMessage
 	UE_API const FTraceMotionMatchingStatePoseEntry* GetCurrentPoseEntry() const;
 
 	template<typename T>
+	UE_DEPRECATED(5.7, "This API will soon be removed. Look for ResolveDatabaseFromId to resolve databases from ids")
 	static const T* GetObjectFromId(uint64 ObjectId)
 	{
 #if OBJECT_TRACE_ENABLED
@@ -145,6 +146,7 @@ struct FTraceMotionMatchingStateMessage : public FTraceMessage
 	UE_DEPRECATED(5.6, "Use FText GenerateSearchName(const FTraceMotionMatchingStateMessage&, const IGameplayProvider*) instead")
 	UE_API FText GenerateSearchName() const;
 
+	UE_DEPRECATED(5.7, "FTraceMotionMatchingStateMessage::Name is no longer necessary and it'll be removed")
 	static UE_API const FName Name;
 };
 
diff --git a/Engine/Plugins/Chooser/Source/Chooser/Internal/AnimNode_ChooserPlayer.h b/Engine/Plugins/Chooser/Source/Chooser/Internal/AnimNode_ChooserPlayer.h
--- a/Engine/Plugins/Chooser/Source/Chooser/Internal/AnimNode_ChooserPlayer.h
+++ b/Engine/Plugins/Chooser/Source/Chooser/Internal/AnimNode_ChooserPlayer.h
@@ -93,9 +93,7 @@ private:
 	float CurrentStartTime = 0;
 	bool CurrentMirror = false;
 	bool bInitialized = false;
-	bool bForceBlendTo = false;
 	uint32 CurrentCurveOverridesHash = 0;
-	
 
 	// Update Counter for detecting being relevant
 	FGraphTraversalCounter UpdateCounter;
diff --git a/Engine/Plugins/Chooser/Source/Chooser/Private/AnimNode_ChooserPlayer.cpp b/Engine/Plugins/Chooser/Source/Chooser/Private/AnimNode_ChooserPlayer.cpp
--- a/Engine/Plugins/Chooser/Source/Chooser/Private/AnimNode_ChooserPlayer.cpp
+++ b/Engine/Plugins/Chooser/Source/Chooser/Private/AnimNode_ChooserPlayer.cpp
@@ -64,7 +64,7 @@ UAnimationAsset* FAnimNode_ChooserPlayer::ChooseAsset(const FAnimationUpdateCont
 				{
 					if (!SearchResult.bIsFromContinuingPlaying)
 					{
-						bForceBlendTo = true;
+						Settings.bForceBlendTo = true;
 						Settings.StartTime = SearchResult.TimeOffsetSeconds;
 						Settings.bMirror = SearchResult.bMirrored;
 					}
@@ -147,14 +147,14 @@ void FAnimNode_ChooserPlayer::UpdateAssetPlayer(const FAnimationUpdateContext& C
 	// - if this node just became relevant
 	// - if we chose a new animation
 	// - if the mirror setting has changed
-	// - for playback rate of 0, when the start time changes - for choosing poses as frames of an animation sequence
+	// - for playback rate of 0, when the start time changes (within Settings.MinDeltaTimeToForceBlendTo seconds of tolerance) - for choosing poses as frames of an animation sequence - @todo: handle looping properly
 	// - if the curve values are different
-	if (bForceBlendTo || bJustBecameRelevant || NewAsset != CurrentAsset || AnimPlayers.IsEmpty() ||
+	if (Settings.bForceBlendTo || bJustBecameRelevant || NewAsset != CurrentAsset || AnimPlayers.IsEmpty() ||
 		CurrentMirror != Settings.bMirror ||
-		(CurrentStartTime != Settings.StartTime && Settings.PlaybackRate == 0.0f) ||
+		(!FMath::IsNearlyEqual(CurrentStartTime, Settings.StartTime, Settings.MinDeltaTimeToForceBlendTo) && Settings.PlaybackRate == 0.0f) ||
 		CurrentCurveOverridesHash != Settings.CurveOverrides.Hash)
 	{
-		bForceBlendTo = false;
+		Settings.bForceBlendTo = false;
 		CurrentCurveOverridesHash = Settings.CurveOverrides.Hash;
 		CurrentMirror = Settings.bMirror;
 		
diff --git a/Engine/Plugins/Chooser/Source/Chooser/Public/ChooserTypes.h b/Engine/Plugins/Chooser/Source/Chooser/Public/ChooserTypes.h
--- a/Engine/Plugins/Chooser/Source/Chooser/Public/ChooserTypes.h
+++ b/Engine/Plugins/Chooser/Source/Chooser/Public/ChooserTypes.h
@@ -12,10 +12,10 @@ struct FAnimCurveOverride
 {
 	GENERATED_BODY()
 	// Name of curve to override
-	UPROPERTY(EditAnywhere, BlueprintReadWrite, Category = "CurveValue")
+	UPROPERTY(EditAnywhere, BlueprintReadWrite, Category = CurveValue)
 	FName CurveName;
 	// Value to set to the curve
-	UPROPERTY(EditAnywhere, BlueprintReadWrite, Category = "CurveValue")
+	UPROPERTY(EditAnywhere, BlueprintReadWrite, Category = CurveValue)
 	float CurveValue = 0.0f;
 };
 
@@ -29,7 +29,7 @@ USTRUCT(BlueprintType)
 struct FAnimCurveOverrideList
 {
 	GENERATED_BODY()
-	UPROPERTY(EditAnywhere, BlueprintReadWrite, Category = "Values", meta = (ExpandByDefault = true))
+	UPROPERTY(EditAnywhere, BlueprintReadWrite, Category = Values, meta = (ExpandByDefault = true))
 	TArray<FAnimCurveOverride> Values;
 
 	UPROPERTY()
@@ -47,27 +47,27 @@ struct FChooserPlayerSettings
 	bool Serialize(FStructuredArchive::FSlot Slot);
 
 	// Set this value to mirror animations - the MirrorDataTable must also be set on the AnimNode
-	UPROPERTY(EditAnywhere, BlueprintReadWrite, Category = "Settings")
+	UPROPERTY(EditAnywhere, BlueprintReadWrite, Category = Settings)
 	bool bMirror = false;
 
 	// Start offset when starting the Animation Asset
-	UPROPERTY(EditAnywhere, BlueprintReadWrite, Category = "Settings", meta = (FrameTimeEditor))
-	float StartTime = 0;
+	UPROPERTY(EditAnywhere, BlueprintReadWrite, Category = Settings, meta = (FrameTimeEditor))
+	float StartTime = 0.f;
 
 	// Loop the animation asset, even if the asset is not set as looping
-	UPROPERTY(EditAnywhere, BlueprintReadWrite, Category = "Settings")
+	UPROPERTY(EditAnywhere, BlueprintReadWrite, Category = Settings)
 	bool bForceLooping = false;
 
 	// playback rate modifier
-	UPROPERTY(EditAnywhere, BlueprintReadWrite, Category = "Settings")
-	float PlaybackRate = 1.0;
+	UPROPERTY(EditAnywhere, BlueprintReadWrite, Category = Settings)
+	float PlaybackRate = 1.f;
 
 	// List of curve values to set 
-	UPROPERTY(EditAnywhere, BlueprintReadWrite, Category = "Settings", meta = (ExpandByDefault = true))
+	UPROPERTY(EditAnywhere, BlueprintReadWrite, Category = Settings, meta = (ExpandByDefault = true))
 	FAnimCurveOverrideList CurveOverrides;
 
-	UPROPERTY(EditAnywhere, BlueprintReadWrite, Category = "Blending", meta = (FrameTimeEditor))
-	float BlendTime = 0.2;
+	UPROPERTY(EditAnywhere, BlueprintReadWrite, Category = Blending, meta = (FrameTimeEditor))
+	float BlendTime = 0.2f;
 
 	// Set Blend Profiles (editable in the skeleton) to determine how the blending is distributed among your character's bones. It could be used to differentiate between upper body and lower body to blend timing.
 	UPROPERTY(EditAnywhere, BlueprintReadWrite, Category = Blending, meta = (UseAsBlendProfile = true))
@@ -80,6 +80,17 @@ struct FChooserPlayerSettings
 	UPROPERTY(EditAnywhere, BlueprintReadWrite, Category = Blending)
 	bool bUseInertialBlend = false;
 
+	// Experimental, this feature might be removed without warning, not for production use
+	UPROPERTY(EditAnywhere, BlueprintReadWrite, Category = Blending)
+	bool bForceBlendTo = false;
+
+	// Experimental, this feature might be removed without warning, not for production use
+	// the chooser player will continue playing an asset previously selected at time "PreviousStartTime", integrating it by the simulation "DeltaSeconds",
+	// unless the current "StartTime" is far away "MinDeltaTimeToForceBlendTo" from "PreviousStartTime":
+	// if |PreviousStartTime - StartTime| > MinDeltaTimeToForceBlendTo then bForceBlendTo = true and the chooser player will ask blend stack to blend to the same asset at a different "StartTime"
+	UPROPERTY(EditAnywhere, BlueprintReadWrite, Category = Blending, meta = (FrameTimeEditor))
+	float MinDeltaTimeToForceBlendTo = 0.f;
+
 	// Experimental, this feature might be removed without warning, not for production use
 	// Currently playing animation
 	UPROPERTY(Transient)
