PoseSearch - adding asset picker to select the owned database schema for the PoseSearchChooserColumn

diff --git a/Engine/Plugins/Animation/PoseSearch/Source/Editor/Private/Chooser/PoseSearchColumnEditor.cpp b/Engine/Plugins/Animation/PoseSearch/Source/Editor/Private/Chooser/PoseSearchColumnEditor.cpp
--- a/Engine/Plugins/Animation/PoseSearch/Source/Editor/Private/Chooser/PoseSearchColumnEditor.cpp
+++ b/Engine/Plugins/Animation/PoseSearch/Source/Editor/Private/Chooser/PoseSearchColumnEditor.cpp
@@ -3,9 +3,14 @@
 #include "PoseSearchColumnEditor.h"
 #include "Animation/AnimationAsset.h"
 #include "ChooserColumnHeader.h"
+#include "ContentBrowserModule.h"
+#include "Editor.h"
+#include "IContentBrowserSingleton.h"
 #include "ObjectChooserWidgetFactories.h"
 #include "PoseSearch/Chooser/PoseSearchChooserColumn.h"
 #include "PoseSearch/PoseSearchDatabase.h"
+#include "PoseSearch/PoseSearchSchema.h"
+#include "ScopedTransaction.h"
 #include "SPropertyAccessChainWidget.h"
 #include "Widgets/Input/SNumericEntryBox.h"
 #include "Widgets/SBoxPanel.h"
@@ -19,6 +24,45 @@ TSharedRef<SWidget> CreatePoseSearchColumnWidget(UChooserTable* Chooser, FChoose
 {
 	FPoseSearchColumn* PoseSearchColumn = static_cast<FPoseSearchColumn*>(Column);
 	
+	if (!PoseSearchColumn->InternalDatabase->Schema)
+	{
+		TSharedPtr<SWindow> PickerWindow;
+
+		// asset picker config to select Schema
+		FAssetPickerConfig AssetPickerConfig;
+		AssetPickerConfig.Filter.ClassPaths.Add(UPoseSearchSchema::StaticClass()->GetClassPathName());
+		AssetPickerConfig.Filter.bRecursiveClasses = true;
+		AssetPickerConfig.InitialAssetViewType = EAssetViewType::List;
+		AssetPickerConfig.OnAssetSelected = FOnAssetSelected::CreateLambda([PoseSearchColumn, &PickerWindow](const FAssetData& SelectedAsset)
+			{
+				if (UPoseSearchSchema* Schema = Cast<UPoseSearchSchema>(SelectedAsset.GetAsset()))
+				{
+					FScopedTransaction ScopedTransaction(LOCTEXT("AssignPoseSearchColumnDatabaseSchema", "Assign Pose Search Column Database Schema"));
+					PoseSearchColumn->InternalDatabase->Modify();
+					PoseSearchColumn->InternalDatabase->Schema = Schema;
+				}
+				PickerWindow->RequestDestroyWindow();
+			});
+
+		// Load the content browser module to display an asset picker
+		FContentBrowserModule& ContentBrowserModule = FModuleManager::LoadModuleChecked<FContentBrowserModule>("ContentBrowser");
+		PickerWindow = SNew(SWindow)
+			.Title(LOCTEXT("SelectPoseSearchColumnDatabaseSchema", "Pick Schema for owned Pose Search Database"))
+			.ClientSize(FVector2D(500, 600))
+			.SupportsMinimize(false)
+			.SupportsMaximize(false)
+			[
+				SNew(SBorder)
+					.BorderImage(FAppStyle::GetBrush("Menu.Background"))
+					[
+						ContentBrowserModule.Get().CreateAssetPicker(AssetPickerConfig)
+					]
+			];
+		
+		GEditor->EditorAddModalWindow(PickerWindow.ToSharedRef());
+		PickerWindow.Reset();
+	}
+
 	if (Row == UE::ChooserEditor::ColumnWidget_SpecialIndex_Fallback)
 	{
 		return SNullWidget::NullWidget;
