PoseSearch - adding missing APIs to FPoseSearchDatabaseAnimationAsset to support UMultiAnimAsset

diff --git a/Engine/Plugins/Animation/PoseSearch/Source/Runtime/Private/PoseSearchDatabase.cpp b/Engine/Plugins/Animation/PoseSearch/Source/Runtime/Private/PoseSearchDatabase.cpp
--- a/Engine/Plugins/Animation/PoseSearch/Source/Runtime/Private/PoseSearchDatabase.cpp
+++ b/Engine/Plugins/Animation/PoseSearch/Source/Runtime/Private/PoseSearchDatabase.cpp
@@ -555,6 +555,42 @@ float FPoseSearchDatabaseAnimationAsset::GetPlayLength(const FVector& BlendParam
 	return FPoseSearchDatabaseAnimationAssetBase::GetPlayLength(BlendParameters);
 }
 
+int32 FPoseSearchDatabaseAnimationAsset::GetNumRoles() const
+{
+	if (const UMultiAnimAsset* MultiAnimAsset = Cast<UMultiAnimAsset>(AnimAsset))
+	{
+		return MultiAnimAsset->GetNumRoles();
+	}
+	return FPoseSearchDatabaseAnimationAssetBase::GetNumRoles();
+}
+
+UE::PoseSearch::FRole FPoseSearchDatabaseAnimationAsset::GetRole(int32 RoleIndex) const
+{
+	if (const UMultiAnimAsset* MultiAnimAsset = Cast<UMultiAnimAsset>(AnimAsset))
+	{
+		return MultiAnimAsset->GetRole(RoleIndex);
+	}
+	return FPoseSearchDatabaseAnimationAssetBase::GetRole(RoleIndex);
+}
+
+UAnimationAsset* FPoseSearchDatabaseAnimationAsset::GetAnimationAssetForRole(const UE::PoseSearch::FRole& Role) const
+{
+	if (const UMultiAnimAsset* MultiAnimAsset = Cast<UMultiAnimAsset>(AnimAsset))
+	{
+		return MultiAnimAsset->GetAnimationAsset(Role);
+	}
+	return FPoseSearchDatabaseAnimationAssetBase::GetAnimationAssetForRole(Role);
+}
+
+FTransform FPoseSearchDatabaseAnimationAsset::GetRootTransformOriginForRole(const UE::PoseSearch::FRole& Role) const
+{
+	if (const UMultiAnimAsset* MultiAnimAsset = Cast<UMultiAnimAsset>(AnimAsset))
+	{
+		return MultiAnimAsset->GetOrigin(Role);
+	}
+	return FPoseSearchDatabaseAnimationAssetBase::GetRootTransformOriginForRole(Role);
+}
+
 #if WITH_EDITORONLY_DATA
 // @todo: Fix me!
 UClass* FPoseSearchDatabaseAnimationAsset::GetAnimationAssetStaticClass() const
@@ -693,7 +729,6 @@ USkeletalMesh* FPoseSearchDatabaseAnimationAsset::GetPreviewMeshForRole(const UE
 	{
 		return MultiAnimAsset->GetPreviewMesh(Role);
 	}
-
 	return FPoseSearchDatabaseAnimationAssetBase::GetPreviewMeshForRole(Role);
 }
 #endif // WITH_EDITOR
diff --git a/Engine/Plugins/Animation/PoseSearch/Source/Runtime/Public/PoseSearch/PoseSearchDatabase.h b/Engine/Plugins/Animation/PoseSearch/Source/Runtime/Public/PoseSearch/PoseSearchDatabase.h
--- a/Engine/Plugins/Animation/PoseSearch/Source/Runtime/Public/PoseSearch/PoseSearchDatabase.h
+++ b/Engine/Plugins/Animation/PoseSearch/Source/Runtime/Public/PoseSearch/PoseSearchDatabase.h
@@ -69,7 +69,7 @@ struct FPoseSearchDatabaseAnimationAssetBase
 	UE_API virtual float GetPlayLength(const FVector& BlendParameters) const;
 
 	virtual int32 GetNumRoles() const { return 1; }
-	virtual UE::PoseSearch::FRole GetRole(int32 RoleIndex) const { return UE::PoseSearch::DefaultRole; }
+	virtual UE::PoseSearch::FRole GetRole(int32 RoleIndex) const { check(GetNumRoles() == 1);  return UE::PoseSearch::DefaultRole; }
 	UE_API virtual UAnimationAsset* GetAnimationAssetForRole(const UE::PoseSearch::FRole& Role) const;
 	UE_API virtual FTransform GetRootTransformOriginForRole(const UE::PoseSearch::FRole& Role) const;
 #if WITH_EDITOR
@@ -209,6 +209,11 @@ struct FPoseSearchDatabaseAnimationAsset : public FPoseSearchDatabaseAnimationAs
 	UE_API virtual UObject* GetAnimationAsset() const override;
 	UE_API virtual float GetPlayLength(const FVector& BlendParameters) const override;
 
+	UE_API virtual int32 GetNumRoles() const override;
+	UE_API virtual UE::PoseSearch::FRole GetRole(int32 RoleIndex) const override;
+	UE_API virtual UAnimationAsset* GetAnimationAssetForRole(const UE::PoseSearch::FRole& Role) const override;
+	UE_API virtual FTransform GetRootTransformOriginForRole(const UE::PoseSearch::FRole& Role) const override;
+
 	friend bool operator==(const FPoseSearchDatabaseAnimationAsset& A, const FPoseSearchDatabaseAnimationAsset& B)
 	{
 		return static_cast<const FPoseSearchDatabaseAnimationAssetBase&>(A) == static_cast<const FPoseSearchDatabaseAnimationAssetBase&>(B)
