PoseSearch - fix for SDatabaseAssetListItem potentially referencing a deleted asset

diff --git a/Engine/Plugins/Animation/PoseSearch/Source/Editor/Private/PoseSearchDatabaseAssetListItem.cpp b/Engine/Plugins/Animation/PoseSearch/Source/Editor/Private/PoseSearchDatabaseAssetListItem.cpp
--- a/Engine/Plugins/Animation/PoseSearch/Source/Editor/Private/PoseSearchDatabaseAssetListItem.cpp
+++ b/Engine/Plugins/Animation/PoseSearch/Source/Editor/Private/PoseSearchDatabaseAssetListItem.cpp
@@ -372,8 +372,6 @@ namespace UE::PoseSearch
 		}
 		else
 		{
-			TSharedPtr<FDatabaseViewModel> ViewModel = EditorViewModel.Pin();
-
 			// Branch in
 			TSharedPtr<SImage> BranchInIconWidget;
 			{
@@ -410,16 +408,24 @@ namespace UE::PoseSearch
 			{
 				// Get item Icon
 				TSharedPtr<SImage> ItemIconWidget;
-				if (UPoseSearchDatabase* Database = ViewModel->GetPoseSearchDatabase())
 				{
-					if (const FPoseSearchDatabaseAnimationAssetBase* DatabaseAnimationAsset = Database->GetDatabaseAnimationAsset(SourceAssetIdx))
-					{
-						SAssignNew(ItemIconWidget, SImage)
-							.Image_Lambda([DatabaseAnimationAsset]
+					SAssignNew(ItemIconWidget, SImage)
+						.Image_Lambda([this]
+							{
+								UClass* AnimationAssetStaticClass = nullptr;
+								TSharedPtr<FDatabaseViewModel> ViewModelPtr = EditorViewModel.Pin();
+								if (const UPoseSearchDatabase* Database = ViewModelPtr->GetPoseSearchDatabase())
 								{
-									return FSlateIconFinder::FindIconBrushForClass(DatabaseAnimationAsset->GetAnimationAssetStaticClass());
-								});
-					}
+									if (const TSharedPtr<FDatabaseAssetTreeNode> AssetTreeNode = WeakAssetTreeNode.Pin())
+									{
+										if (const FPoseSearchDatabaseAnimationAssetBase* DatabaseAnimationAsset = Database->GetDatabaseAnimationAsset(AssetTreeNode->SourceAssetIdx))
+										{
+											AnimationAssetStaticClass = DatabaseAnimationAsset->GetAnimationAssetStaticClass();
+										}
+									}
+								}
+								return FSlateIconFinder::FindIconBrushForClass(AnimationAssetStaticClass);
+							});
 				}
 				
 				SAssignNew(AssetThumbnailOverlay, SOverlay)
@@ -481,6 +487,7 @@ namespace UE::PoseSearch
 			}
 			
 			// Picker
+			TSharedPtr<FDatabaseViewModel> ViewModel = EditorViewModel.Pin();
 			TSharedPtr<SObjectPropertyEntryBox> AssetPickerWidget;
 			if (UPoseSearchDatabase* Database = ViewModel->GetPoseSearchDatabase())
 			{
