PoseSearch - fix for crash in FInteractionIsland::DoSearch_AnyThread when multiple searches happen in a single interaction island

diff --git a/Engine/Plugins/Animation/PoseSearch/Source/Runtime/Private/PoseSearchInteractionIsland.cpp b/Engine/Plugins/Animation/PoseSearch/Source/Runtime/Private/PoseSearchInteractionIsland.cpp
--- a/Engine/Plugins/Animation/PoseSearch/Source/Runtime/Private/PoseSearchInteractionIsland.cpp
+++ b/Engine/Plugins/Animation/PoseSearch/Source/Runtime/Private/PoseSearchInteractionIsland.cpp
@@ -1012,6 +1012,7 @@ bool FInteractionIsland::DoSearch_AnyThread(const UObject* AnimContext, const TC
 		}
 
 		const int32 NumRoles = InteractionSearchContext.Num();
+		Contexts.Reset();
 		Contexts.SetNum(NumRoles);
 
 		FSearchContext SearchContext(0.f, FFloatInterval(0.f, 0.f), FPoseSearchEvent());
