PoseSearch - fix for UPoseSearchDatabase::SynchronizeWithExternalDependencies synchronizing mistakenly transient assets into a database, leading to reference garbaged collected assets

diff --git a/Engine/Plugins/Animation/PoseSearch/Source/Runtime/Private/PoseSearchAssetSampler.cpp b/Engine/Plugins/Animation/PoseSearch/Source/Runtime/Private/PoseSearchAssetSampler.cpp
--- a/Engine/Plugins/Animation/PoseSearch/Source/Runtime/Private/PoseSearchAssetSampler.cpp
+++ b/Engine/Plugins/Animation/PoseSearch/Source/Runtime/Private/PoseSearchAssetSampler.cpp
@@ -254,6 +254,8 @@ FAnimationAssetSampler::~FAnimationAssetSampler() = default;
 
 void FAnimationAssetSampler::Init(const UAnimationAsset* InAnimationAsset, const FTransform& InRootTransformOrigin, const FVector& InBlendParameters, int32 InRootTransformSamplingRate, bool bPreProcessRootTransform, bool bInEnforceCompressedDataSampling)
 {
+	check(!InAnimationAsset || (InAnimationAsset->GetPackage() && !InAnimationAsset->HasAnyFlags(RF_MirroredGarbage)));
+	
 	AnimationAssetPtr = InAnimationAsset;
 	RootTransformOrigin = InRootTransformOrigin;
 	BlendParameters = InBlendParameters;
diff --git a/Engine/Plugins/Animation/PoseSearch/Source/Runtime/Private/PoseSearchDatabase.cpp b/Engine/Plugins/Animation/PoseSearch/Source/Runtime/Private/PoseSearchDatabase.cpp
--- a/Engine/Plugins/Animation/PoseSearch/Source/Runtime/Private/PoseSearchDatabase.cpp
+++ b/Engine/Plugins/Animation/PoseSearch/Source/Runtime/Private/PoseSearchDatabase.cpp
@@ -24,6 +24,7 @@
 #include "UObject/ObjectSaveContext.h"
 
 #if WITH_EDITOR
+#include "ScopedTransaction.h"
 #include "AssetRegistry/AssetRegistryModule.h"
 #endif //WITH_EDITOR
 
@@ -44,6 +45,7 @@ DEFINE_STAT(STAT_PoseSearch_PCAKNN);
 DEFINE_STAT(STAT_PoseSearch_VPTree);
 DEFINE_STAT(STAT_PoseSearch_Event);
 
+#define LOCTEXT_NAMESPACE "PoseSearchDatabase"
 
 namespace UE::PoseSearch
 {
@@ -1381,16 +1383,16 @@ void UPoseSearchDatabase::SynchronizeWithExternalDependencies()
 void UPoseSearchDatabase::SynchronizeWithExternalDependencies(TConstArrayView<UAnimSequenceBase*> SequencesBase)
 {
 	// cannot use TSet since FInstancedStruct doesn't implement GetTypeHash
-	TArray<FPoseSearchDatabaseAnimationAsset> NewDatabaseAnimationAssets;
+	TArray<FPoseSearchDatabaseAnimationAsset> NewDatabaseAnimationAssets = DatabaseAnimationAssets;
+	TArray<FPoseSearchDatabaseAnimationAsset> TempDatabaseAnimationAssets;
 
 	// collecting all the database AnimationAsset(s) that don't require synchronization
-	
-	for (FPoseSearchDatabaseAnimationAsset& DatabaseAnimationAsset : DatabaseAnimationAssets)
+	for (FPoseSearchDatabaseAnimationAsset& DatabaseAnimationAsset : NewDatabaseAnimationAssets)
 	{
 		const bool bRequiresSynchronization = DatabaseAnimationAsset.IsSynchronizedWithExternalDependency() && SequencesBase.Contains(DatabaseAnimationAsset.GetAnimationAsset());
 		if (!bRequiresSynchronization)
 		{
-			NewDatabaseAnimationAssets.Add(DatabaseAnimationAsset);
+			TempDatabaseAnimationAssets.Add(DatabaseAnimationAsset);
 		}
 	}
 
@@ -1399,43 +1401,59 @@ void UPoseSearchDatabase::SynchronizeWithExternalDependencies(TConstArrayView<UA
 	{
 		if (SequenceBase)
 		{
-			for (const FAnimNotifyEvent& NotifyEvent : SequenceBase->Notifies)
+			UPackage* SequenceBasePackage = SequenceBase->GetPackage();
+			if (!SequenceBasePackage)
+			{
+				UE_LOG(LogPoseSearch, Warning, TEXT("AnimSequenceBase '%s' cannot be synchronized with PoseSearchDatabase '%s' because it's not in a package"), *SequenceBase->GetName(), *GetName());
+			}
+			else if (SequenceBasePackage->HasAnyFlags(RF_Transient | RF_DuplicateTransient | RF_MirroredGarbage))
+			{
+				UE_LOG(LogPoseSearch, Warning, TEXT("AnimSequenceBase '%s' cannot be synchronized with PoseSearchDatabase '%s' because it's in a package transient or garbaged package"), *SequenceBase->GetName(), *GetName());
+			}
+			else if (SequenceBase->HasAnyFlags(RF_Transient | RF_DuplicateTransient | RF_MirroredGarbage))
 			{
-				if (const UAnimNotifyState_PoseSearchBranchIn* PoseSearchBranchIn = Cast<UAnimNotifyState_PoseSearchBranchIn>(NotifyEvent.NotifyStateClass))
+				UE_LOG(LogPoseSearch, Warning, TEXT("AnimSequenceBase '%s' cannot be synchronized with PoseSearchDatabase '%s' because it's transient or garbaged"), *SequenceBase->GetName(), *GetName());
+			}
+			else
+			{
+				for (const FAnimNotifyEvent& NotifyEvent : SequenceBase->Notifies)
 				{
-					if (PoseSearchBranchIn->Database == this)
+					if (const UAnimNotifyState_PoseSearchBranchIn* PoseSearchBranchIn = Cast<UAnimNotifyState_PoseSearchBranchIn>(NotifyEvent.NotifyStateClass))
 					{
-						auto GetSamplingRange = [](const FAnimNotifyEvent& NotifyEvent, const UAnimSequenceBase* SequenceBase) -> FFloatInterval
+						if (PoseSearchBranchIn->Database == this)
 						{
-							FFloatInterval SamplingRange(NotifyEvent.GetTime(), NotifyEvent.GetTime() + NotifyEvent.GetDuration());
-							if (SamplingRange.Min <= NotifyEvent.TriggerTimeOffset && SamplingRange.Max >= SequenceBase->GetPlayLength() - NotifyEvent.TriggerTimeOffset)
+							auto GetSamplingRange = [](const FAnimNotifyEvent& NotifyEvent, const UAnimSequenceBase* SequenceBase) -> FFloatInterval
 							{
-								SamplingRange = FFloatInterval(0.f, 0.f);
-							}
-							return SamplingRange;
-						};
-
-						FPoseSearchDatabaseAnimationAsset DatabaseAnimationAsset;
-						DatabaseAnimationAsset.AnimAsset = SequenceBase;
-						DatabaseAnimationAsset.SamplingRange = GetSamplingRange(NotifyEvent, SequenceBase);
-						DatabaseAnimationAsset.BranchInId = PoseSearchBranchIn->GetBranchInId();
-						NewDatabaseAnimationAssets.Add(DatabaseAnimationAsset);
+								FFloatInterval SamplingRange(NotifyEvent.GetTime(), NotifyEvent.GetTime() + NotifyEvent.GetDuration());
+								if (SamplingRange.Min <= NotifyEvent.TriggerTimeOffset && SamplingRange.Max >= SequenceBase->GetPlayLength() - NotifyEvent.TriggerTimeOffset)
+								{
+									SamplingRange = FFloatInterval(0.f, 0.f);
+								}
+								return SamplingRange;
+							};
+
+							FPoseSearchDatabaseAnimationAsset DatabaseAnimationAsset;
+							DatabaseAnimationAsset.AnimAsset = SequenceBase;
+							DatabaseAnimationAsset.SamplingRange = GetSamplingRange(NotifyEvent, SequenceBase);
+							DatabaseAnimationAsset.BranchInId = PoseSearchBranchIn->GetBranchInId();
+							TempDatabaseAnimationAssets.Add(DatabaseAnimationAsset);
+						}
 					}
 				}
 			}
 		}
 	}
 
-	// updating DatabaseAnimationAssets from NewDatabaseAnimationAssets preserving the original sorting
+	// updating NewDatabaseAnimationAssets from TempDatabaseAnimationAssets preserving the original sorting
 	bool bModified = false;
 	for (int32 AnimationAssetIndex = GetNumAnimationAssets() - 1; AnimationAssetIndex >= 0; --AnimationAssetIndex)
 	{
-		FPoseSearchDatabaseAnimationAsset& DatabaseAnimationAsset = DatabaseAnimationAssets[AnimationAssetIndex];
+		FPoseSearchDatabaseAnimationAsset& DatabaseAnimationAsset = NewDatabaseAnimationAssets[AnimationAssetIndex];
 		
 		int32 FoundIndex = -1;
-		for(int i=0; i < NewDatabaseAnimationAssets.Num(); i++)
+		for(int i=0; i < TempDatabaseAnimationAssets.Num(); i++)
 		{
-			const FPoseSearchDatabaseAnimationAssetBase& NewDatabaseAnimationAsset = NewDatabaseAnimationAssets[i];
+			const FPoseSearchDatabaseAnimationAssetBase& NewDatabaseAnimationAsset = TempDatabaseAnimationAssets[i];
 			if (DatabaseAnimationAsset.UpdateFrom(NewDatabaseAnimationAsset))
 			{
 				FoundIndex = i;
@@ -1446,37 +1464,25 @@ void UPoseSearchDatabase::SynchronizeWithExternalDependencies(TConstArrayView<UA
 		
 		if (FoundIndex >= 0)
 		{
-			NewDatabaseAnimationAssets.RemoveAt(FoundIndex); 
+			TempDatabaseAnimationAssets.RemoveAt(FoundIndex); 
 		}
 		else
 		{
-			if (!bModified)
-			{
-				Modify();
-				bModified = true;
-			}
-
-			DatabaseAnimationAssets.RemoveAt(AnimationAssetIndex);
+			NewDatabaseAnimationAssets.RemoveAt(AnimationAssetIndex);
 		}
 	}
 
-	// adding the remaining NewDatabaseAnimationAsset(s) from NewDatabaseAnimationAssets
-	if (!NewDatabaseAnimationAssets.IsEmpty())
+	// adding the remaining TempDatabaseAnimationAsset(s) from TempDatabaseAnimationAssets
+	for (const FPoseSearchDatabaseAnimationAsset& TempDatabaseAnimationAsset : TempDatabaseAnimationAssets)
 	{
-		if (!bModified)
-		{
-			Modify();
-			bModified = true;
-		}
-
-		for (const FPoseSearchDatabaseAnimationAsset& NewDatabaseAnimationAsset : NewDatabaseAnimationAssets)
-		{
-			DatabaseAnimationAssets.Add(NewDatabaseAnimationAsset);
-		}
+		NewDatabaseAnimationAssets.Add(TempDatabaseAnimationAsset);
 	}
 
-	if (bModified)
+	if (DatabaseAnimationAssets != NewDatabaseAnimationAssets)
 	{
+		Modify();
+
+		DatabaseAnimationAssets = NewDatabaseAnimationAssets;
 		NotifySynchronizeWithExternalDependencies();
 	}
 }
@@ -1514,7 +1520,8 @@ void UPoseSearchDatabase::TestSynchronizeWithExternalDependencies()
 void UPoseSearchDatabase::PreSaveRoot(FObjectPreSaveRootContext ObjectSaveContext)
 {
 #if WITH_EDITOR
-	// in case the database desynchronized with the UAnimNotifyState_PoseSearchBranchIn referencing it, we need to resyncrhonize
+	FScopedTransaction ScopedTransaction(LOCTEXT("PoseSearchDatabaseSynchronizeWithExternalDependencies", "Pose Search Database Synchronize With External Dependencies"));
+	// in case the database desynchronized with the UAnimNotifyState_PoseSearchBranchIn referencing it, we need to resynchronize
 	SynchronizeWithExternalDependencies();
 #endif
 
@@ -2567,3 +2574,5 @@ void UPoseSearchDatabase::SearchEvent(UE::PoseSearch::FSearchContext& SearchCont
 	SearchContext.Track(this);
 #endif // UE_POSE_SEARCH_TRACE_ENABLED
 }
+
+#undef LOCTEXT_NAMESPACE
\ No newline at end of file
diff --git a/Engine/Plugins/Animation/PoseSearch/Source/Runtime/Private/PoseSearchDerivedData.cpp b/Engine/Plugins/Animation/PoseSearch/Source/Runtime/Private/PoseSearchDerivedData.cpp
--- a/Engine/Plugins/Animation/PoseSearch/Source/Runtime/Private/PoseSearchDerivedData.cpp
+++ b/Engine/Plugins/Animation/PoseSearch/Source/Runtime/Private/PoseSearchDerivedData.cpp
@@ -1897,6 +1897,12 @@ void FPoseSearchDatabaseAsyncCacheTask::OnGetComplete(UE::DerivedData::FCacheGet
 							{
 								UE_LOG(LogPoseSearch, Warning, TEXT("OnGetComplete - - No asset has been selected."));
 							}
+							else if (!DatabaseAsset->GetAnimationAsset()->GetPackage() || DatabaseAsset->GetAnimationAsset()->HasAnyFlags(RF_MirroredGarbage))
+							{
+								UE_LOG(LogPoseSearch, Error, TEXT("%s - %s BuildIndex Failed because '%s' referenced asset '%s' is garbage or is not contained in a package"), *LexToString(FullIndexKey.Hash), *MainDatabaseName, *DependentDatabaseName, *DatabaseAsset->GetAnimationAsset()->GetName());
+								ResetSearchIndex();
+								return;
+							}
 							// @todo: Find a way to prevent accessing UObjects while GC since this func is ran async. Commenting out for now to prevent crashes.
 							/*else if (!DatabaseAsset->IsSkeletonCompatible(Database->Schema))
 							{
diff --git a/Engine/Plugins/Animation/PoseSearch/Source/Runtime/Private/PoseSearchDerivedDataKey.cpp b/Engine/Plugins/Animation/PoseSearch/Source/Runtime/Private/PoseSearchDerivedDataKey.cpp
--- a/Engine/Plugins/Animation/PoseSearch/Source/Runtime/Private/PoseSearchDerivedDataKey.cpp
+++ b/Engine/Plugins/Animation/PoseSearch/Source/Runtime/Private/PoseSearchDerivedDataKey.cpp
@@ -171,6 +171,7 @@ FArchive& FKeyBuilder::operator<<(FName& Name)
 
 FArchive& FKeyBuilder::TryAddDependency(UObject* Object, bool bAddToPartialKeyHashes)
 {
+	// @todo: add RF_NeedPostLoadSubobjects, RF_NeedInitialization, RF_NeedLoad, RF_WillBeLoaded? 
 	if (Object->HasAnyFlags(RF_NeedPostLoad))
 	{
 		if (bPerformConditionalPostLoad)
@@ -221,19 +222,42 @@ FArchive& FKeyBuilder::operator<<(class UObject*& Object)
 		return *this;
 	}
 
-	if (Object->HasAnyFlags(RF_Transient))
+	if (Object->HasAnyFlags(RF_Transient | RF_MirroredGarbage | RF_NewerVersionExists | RF_BeginDestroyed | RF_FinishDestroyed | RF_DuplicateTransient | RF_NonPIEDuplicateTransient) || IsExcludedType(Object))
 	{
 #if UE_POSE_SEARCH_DERIVED_DATA_LOGGING
-		UE_LOG(LogPoseSearch, Log, TEXT("Transient '%s' (%s)"), *Object->GetName(), *Object->GetClass()->GetName());
-#endif
-		return *this;
-	}
-
-	if (IsExcludedType(Object))
-	{
-#if UE_POSE_SEARCH_DERIVED_DATA_LOGGING
-		UE_LOG(LogPoseSearch, Log, TEXT("Excluded '%s' (%s)"), *Object->GetName(), *Object->GetClass()->GetName());
-#endif
+		if (Object->HasAnyFlags(RF_Transient))
+		{
+			UE_LOG(LogPoseSearch, Log, TEXT("Transient '%s' (%s)"), *Object->GetName(), *Object->GetClass()->GetName());
+		}
+		if (Object->HasAnyFlags(RF_MirroredGarbage))
+		{
+			UE_LOG(LogPoseSearch, Log, TEXT("MirroredGarbage '%s' (%s)"), *Object->GetName(), *Object->GetClass()->GetName());
+		}
+		if (Object->HasAnyFlags(RF_NewerVersionExists))
+		{
+			UE_LOG(LogPoseSearch, Log, TEXT("NewerVersionExists '%s' (%s)"), *Object->GetName(), *Object->GetClass()->GetName());
+		}
+		if (Object->HasAnyFlags(RF_BeginDestroyed))
+		{
+			UE_LOG(LogPoseSearch, Log, TEXT("BeginDestroyed '%s' (%s)"), *Object->GetName(), *Object->GetClass()->GetName());
+		}
+		if (Object->HasAnyFlags(RF_FinishDestroyed))
+		{
+			UE_LOG(LogPoseSearch, Log, TEXT("FinishDestroyed '%s' (%s)"), *Object->GetName(), *Object->GetClass()->GetName());
+		}
+		if (Object->HasAnyFlags(RF_DuplicateTransient))
+		{
+			UE_LOG(LogPoseSearch, Log, TEXT("DuplicateTransient '%s' (%s)"), *Object->GetName(), *Object->GetClass()->GetName());
+		}
+		if (Object->HasAnyFlags(RF_NonPIEDuplicateTransient))
+		{
+			UE_LOG(LogPoseSearch, Log, TEXT("DuplicateTransient '%s' (%s)"), *Object->GetName(), *Object->GetClass()->GetName());
+		}
+		if (IsExcludedType(Object))
+		{
+			UE_LOG(LogPoseSearch, Log, TEXT("Excluded '%s' (%s)"), *Object->GetName(), *Object->GetClass()->GetName());
+		}
+#endif // UE_POSE_SEARCH_DERIVED_DATA_LOGGING
 		return *this;
 	}
 
diff --git a/Engine/Plugins/Animation/PoseSearch/Source/Runtime/Public/PoseSearch/PoseSearchDefines.h b/Engine/Plugins/Animation/PoseSearch/Source/Runtime/Public/PoseSearch/PoseSearchDefines.h
--- a/Engine/Plugins/Animation/PoseSearch/Source/Runtime/Public/PoseSearch/PoseSearchDefines.h
+++ b/Engine/Plugins/Animation/PoseSearch/Source/Runtime/Public/PoseSearch/PoseSearchDefines.h
@@ -43,7 +43,7 @@ static constexpr int32 PreallocatedSearchesNum = 4;
 static constexpr float FiniteDelta = 1 / 60.0f; // Time delta used for computing pose derivatives
 static constexpr int32 MaxNumberOfCollectedPoseCandidatesPerDatabase = 256;
 static constexpr int32 InvalidSearchId = INDEX_NONE; // @todo: it should be 0, since the search id is an hash!
-static constexpr int32 DatabaseIndexDerivedDataCacheKeyVersion = 2713; // bump this value to force reindex ALL the UPoseSearchDatabases, without having to change the global FDevSystemGuids::POSESEARCHDB_DERIVEDDATA_VER
+static constexpr int32 DatabaseIndexDerivedDataCacheKeyVersion = 2719; // bump this value to force reindex ALL the UPoseSearchDatabases, without having to change the global FDevSystemGuids::POSESEARCHDB_DERIVEDDATA_VER
 static constexpr int32 DefaultTickPriority = 100;
 
 template<uint32 Alignment = DEFAULT_ALIGNMENT> class TMemStackSetAllocator : public TSetAllocator<TSparseArrayAllocator<TMemStackAllocator<Alignment>, TMemStackAllocator<Alignment>>, TMemStackAllocator<Alignment>> {};
