PoseSearch - adding fault tolerance on missing roles when drawing debug lines using FDebugDrawParams (it could happen when switching from a MM interaction database to an non interaction database while debugging replays from the pose search debugger track)

diff --git a/Engine/Plugins/Animation/PoseSearch/Source/Runtime/Private/PoseSearchContext.cpp b/Engine/Plugins/Animation/PoseSearch/Source/Runtime/Private/PoseSearchContext.cpp
--- a/Engine/Plugins/Animation/PoseSearch/Source/Runtime/Private/PoseSearchContext.cpp
+++ b/Engine/Plugins/Animation/PoseSearch/Source/Runtime/Private/PoseSearchContext.cpp
@@ -345,26 +345,28 @@ FVector FDebugDrawParams::ExtractPosition(TConstArrayView<float> PoseVector, flo
 			return WorldBonePosition;
 		}
 
-		const int32 RoleIndex = RoleToIndex[Role];
-		const IPoseHistory* PoseHistory = PoseHistories[RoleIndex];
-		check(PoseHistory);
-			
-		if (const USkeleton* Skeleton = Schema->GetSkeleton(Role))
+		if (const int32* RoleIndex = RoleToIndex.Find(Role))
 		{
-			const FBoneIndexType BoneIndexType = SchemaBoneIdx == TrajectorySchemaBoneIdx ? ComponentSpaceIndexType : Schema->GetBoneReferences(Role)[SchemaBoneIdx].BoneIndex;
+			const IPoseHistory* PoseHistory = PoseHistories[*RoleIndex];
+			check(PoseHistory);
 
-			FTransform WorldBoneTransform;
-			if (PoseHistory->GetTransformAtTime(SampleTimeOffset + PermutationSampleTimeOffset, WorldBoneTransform, Skeleton, BoneIndexType, WorldSpaceIndexType))
+			if (const USkeleton* Skeleton = Schema->GetSkeleton(Role))
 			{
-				return WorldBoneTransform.GetTranslation();
+				const FBoneIndexType BoneIndexType = SchemaBoneIdx == TrajectorySchemaBoneIdx ? ComponentSpaceIndexType : Schema->GetBoneReferences(Role)[SchemaBoneIdx].BoneIndex;
+
+				FTransform WorldBoneTransform;
+				if (PoseHistory->GetTransformAtTime(SampleTimeOffset + PermutationSampleTimeOffset, WorldBoneTransform, Skeleton, BoneIndexType, WorldSpaceIndexType))
+				{
+					return WorldBoneTransform.GetTranslation();
+				}
 			}
-		}
 
-		if (SchemaBoneIdx > RootSchemaBoneIdx)
-		{
-			if (const USkinnedMeshComponent* Mesh = GetMesh(AnimContexts, RoleIndex))
+			if (SchemaBoneIdx > RootSchemaBoneIdx)
 			{
-				return Mesh->GetSocketTransform(Schema->GetBoneReferences(Role)[SchemaBoneIdx].BoneName).GetTranslation();
+				if (const USkinnedMeshComponent* Mesh = GetMesh(AnimContexts, *RoleIndex))
+				{
+					return Mesh->GetSocketTransform(Schema->GetBoneReferences(Role)[SchemaBoneIdx].BoneName).GetTranslation();
+				}
 			}
 		}
 	}
@@ -482,26 +484,28 @@ FQuat FDebugDrawParams::ExtractRotation(TConstArrayView<float> PoseVector, float
 			}
 		}
 
-		const int32 RoleIndex = RoleToIndex[Role];
-		const IPoseHistory* PoseHistory = PoseHistories[RoleIndex];
-		check(PoseHistory);
-			
-		if (const USkeleton* Skeleton = Schema->GetSkeleton(Role))
+		if (const int32* RoleIndex = RoleToIndex.Find(Role))
 		{
-			const FBoneIndexType BoneIndexType = SchemaBoneIdx == TrajectorySchemaBoneIdx ? ComponentSpaceIndexType : Schema->GetBoneReferences(Role)[SchemaBoneIdx].BoneIndex;
+			const IPoseHistory* PoseHistory = PoseHistories[*RoleIndex];
+			check(PoseHistory);
 
-			FTransform WorldBoneTransform;
-			if (PoseHistory->GetTransformAtTime(SampleTimeOffset + PermutationSampleTimeOffset, WorldBoneTransform, Skeleton, BoneIndexType, WorldSpaceIndexType))
+			if (const USkeleton* Skeleton = Schema->GetSkeleton(Role))
 			{
-				return WorldBoneTransform.GetRotation();
+				const FBoneIndexType BoneIndexType = SchemaBoneIdx == TrajectorySchemaBoneIdx ? ComponentSpaceIndexType : Schema->GetBoneReferences(Role)[SchemaBoneIdx].BoneIndex;
+
+				FTransform WorldBoneTransform;
+				if (PoseHistory->GetTransformAtTime(SampleTimeOffset + PermutationSampleTimeOffset, WorldBoneTransform, Skeleton, BoneIndexType, WorldSpaceIndexType))
+				{
+					return WorldBoneTransform.GetRotation();
+				}
 			}
-		}
 
-		if (SchemaBoneIdx > RootSchemaBoneIdx)
-		{
-			if (const USkinnedMeshComponent* Mesh = GetMesh(AnimContexts, RoleIndex))
+			if (SchemaBoneIdx > RootSchemaBoneIdx)
 			{
-				return Mesh->GetSocketTransform(Schema->GetBoneReferences(Role)[SchemaBoneIdx].BoneName).GetRotation();
+				if (const USkinnedMeshComponent* Mesh = GetMesh(AnimContexts, *RoleIndex))
+				{
+					return Mesh->GetSocketTransform(Schema->GetBoneReferences(Role)[SchemaBoneIdx].BoneName).GetRotation();
+				}
 			}
 		}
 	}
@@ -512,7 +516,11 @@ FQuat FDebugDrawParams::ExtractRotation(TConstArrayView<float> PoseVector, float
 FTransform FDebugDrawParams::GetRootBoneTransform(const FRole& Role, float SampleTimeOffset) const
 {
 	FTransform RootBoneTransform = FTransform::Identity;
-	PoseHistories[RoleToIndex[Role]]->GetTransformAtTime(SampleTimeOffset, RootBoneTransform, nullptr, RootBoneIndexType, WorldSpaceIndexType);
+
+	if (const int32* RoleIndex = RoleToIndex.Find(Role))
+	{
+		PoseHistories[*RoleIndex]->GetTransformAtTime(SampleTimeOffset, RootBoneTransform, nullptr, RootBoneIndexType, WorldSpaceIndexType);
+	}
 	return RootBoneTransform;
 }
 
