Pose History Reference access for ABP

diff --git a/Engine/Plugins/Animation/PoseSearch/Source/Runtime/Private/PoseSearchHistoryCollectorAnimNodeLibrary.cpp b/Engine/Plugins/Animation/PoseSearch/Source/Runtime/Private/PoseSearchHistoryCollectorAnimNodeLibrary.cpp
--- a/Engine/Plugins/Animation/PoseSearch/Source/Runtime/Private/PoseSearchHistoryCollectorAnimNodeLibrary.cpp
+++ b/Engine/Plugins/Animation/PoseSearch/Source/Runtime/Private/PoseSearchHistoryCollectorAnimNodeLibrary.cpp
@@ -17,6 +17,17 @@ void UPoseSearchHistoryCollectorAnimNodeLibrary::ConvertToPoseHistoryNodePure(co
 	Result = (ConversionResult == EAnimNodeReferenceConversionResult::Succeeded);
 }
 
+
+FPoseHistoryReference UPoseSearchHistoryCollectorAnimNodeLibrary::GetPoseHistoryReference(const FPoseSearchHistoryCollectorAnimNodeReference& Node)
+{
+	if (const FAnimNode_PoseSearchHistoryCollector* PoseSearchHistoryCollectorNodePtr = Node.GetAnimNodePtr<FAnimNode_PoseSearchHistoryCollector>())
+	{
+		return PoseSearchHistoryCollectorNodePtr->GetPoseHistoryReference();
+	}
+	return FPoseHistoryReference();
+}
+
+
 void UPoseSearchHistoryCollectorAnimNodeLibrary::GetPoseHistoryNodeTransformTrajectory(const FPoseSearchHistoryCollectorAnimNodeReference& PoseSearchHistoryCollectorNode, FTransformTrajectory& Trajectory)
 {
 	if (FAnimNode_PoseSearchHistoryCollector* PoseSearchHistoryCollectorNodePtr = PoseSearchHistoryCollectorNode.GetAnimNodePtr<FAnimNode_PoseSearchHistoryCollector>())
diff --git a/Engine/Plugins/Animation/PoseSearch/Source/Runtime/Public/PoseSearch/AnimNode_PoseSearchHistoryCollector.h b/Engine/Plugins/Animation/PoseSearch/Source/Runtime/Public/PoseSearch/AnimNode_PoseSearchHistoryCollector.h
--- a/Engine/Plugins/Animation/PoseSearch/Source/Runtime/Public/PoseSearch/AnimNode_PoseSearchHistoryCollector.h
+++ b/Engine/Plugins/Animation/PoseSearch/Source/Runtime/Public/PoseSearch/AnimNode_PoseSearchHistoryCollector.h
@@ -116,6 +116,8 @@ public:
 	const UE::PoseSearch::FPoseHistory& GetPoseHistory() const { check(PoseHistoryPtr.IsValid()); return *PoseHistoryPtr.Get(); }
 	UE::PoseSearch::FPoseHistory& GetPoseHistory() { check(PoseHistoryPtr.IsValid()); return *PoseHistoryPtr.Get(); }
 	
+	FPoseHistoryReference GetPoseHistoryReference() const { check(PoseHistoryPtr.IsValid()); return FPoseHistoryReference {PoseHistoryPtr}; }
+	
 	UE_DEPRECATED(5.6, "Use this::IPoseHistory::GenerateTrajectory interface instead")
 	UE_API void GenerateTrajectory(const UAnimInstance* InAnimInstance);
 
diff --git a/Engine/Plugins/Animation/PoseSearch/Source/Runtime/Public/PoseSearch/PoseSearchHistoryCollectorAnimNodeLibrary.h b/Engine/Plugins/Animation/PoseSearch/Source/Runtime/Public/PoseSearch/PoseSearchHistoryCollectorAnimNodeLibrary.h
--- a/Engine/Plugins/Animation/PoseSearch/Source/Runtime/Public/PoseSearch/PoseSearchHistoryCollectorAnimNodeLibrary.h
+++ b/Engine/Plugins/Animation/PoseSearch/Source/Runtime/Public/PoseSearch/PoseSearchHistoryCollectorAnimNodeLibrary.h
@@ -2,6 +2,7 @@
 
 #pragma once
 
+#include "PoseSearchHistory.h"
 #include "Animation/AnimExecutionContext.h"
 #include "Animation/AnimNodeReference.h"
 #include "Kismet/BlueprintFunctionLibrary.h"
@@ -43,6 +44,9 @@ public:
 
 	UFUNCTION(BlueprintCallable, Category = "Animation|PoseHistory", meta = (BlueprintThreadSafe, DisplayName = "Set Pose History Node Trajectory"))
 	static UE_API void SetPoseHistoryNodeTransformTrajectory(const FPoseSearchHistoryCollectorAnimNodeReference& PoseSearchHistoryCollectorNode, const FTransformTrajectory& Trajectory);
+
+	UFUNCTION(BlueprintPure, Category = "Animation|Pose Search|Experimental", meta = (BlueprintThreadSafe))
+	static UE_API FPoseHistoryReference GetPoseHistoryReference(const FPoseSearchHistoryCollectorAnimNodeReference& Node);
 	
 	PRAGMA_DISABLE_DEPRECATION_WARNINGS
 	
