PoseSearch - fixes for database index feature values (pose data) pruning, allowing to search a single KNNQueryNumNeighbors efficiently via reconstructed values from PCA values

diff --git a/Engine/Plugins/Animation/PoseSearch/Source/Runtime/Private/PoseSearchContext.cpp b/Engine/Plugins/Animation/PoseSearch/Source/Runtime/Private/PoseSearchContext.cpp
--- a/Engine/Plugins/Animation/PoseSearch/Source/Runtime/Private/PoseSearchContext.cpp
+++ b/Engine/Plugins/Animation/PoseSearch/Source/Runtime/Private/PoseSearchContext.cpp
@@ -12,6 +12,7 @@
 #include "PoseSearch/PoseSearchFeatureChannel_Position.h"
 #include "PoseSearch/PoseSearchHistory.h"
 #include "PoseSearch/PoseSearchSchema.h"
+#include "PoseSearchIndex.inl"
 
 namespace UE::PoseSearch
 {
@@ -759,6 +760,15 @@ FCachedQuery::FCachedQuery(const UPoseSearchSchema* InSchema)
 	Values.SetNumZeroed(Schema->SchemaCardinality);
 }
 
+//////////////////////////////////////////////////////////////////////////
+// FCachedPCAQuery
+FCachedPCAQuery::FCachedPCAQuery(const UPoseSearchDatabase* InDatabase)
+{
+	check(InDatabase);
+	Database = InDatabase;
+	Values.SetNumZeroed(Database->GetSearchIndex().GetNumberOfPrincipalComponents());
+}
+
 #if UE_POSE_SEARCH_TRACE_ENABLED
 //////////////////////////////////////////////////////////////////////////
 // FSearchContext::FBestPoseCandidates
@@ -915,7 +925,7 @@ void FSearchContext::UpdateContinuingPoseValues()
 		{
 			const int32 NumDimensions = ContinuingPoseValuesDatabasePoseIdx.Database->Schema->SchemaCardinality;
 			ContinuingPoseValuesData.SetNumUninitialized(NumDimensions, EAllowShrinking::No);
-			ContinuingPoseValues = SearchIndex.GetReconstructedPoseValues(ContinuingPoseValuesDatabasePoseIdx.PoseIdx, MakeArrayView(ContinuingPoseValuesData).Slice(NumDimensions, NumDimensions));
+			ContinuingPoseValues = SearchIndex.GetReconstructedPoseValues(ContinuingPoseValuesDatabasePoseIdx.PoseIdx, ContinuingPoseValuesData);
 		}
 		else
 		{
@@ -1307,6 +1317,28 @@ TConstArrayView<float> FSearchContext::GetOrBuildQuery(const UPoseSearchSchema*
 	return Schema->BuildQuery(*this);
 }
 
+TConstArrayView<float> FSearchContext::GetOrBuildPCAQuery(const UPoseSearchDatabase* Database)
+{
+	QUICK_SCOPE_CYCLE_COUNTER(STAT_PoseSearch_GetOrBuildPCAQuery);
+
+	check(Database);
+	if (const FCachedPCAQuery* FoundCachedPCAQuery = CachedPCAQueries.FindByPredicate(
+		[Database](const FCachedPCAQuery& CachedPCAQuery)
+		{
+			return CachedPCAQuery.GetDatabase() == Database;
+		}))
+	{
+		return FoundCachedPCAQuery->GetValues();
+	}
+	
+	// building the query for the database schema
+	TConstArrayView<float> QueryValues = GetOrBuildQuery(Database->Schema);
+
+	// converting it into PCA space
+	FCachedPCAQuery& NewCachedPCAQuery = CachedPCAQueries.Emplace_GetRef(Database);
+	return Database->GetSearchIndex().PCAProject(QueryValues, NewCachedPCAQuery.EditValues());
+}
+
 bool FSearchContext::IsCurrentResultFromDatabase(const UPoseSearchDatabase* Database) const
 {
 	return IsContinuingPoseDatabase(Database);
@@ -1331,4 +1363,49 @@ bool FSearchContext::CanUseContinuingPoseValues() const
 	return bCanUseContinuingPoseValues;
 }
 
+#if UE_POSE_SEARCH_TRACE_ENABLED
+void FSearchContext::Track(const UPoseSearchDatabase* Database, int32 PoseIdx, EPoseCandidateFlags PoseCandidateFlags, const FPoseSearchCost& Cost)
+{
+	check(Database);
+	FBestPoseCandidates& BestPoseCandidates = BestPoseCandidatesMap.FindOrAdd(Database);
+	if (PoseIdx != INDEX_NONE)
+	{
+		BestPoseCandidates.Add(PoseIdx, PoseCandidateFlags, Cost);
+	}
+}
+
+void FSearchContext::Track(const UPoseSearchDatabase* Database, int32 PoseIdx, EPoseCandidateFlags PoseCandidateFlags, TConstArrayView<float> DynamicWeightsSqrt, float ContinuingPoseCostAddend, float ContinuingInteractionCostAddend)
+{
+	check(PoseIdx >= 0 && Database && Database->Schema);
+
+	// if we're editing the database and removing assets it's possible that the PoseIndicesHistory contains invalid pose indexes
+	const FSearchIndex& SearchIndex = Database->GetSearchIndex();
+	if (PoseIdx < SearchIndex.GetNumPoses())
+	{
+		float DissimilarityCost = 0.f;
+		if (SearchIndex.IsValuesEmpty())
+		{
+			const int32 PCAValueIdx = SearchIndex.GetPCAValueIndex(PoseIdx);
+			if (PCAValueIdx != INDEX_NONE)
+			{
+				TConstArrayView<float> PCAQueryValues = GetOrBuildPCAQuery(Database);
+				// estimating the cost with the PCAVAlues rather than the Values
+				DissimilarityCost = CompareFeatureVectors(SearchIndex.GetPCAPoseValues(PCAValueIdx), PCAQueryValues);
+			}
+			else
+			{
+				DissimilarityCost = 100.f;
+			}
+		}
+		else
+		{
+			TConstArrayView<float> QueryValues = GetOrBuildQuery(Database->Schema);
+			DissimilarityCost = CompareFeatureVectors<false>(SearchIndex.GetPoseValues(PoseIdx), QueryValues, DynamicWeightsSqrt);
+		}
+		const FPoseSearchCost PoseCost(DissimilarityCost, SearchIndex.PoseMetadata[PoseIdx].GetCostAddend(), ContinuingPoseCostAddend, ContinuingInteractionCostAddend);
+		Track(Database, PoseIdx, PoseCandidateFlags, PoseCost);
+	}
+}
+#endif // UE_POSE_SEARCH_TRACE_ENABLED
+
 } // namespace UE::PoseSearch
\ No newline at end of file
diff --git a/Engine/Plugins/Animation/PoseSearch/Source/Runtime/Private/PoseSearchDatabase.cpp b/Engine/Plugins/Animation/PoseSearch/Source/Runtime/Private/PoseSearchDatabase.cpp
--- a/Engine/Plugins/Animation/PoseSearch/Source/Runtime/Private/PoseSearchDatabase.cpp
+++ b/Engine/Plugins/Animation/PoseSearch/Source/Runtime/Private/PoseSearchDatabase.cpp
@@ -193,18 +193,15 @@ struct FSearchFilters
 				}
 				else if (Filter == &SelectableAssetIdxFilter)
 				{
-					const FPoseSearchCost PoseCost(CompareFeatureVectors<false>(PoseValues, QueryValues, DynamicWeightsSqrt), SearchIndex.PoseMetadata[PoseIdx].GetCostAddend(), ContinuingPoseCostAddend, ContinuingInteractionCostAddend);
-					SearchContext.Track(Database, PoseIdx, EPoseCandidateFlags::DiscardedBy_AssetIdxFilter, PoseCost);
+					SearchContext.Track(Database, PoseIdx, EPoseCandidateFlags::DiscardedBy_AssetIdxFilter, DynamicWeightsSqrt, ContinuingPoseCostAddend, ContinuingInteractionCostAddend);
 				}
 				else if (Filter == &BlockTransitionFilter)
 				{
-					const FPoseSearchCost PoseCost(CompareFeatureVectors<false>(PoseValues, QueryValues, DynamicWeightsSqrt), SearchIndex.PoseMetadata[PoseIdx].GetCostAddend(), ContinuingPoseCostAddend, ContinuingInteractionCostAddend);
-					SearchContext.Track(Database, PoseIdx, EPoseCandidateFlags::DiscardedBy_BlockTransition, PoseCost);
+					SearchContext.Track(Database, PoseIdx, EPoseCandidateFlags::DiscardedBy_BlockTransition, DynamicWeightsSqrt, ContinuingPoseCostAddend, ContinuingInteractionCostAddend);
 				}
 				else
 				{
-					const FPoseSearchCost PoseCost(CompareFeatureVectors<false>(PoseValues, QueryValues, DynamicWeightsSqrt), SearchIndex.PoseMetadata[PoseIdx].GetCostAddend(), ContinuingPoseCostAddend, ContinuingInteractionCostAddend);
-					SearchContext.Track(Database, PoseIdx, EPoseCandidateFlags::DiscardedBy_PoseFilter, PoseCost);
+					SearchContext.Track(Database, PoseIdx, EPoseCandidateFlags::DiscardedBy_PoseFilter, DynamicWeightsSqrt, ContinuingPoseCostAddend, ContinuingInteractionCostAddend);
 				}
 #endif // UE_POSE_SEARCH_TRACE_ENABLED
 				return false;
@@ -1617,7 +1614,7 @@ void UPoseSearchDatabase::Search(UE::PoseSearch::FSearchContext& SearchContext,
 
 void UPoseSearchDatabase::PopulateNonSelectableIdx(FNonSelectableIdx& NonSelectableIdx, UE::PoseSearch::FSearchContext& SearchContext
 #if UE_POSE_SEARCH_TRACE_ENABLED
-	, float ContinuingPoseCostAddend, float ContinuingInteractionCostAddend, TConstArrayView<float> QueryValues, TConstArrayView<float> DynamicWeightsSqrt
+	, float ContinuingPoseCostAddend, float ContinuingInteractionCostAddend, TConstArrayView<float> DynamicWeightsSqrt
 #endif //UE_POSE_SEARCH_TRACE_ENABLED
 ) const
 {
@@ -1625,10 +1622,6 @@ void UPoseSearchDatabase::PopulateNonSelectableIdx(FNonSelectableIdx& NonSelecta
 
 	const FSearchIndex& SearchIndex = GetSearchIndex();
 
-#if UE_POSE_SEARCH_TRACE_ENABLED
-	TArray<float> BufferUsedForReconstruction;
-#endif // UE_POSE_SEARCH_TRACE_ENABLED
-
 	NonSelectableIdx.Reset();
 	if (SearchContext.IsContinuingPoseDatabase(this))
 	{
@@ -1651,9 +1644,7 @@ void UPoseSearchDatabase::PopulateNonSelectableIdx(FNonSelectableIdx& NonSelecta
 						NonSelectableIdx.Add(PoseIdx);
 
 #if UE_POSE_SEARCH_TRACE_ENABLED
-						const TConstArrayView<float> PoseValues = SearchIndex.GetPoseValuesSafe(PoseIdx, BufferUsedForReconstruction);
-						const FPoseSearchCost PoseCost(CompareFeatureVectors<false>(PoseValues, QueryValues, DynamicWeightsSqrt), SearchIndex.PoseMetadata[PoseIdx].GetCostAddend(), ContinuingPoseCostAddend, ContinuingInteractionCostAddend);
-						SearchContext.Track(this, PoseIdx, EPoseCandidateFlags::DiscardedBy_AssetReselection, PoseCost);
+						SearchContext.Track(this, PoseIdx, EPoseCandidateFlags::DiscardedBy_AssetReselection, DynamicWeightsSqrt, ContinuingPoseCostAddend, ContinuingInteractionCostAddend);
 #endif // UE_POSE_SEARCH_TRACE_ENABLED
 					}
 				}
@@ -1678,9 +1669,7 @@ void UPoseSearchDatabase::PopulateNonSelectableIdx(FNonSelectableIdx& NonSelecta
 						NonSelectableIdx.AddUnique(PoseIdx);
 
 #if UE_POSE_SEARCH_TRACE_ENABLED
-						const TConstArrayView<float> PoseValues = SearchIndex.GetPoseValuesSafe(PoseIdx, BufferUsedForReconstruction);
-						const FPoseSearchCost PoseCost(CompareFeatureVectors<false>(PoseValues, QueryValues, DynamicWeightsSqrt), SearchIndex.PoseMetadata[PoseIdx].GetCostAddend(), ContinuingPoseCostAddend, ContinuingInteractionCostAddend);
-						SearchContext.Track(this, PoseIdx, EPoseCandidateFlags::DiscardedBy_PoseJumpThresholdTime, PoseCost);
+						SearchContext.Track(this, PoseIdx, EPoseCandidateFlags::DiscardedBy_PoseJumpThresholdTime, DynamicWeightsSqrt, ContinuingPoseCostAddend, ContinuingInteractionCostAddend);
 #endif // UE_POSE_SEARCH_TRACE_ENABLED
 					}
 				}
@@ -1692,11 +1681,8 @@ void UPoseSearchDatabase::PopulateNonSelectableIdx(FNonSelectableIdx& NonSelecta
 					for (int32 PoseIdx = MinPoseIdx; PoseIdx < MaxPoseIdx; ++PoseIdx)
 					{
 						NonSelectableIdx.AddUnique(PoseIdx);
-
 #if UE_POSE_SEARCH_TRACE_ENABLED
-						const TConstArrayView<float> PoseValues = SearchIndex.GetPoseValuesSafe(PoseIdx, BufferUsedForReconstruction);
-						const FPoseSearchCost PoseCost(CompareFeatureVectors<false>(PoseValues, QueryValues, DynamicWeightsSqrt), SearchIndex.PoseMetadata[PoseIdx].GetCostAddend(), ContinuingPoseCostAddend, ContinuingInteractionCostAddend);
-						SearchContext.Track(this, PoseIdx, EPoseCandidateFlags::DiscardedBy_PoseJumpThresholdTime, PoseCost);
+						SearchContext.Track(this, PoseIdx, EPoseCandidateFlags::DiscardedBy_PoseJumpThresholdTime, DynamicWeightsSqrt, ContinuingPoseCostAddend, ContinuingInteractionCostAddend);
 #endif // UE_POSE_SEARCH_TRACE_ENABLED
 					}
 				}
@@ -1719,14 +1705,7 @@ void UPoseSearchDatabase::PopulateNonSelectableIdx(FNonSelectableIdx& NonSelecta
 					NonSelectableIdx.AddUnique(HistoricalPoseIndex.PoseIndex);
 
 #if UE_POSE_SEARCH_TRACE_ENABLED
-					check(HistoricalPoseIndex.PoseIndex >= 0);
-
-					// if we're editing the database and removing assets it's possible that the PoseIndicesHistory contains invalid pose indexes
-					if (HistoricalPoseIndex.PoseIndex < SearchIndex.GetNumPoses())
-					{
-						const FPoseSearchCost PoseCost(CompareFeatureVectors<false>(SearchIndex.GetPoseValuesSafe(HistoricalPoseIndex.PoseIndex, BufferUsedForReconstruction), QueryValues, DynamicWeightsSqrt), SearchIndex.PoseMetadata[HistoricalPoseIndex.PoseIndex].GetCostAddend(), ContinuingPoseCostAddend, ContinuingInteractionCostAddend);
-						SearchContext.Track(this, HistoricalPoseIndex.PoseIndex, EPoseCandidateFlags::DiscardedBy_PoseReselectHistory, PoseCost);
-					}
+					SearchContext.Track(this, HistoricalPoseIndex.PoseIndex, EPoseCandidateFlags::DiscardedBy_PoseReselectHistory, DynamicWeightsSqrt, ContinuingPoseCostAddend, ContinuingInteractionCostAddend);
 #endif // UE_POSE_SEARCH_TRACE_ENABLED
 				}
 			}
@@ -1919,9 +1898,6 @@ void UPoseSearchDatabase::SearchPCAKDTree(UE::PoseSearch::FSearchContext& Search
 		const uint32 ClampedKDTreeQueryNumNeighbors = FMath::Clamp<uint32>(KDTreeQueryNumNeighbors, 1, SearchIndex.GetNumPoses());
 		const bool bArePCAValuesPruned = SearchIndex.PCAValuesVectorToPoseIndexes.Num() > 0;
 
-		//stack allocated temporaries
-		TArrayView<float> ProjectedQueryValues((float*)FMemory_Alloca(ClampedNumberOfPrincipalComponents * sizeof(float)), ClampedNumberOfPrincipalComponents);
-	
 		// @todo: perhaps store the DynamicWeightsSqrt into the GetOrBuildQuery to share between the continuing pose and the full search
 		TArrayView<float> DynamicWeightsSqrtBuffer((float*)FMemory_Alloca(NumDimensions * sizeof(float)), NumDimensions);
 		const TConstArrayView<float> DynamicWeightsSqrt = CalculateDynamicWeightsSqrt(DynamicWeightsSqrtBuffer);
@@ -1935,7 +1911,7 @@ void UPoseSearchDatabase::SearchPCAKDTree(UE::PoseSearch::FSearchContext& Search
 		FNonSelectableIdx NonSelectableIdx;
 		PopulateNonSelectableIdx(NonSelectableIdx, SearchContext
 #if UE_POSE_SEARCH_TRACE_ENABLED
-			, ContinuingPoseCostAddend, ContinuingInteractionCostAddend, QueryValues, DynamicWeightsSqrt
+			, ContinuingPoseCostAddend, ContinuingInteractionCostAddend, DynamicWeightsSqrt
 #endif // UE_POSE_SEARCH_TRACE_ENABLED
 		);
 
@@ -1945,7 +1921,7 @@ void UPoseSearchDatabase::SearchPCAKDTree(UE::PoseSearch::FSearchContext& Search
 #endif // WITH_EDITOR && ENABLE_ANIM_DEBUG
 
 		// projecting QueryValues into the PCA space 
-		TConstArrayView<float> PCAQueryValues = SearchIndex.PCAProject(QueryValues, ProjectedQueryValues);
+		TConstArrayView<float> PCAQueryValues = SearchContext.GetOrBuildPCAQuery(this);
 		check(PCAQueryValues.Num() == ClampedNumberOfPrincipalComponents);
 
 		int32 NumResults = 0;
@@ -2070,46 +2046,92 @@ void UPoseSearchDatabase::SearchPCAKDTree(UE::PoseSearch::FSearchContext& Search
 		// are the PCAValues pruned out of duplicates (multiple poses are associated with the same PCAValuesVectorIdx)
 		if (bArePCAValuesPruned)
 		{
-			// @todo: reconstruction is not yet supported with pruned PCAValues
-			check(!SearchIndex.IsValuesEmpty());
-			
 			const int32 MaxNumEvaluatePoseKernelCalls = KDTreeQueryNumNeighborsWithDuplicates > 0 ? KDTreeQueryNumNeighborsWithDuplicates : INT32_MAX;
 
-			if (NumDimensions % 4 == 0)
+			// do we need to reconstruct pose values?
+			if (SearchIndex.IsValuesEmpty())
 			{
-				int32 NumEvaluatePoseKernelCalls = 0;
-				for (int32 ResultIndex = 0; ResultIndex < NumResults; ++ResultIndex)
+				// FMemory_Alloca is forced 16 bytes aligned
+				TArrayView<float> ReconstructedPoseValuesBuffer((float*)FMemory_Alloca(NumDimensions * sizeof(float)), NumDimensions);
+				check(IsAligned(ReconstructedPoseValuesBuffer.GetData(), alignof(VectorRegister4Float)));
+				if (NumDimensions % 4 == 0)
 				{
-					const TConstArrayView<int32> PoseIndexes = SearchIndex.PCAValuesVectorToPoseIndexes[Results[ResultIndex].Index];
-					for (int32 Index = 0; Index < PoseIndexes.Num() && NumEvaluatePoseKernelCalls < MaxNumEvaluatePoseKernelCalls; ++Index, ++NumEvaluatePoseKernelCalls)
+					int32 NumEvaluatePoseKernelCalls = 0;
+					for (int32 ResultIndex = 0; ResultIndex < NumResults; ++ResultIndex)
 					{
-						EvaluatePoseKernel<false, true>(SearchResults, SearchIndex, QueryValues, TArrayView<float>(), PoseIndexes[Index], INDEX_NONE, SearchFilters, ContinuingPoseCostAddend, ContinuingInteractionCostAddend, SearchContext, this, DynamicWeightsSqrt
+						const TConstArrayView<int32> PoseIndexes = SearchIndex.PCAValuesVectorToPoseIndexes[Results[ResultIndex].Index];
+						for (int32 Index = 0; Index < PoseIndexes.Num() && NumEvaluatePoseKernelCalls < MaxNumEvaluatePoseKernelCalls; ++Index, ++NumEvaluatePoseKernelCalls)
+						{
+							EvaluatePoseKernel<true, true>(SearchResults, SearchIndex, QueryValues, ReconstructedPoseValuesBuffer, PoseIndexes[Index], INDEX_NONE, SearchFilters, ContinuingPoseCostAddend, ContinuingInteractionCostAddend, SearchContext, this, DynamicWeightsSqrt
 #if UE_POSE_SEARCH_TRACE_ENABLED
-							, true
+								, true
 #endif // UE_POSE_SEARCH_TRACE_ENABLED
 #if WITH_EDITOR && ENABLE_ANIM_DEBUG && UE_POSE_SEARCH_TRACE_ENABLED
-							, BruteForceSearchResult, SortedResultsIndexes[ResultIndex]
+								, BruteForceSearchResult, SortedResultsIndexes[ResultIndex]
 #endif // WITH_EDITOR && ENABLE_ANIM_DEBUG && UE_POSE_SEARCH_TRACE_ENABLED
-						);
+							);
+						}
+					}
+				}
+				else
+				{
+					int32 NumEvaluatePoseKernelCalls = 0;
+					for (int32 ResultIndex = 0; ResultIndex < NumResults; ++ResultIndex)
+					{
+						const TConstArrayView<int32> PoseIndexes = SearchIndex.PCAValuesVectorToPoseIndexes[Results[ResultIndex].Index];
+						for (int32 Index = 0; Index < PoseIndexes.Num() && NumEvaluatePoseKernelCalls < MaxNumEvaluatePoseKernelCalls; ++Index, ++NumEvaluatePoseKernelCalls)
+						{
+							EvaluatePoseKernel<true, false>(SearchResults, SearchIndex, QueryValues, ReconstructedPoseValuesBuffer, PoseIndexes[Index], INDEX_NONE, SearchFilters, ContinuingPoseCostAddend, ContinuingInteractionCostAddend, SearchContext, this, DynamicWeightsSqrt
+#if UE_POSE_SEARCH_TRACE_ENABLED
+								, true
+#endif // UE_POSE_SEARCH_TRACE_ENABLED
+#if WITH_EDITOR && ENABLE_ANIM_DEBUG && UE_POSE_SEARCH_TRACE_ENABLED
+								, BruteForceSearchResult, SortedResultsIndexes[ResultIndex]
+#endif // WITH_EDITOR && ENABLE_ANIM_DEBUG && UE_POSE_SEARCH_TRACE_ENABLED
+							);
+						}
 					}
 				}
 			}
 			else
 			{
-				int32 NumEvaluatePoseKernelCalls = 0;
-				for (int32 ResultIndex = 0; ResultIndex < NumResults; ++ResultIndex)
+				// no need to reconstruct pose values
+				if (NumDimensions % 4 == 0)
 				{
-					const TConstArrayView<int32> PoseIndexes = SearchIndex.PCAValuesVectorToPoseIndexes[Results[ResultIndex].Index];
-					for (int32 Index = 0; Index < PoseIndexes.Num() && NumEvaluatePoseKernelCalls < MaxNumEvaluatePoseKernelCalls; ++Index, ++NumEvaluatePoseKernelCalls)
+					int32 NumEvaluatePoseKernelCalls = 0;
+					for (int32 ResultIndex = 0; ResultIndex < NumResults; ++ResultIndex)
 					{
-						EvaluatePoseKernel<false, false>(SearchResults, SearchIndex, QueryValues, TArrayView<float>(), PoseIndexes[Index], INDEX_NONE, SearchFilters, ContinuingPoseCostAddend, ContinuingInteractionCostAddend, SearchContext, this, DynamicWeightsSqrt
+						const TConstArrayView<int32> PoseIndexes = SearchIndex.PCAValuesVectorToPoseIndexes[Results[ResultIndex].Index];
+						for (int32 Index = 0; Index < PoseIndexes.Num() && NumEvaluatePoseKernelCalls < MaxNumEvaluatePoseKernelCalls; ++Index, ++NumEvaluatePoseKernelCalls)
+						{
+							EvaluatePoseKernel<false, true>(SearchResults, SearchIndex, QueryValues, TArrayView<float>(), PoseIndexes[Index], INDEX_NONE, SearchFilters, ContinuingPoseCostAddend, ContinuingInteractionCostAddend, SearchContext, this, DynamicWeightsSqrt
 #if UE_POSE_SEARCH_TRACE_ENABLED
-							, true
+								, true
 #endif // UE_POSE_SEARCH_TRACE_ENABLED
 #if WITH_EDITOR && ENABLE_ANIM_DEBUG && UE_POSE_SEARCH_TRACE_ENABLED
-						, BruteForceSearchResult, SortedResultsIndexes[ResultIndex]
+								, BruteForceSearchResult, SortedResultsIndexes[ResultIndex]
 #endif // WITH_EDITOR && ENABLE_ANIM_DEBUG && UE_POSE_SEARCH_TRACE_ENABLED
-						);
+							);
+						}
+					}
+				}
+				else
+				{
+					int32 NumEvaluatePoseKernelCalls = 0;
+					for (int32 ResultIndex = 0; ResultIndex < NumResults; ++ResultIndex)
+					{
+						const TConstArrayView<int32> PoseIndexes = SearchIndex.PCAValuesVectorToPoseIndexes[Results[ResultIndex].Index];
+						for (int32 Index = 0; Index < PoseIndexes.Num() && NumEvaluatePoseKernelCalls < MaxNumEvaluatePoseKernelCalls; ++Index, ++NumEvaluatePoseKernelCalls)
+						{
+							EvaluatePoseKernel<false, false>(SearchResults, SearchIndex, QueryValues, TArrayView<float>(), PoseIndexes[Index], INDEX_NONE, SearchFilters, ContinuingPoseCostAddend, ContinuingInteractionCostAddend, SearchContext, this, DynamicWeightsSqrt
+#if UE_POSE_SEARCH_TRACE_ENABLED
+								, true
+#endif // UE_POSE_SEARCH_TRACE_ENABLED
+#if WITH_EDITOR && ENABLE_ANIM_DEBUG && UE_POSE_SEARCH_TRACE_ENABLED
+								, BruteForceSearchResult, SortedResultsIndexes[ResultIndex]
+#endif // WITH_EDITOR && ENABLE_ANIM_DEBUG && UE_POSE_SEARCH_TRACE_ENABLED
+							);
+						}
 					}
 				}
 			}
@@ -2120,46 +2142,68 @@ void UPoseSearchDatabase::SearchPCAKDTree(UE::PoseSearch::FSearchContext& Search
 			// FMemory_Alloca is forced 16 bytes aligned
 			TArrayView<float> ReconstructedPoseValuesBuffer((float*)FMemory_Alloca(NumDimensions * sizeof(float)), NumDimensions);
 			check(IsAligned(ReconstructedPoseValuesBuffer.GetData(), alignof(VectorRegister4Float)));
-			for (int32 ResultIndex = 0; ResultIndex < NumResults; ++ResultIndex)
+
+			// is the data padded at 16 bytes (and 16 bytes aligned by construction)?
+			if (NumDimensions % 4 == 0)
+			{
+				for (int32 ResultIndex = 0; ResultIndex < NumResults; ++ResultIndex)
+				{
+					EvaluatePoseKernel<true, true>(SearchResults, SearchIndex, QueryValues, ReconstructedPoseValuesBuffer, Results[ResultIndex].Index, INDEX_NONE, SearchFilters, ContinuingPoseCostAddend, ContinuingInteractionCostAddend, SearchContext, this, DynamicWeightsSqrt
+	#if UE_POSE_SEARCH_TRACE_ENABLED
+						, true
+	#endif // UE_POSE_SEARCH_TRACE_ENABLED
+	#if WITH_EDITOR && ENABLE_ANIM_DEBUG && UE_POSE_SEARCH_TRACE_ENABLED
+						, BruteForceSearchResult, SortedResultsIndexes[ResultIndex]
+	#endif // WITH_EDITOR && ENABLE_ANIM_DEBUG && UE_POSE_SEARCH_TRACE_ENABLED
+					);
+				}
+			}
+			else
 			{
-				EvaluatePoseKernel<true, false>(SearchResults, SearchIndex, QueryValues, ReconstructedPoseValuesBuffer, Results[ResultIndex].Index, INDEX_NONE, SearchFilters, ContinuingPoseCostAddend, ContinuingInteractionCostAddend, SearchContext, this, DynamicWeightsSqrt
+				for (int32 ResultIndex = 0; ResultIndex < NumResults; ++ResultIndex)
+				{
+					EvaluatePoseKernel<true, false>(SearchResults, SearchIndex, QueryValues, ReconstructedPoseValuesBuffer, Results[ResultIndex].Index, INDEX_NONE, SearchFilters, ContinuingPoseCostAddend, ContinuingInteractionCostAddend, SearchContext, this, DynamicWeightsSqrt
 #if UE_POSE_SEARCH_TRACE_ENABLED
-					, true
+						, true
 #endif // UE_POSE_SEARCH_TRACE_ENABLED
 #if WITH_EDITOR && ENABLE_ANIM_DEBUG && UE_POSE_SEARCH_TRACE_ENABLED
-					, BruteForceSearchResult, SortedResultsIndexes[ResultIndex]
+						, BruteForceSearchResult, SortedResultsIndexes[ResultIndex]
 #endif // WITH_EDITOR && ENABLE_ANIM_DEBUG && UE_POSE_SEARCH_TRACE_ENABLED
-				);
+					);
+				}
 			}
 		}
-		// is the data padded at 16 bytes (and 16 bytes aligned by construction)?
-		else if (NumDimensions % 4 == 0)
+		else
 		{
-			for (int32 ResultIndex = 0; ResultIndex < NumResults; ++ResultIndex)
+			// is the data padded at 16 bytes (and 16 bytes aligned by construction)?
+			if (NumDimensions % 4 == 0)
 			{
-				EvaluatePoseKernel<false, true>(SearchResults, SearchIndex, QueryValues, TArrayView<float>(), Results[ResultIndex].Index, INDEX_NONE, SearchFilters, ContinuingPoseCostAddend, ContinuingInteractionCostAddend, SearchContext, this, DynamicWeightsSqrt
+				for (int32 ResultIndex = 0; ResultIndex < NumResults; ++ResultIndex)
+				{
+					EvaluatePoseKernel<false, true>(SearchResults, SearchIndex, QueryValues, TArrayView<float>(), Results[ResultIndex].Index, INDEX_NONE, SearchFilters, ContinuingPoseCostAddend, ContinuingInteractionCostAddend, SearchContext, this, DynamicWeightsSqrt
 #if UE_POSE_SEARCH_TRACE_ENABLED
-					, true
+						, true
 #endif // UE_POSE_SEARCH_TRACE_ENABLED
 #if WITH_EDITOR && ENABLE_ANIM_DEBUG && UE_POSE_SEARCH_TRACE_ENABLED
-					, BruteForceSearchResult, SortedResultsIndexes[ResultIndex]
+						, BruteForceSearchResult, SortedResultsIndexes[ResultIndex]
 #endif // WITH_EDITOR && ENABLE_ANIM_DEBUG && UE_POSE_SEARCH_TRACE_ENABLED
-				);
+					);
+				}
 			}
-		}
-		// no reconstruction, but data is not 16 bytes padded
-		else
-		{
-			for (int32 ResultIndex = 0; ResultIndex < NumResults; ++ResultIndex)
+			// no reconstruction, but data is not 16 bytes padded
+			else
 			{
-				EvaluatePoseKernel<false, false>(SearchResults, SearchIndex, QueryValues, TArrayView<float>(), Results[ResultIndex].Index, INDEX_NONE, SearchFilters, ContinuingPoseCostAddend, ContinuingInteractionCostAddend, SearchContext, this, DynamicWeightsSqrt
+				for (int32 ResultIndex = 0; ResultIndex < NumResults; ++ResultIndex)
+				{
+					EvaluatePoseKernel<false, false>(SearchResults, SearchIndex, QueryValues, TArrayView<float>(), Results[ResultIndex].Index, INDEX_NONE, SearchFilters, ContinuingPoseCostAddend, ContinuingInteractionCostAddend, SearchContext, this, DynamicWeightsSqrt
 #if UE_POSE_SEARCH_TRACE_ENABLED
-					, true
+						, true
 #endif // UE_POSE_SEARCH_TRACE_ENABLED
 #if WITH_EDITOR && ENABLE_ANIM_DEBUG && UE_POSE_SEARCH_TRACE_ENABLED
-					, BruteForceSearchResult, SortedResultsIndexes[ResultIndex]
+						, BruteForceSearchResult, SortedResultsIndexes[ResultIndex]
 #endif // WITH_EDITOR && ENABLE_ANIM_DEBUG && UE_POSE_SEARCH_TRACE_ENABLED
-				);
+					);
+				}
 			}
 		}
 	}
@@ -2171,9 +2215,8 @@ void UPoseSearchDatabase::SearchPCAKDTree(UE::PoseSearch::FSearchContext& Search
 		const TConstArrayView<float> DynamicWeightsSqrt = CalculateDynamicWeightsSqrt(DynamicWeightsSqrtBuffer);
 
 		// calling just for reporting non selectable poses
-		TConstArrayView<float> QueryValues = SearchContext.GetOrBuildQuery(Schema);
 		FNonSelectableIdx NonSelectableIdx;
-		PopulateNonSelectableIdx(NonSelectableIdx, SearchContext, ContinuingPoseCostAddend, ContinuingInteractionCostAddend, QueryValues, DynamicWeightsSqrt);
+		PopulateNonSelectableIdx(NonSelectableIdx, SearchContext, ContinuingPoseCostAddend, ContinuingInteractionCostAddend, DynamicWeightsSqrt);
 #endif // UE_POSE_SEARCH_TRACE_ENABLED
 	}
 
@@ -2224,7 +2267,7 @@ void UPoseSearchDatabase::SearchVPTree(UE::PoseSearch::FSearchContext& SearchCon
 		FNonSelectableIdx NonSelectableIdx;
 		PopulateNonSelectableIdx(NonSelectableIdx, SearchContext
 #if UE_POSE_SEARCH_TRACE_ENABLED
-			, ContinuingPoseCostAddend, ContinuingInteractionCostAddend, QueryValues, DynamicWeightsSqrt
+			, ContinuingPoseCostAddend, ContinuingInteractionCostAddend, DynamicWeightsSqrt
 #endif // UE_POSE_SEARCH_TRACE_ENABLED
 		);
 		
@@ -2288,9 +2331,8 @@ void UPoseSearchDatabase::SearchVPTree(UE::PoseSearch::FSearchContext& SearchCon
 		const TConstArrayView<float> DynamicWeightsSqrt = CalculateDynamicWeightsSqrt(DynamicWeightsSqrtBuffer);
 
 		// calling just for reporting non selectable poses
-		TConstArrayView<float> QueryValues = SearchContext.GetOrBuildQuery(Schema);
 		FNonSelectableIdx NonSelectableIdx;
-		PopulateNonSelectableIdx(NonSelectableIdx, SearchContext, ContinuingPoseCostAddend, ContinuingInteractionCostAddend, QueryValues, DynamicWeightsSqrt);
+		PopulateNonSelectableIdx(NonSelectableIdx, SearchContext, ContinuingPoseCostAddend, ContinuingInteractionCostAddend, DynamicWeightsSqrt);
 #endif // UE_POSE_SEARCH_TRACE_ENABLED
 	}
 
@@ -2330,7 +2372,7 @@ void UPoseSearchDatabase::SearchBruteForce(UE::PoseSearch::FSearchContext& Searc
 		FNonSelectableIdx NonSelectableIdx;
 		PopulateNonSelectableIdx(NonSelectableIdx, SearchContext
 #if UE_POSE_SEARCH_TRACE_ENABLED
-			, ContinuingPoseCostAddend, ContinuingInteractionCostAddend, QueryValues, DynamicWeightsSqrt
+			, ContinuingPoseCostAddend, ContinuingInteractionCostAddend, DynamicWeightsSqrt
 #endif // UE_POSE_SEARCH_TRACE_ENABLED
 		);
 
@@ -2469,9 +2511,8 @@ void UPoseSearchDatabase::SearchBruteForce(UE::PoseSearch::FSearchContext& Searc
 		const TConstArrayView<float> DynamicWeightsSqrt = CalculateDynamicWeightsSqrt(DynamicWeightsSqrtBuffer);
 
 		// calling just for reporting non selectable poses
-		TConstArrayView<float> QueryValues = SearchContext.GetOrBuildQuery(Schema);
 		FNonSelectableIdx NonSelectableIdx;
-		PopulateNonSelectableIdx(NonSelectableIdx, SearchContext, ContinuingPoseCostAddend, ContinuingInteractionCostAddend, QueryValues, DynamicWeightsSqrt);
+		PopulateNonSelectableIdx(NonSelectableIdx, SearchContext, ContinuingPoseCostAddend, ContinuingInteractionCostAddend, DynamicWeightsSqrt);
 #endif // UE_POSE_SEARCH_TRACE_ENABLED
 	}
 
@@ -2520,7 +2561,7 @@ void UPoseSearchDatabase::SearchEvent(UE::PoseSearch::FSearchContext& SearchCont
 			{
 				PopulateNonSelectableIdx(NonSelectableIdx, SearchContext
 #if UE_POSE_SEARCH_TRACE_ENABLED
-					, ContinuingPoseCostAddend, ContinuingInteractionCostAddend, QueryValues, DynamicWeightsSqrt
+					, ContinuingPoseCostAddend, ContinuingInteractionCostAddend, DynamicWeightsSqrt
 #endif // UE_POSE_SEARCH_TRACE_ENABLED
 				);
 			}
diff --git a/Engine/Plugins/Animation/PoseSearch/Source/Runtime/Private/PoseSearchDerivedData.cpp b/Engine/Plugins/Animation/PoseSearch/Source/Runtime/Private/PoseSearchDerivedData.cpp
--- a/Engine/Plugins/Animation/PoseSearch/Source/Runtime/Private/PoseSearchDerivedData.cpp
+++ b/Engine/Plugins/Animation/PoseSearch/Source/Runtime/Private/PoseSearchDerivedData.cpp
@@ -2104,7 +2104,13 @@ void FPoseSearchDatabaseAsyncCacheTask::OnGetComplete(UE::DerivedData::FCacheGet
 					return;
 				}
 
-				SearchIndex.PrunePCAValuesFromBlockTransitionPoses(MainDatabase->GetNumberOfPrincipalComponents());
+				// removing SearchIndex.Values and relying on FSearchIndex::GetReconstructedPoseValues to reconstruct the Values data from the PCAValues
+				const bool bRemoveSearchIndexValues = MainDatabase->PoseSearchMode == EPoseSearchMode::PCAKDTree && MainDatabase->KDTreeQueryNumNeighbors <= 1;
+				if (!bRemoveSearchIndexValues)
+				{
+					// we cannot prune any PCAValues if we need to reconstruct the Values from them!
+					SearchIndex.PrunePCAValuesFromBlockTransitionPoses(MainDatabase->GetNumberOfPrincipalComponents());
+				}
 
 				PreprocessSearchIndexKDTree(SearchIndex, MainDatabase);
 				if (Owner.IsCanceled())
@@ -2114,8 +2120,7 @@ void FPoseSearchDatabaseAsyncCacheTask::OnGetComplete(UE::DerivedData::FCacheGet
 					return;
 				}
 
-				// removing SearchIndex.Values and relying on FSearchIndex::GetReconstructedPoseValues to reconstruct the Values data from the PCAValues
-				if (MainDatabase->PoseSearchMode == EPoseSearchMode::PCAKDTree && MainDatabase->KDTreeQueryNumNeighbors <= 1)
+				if (bRemoveSearchIndexValues)
 				{
 					SearchIndex.ResetValues();
 				}
diff --git a/Engine/Plugins/Animation/PoseSearch/Source/Runtime/Private/PoseSearchIndex.cpp b/Engine/Plugins/Animation/PoseSearch/Source/Runtime/Private/PoseSearchIndex.cpp
--- a/Engine/Plugins/Animation/PoseSearch/Source/Runtime/Private/PoseSearchIndex.cpp
+++ b/Engine/Plugins/Animation/PoseSearch/Source/Runtime/Private/PoseSearchIndex.cpp
@@ -670,33 +670,38 @@ TConstArrayView<float> FSearchIndex::GetReconstructedPoseValues(int32 PoseIdx, T
 {
 	QUICK_SCOPE_CYCLE_COUNTER(STAT_PoseSearch_PCAReconstruct);
 
-	// @todo: reconstruction is not yet supported with pruned PCAValues
-	check(PCAValuesVectorToPoseIndexes.Num() == 0);
+	const bool bArePCAValuesPruned = PCAValuesVectorToPoseIndexes.Num() > 0;
 
 	const int32 NumDimensions = GetNumDimensions();
 	const int32 NumPoses = GetNumPoses();
+	const int32 NumberOfPrincipalComponents = GetNumberOfPrincipalComponents();
+
 	check(PoseIdx >= 0 && PoseIdx < NumPoses && NumDimensions > 0);
 	check(BufferUsedForReconstruction.Num() == NumDimensions);
 
-	const int32 NumberOfPrincipalComponents = PCAValues.Num() / NumPoses;
-	
-	// NoTe: if one of these checks trigger, most likely PCAValuesPruningSimilarityThreshold > 0.f and we pruned some PCAValues.
-	// currently GetReconstructedPoseValues is not supported with PCAValues pruning
-	check(NumPoses * NumberOfPrincipalComponents == PCAValues.Num());
-	check(PCAProjectionMatrix.Num() == NumDimensions * NumberOfPrincipalComponents);
-
-	const RowMajorVectorMapConst MapWeightsSqrt(WeightsSqrt.GetData(), 1, NumDimensions);
 	const ColMajorMatrixMapConst MapPCAProjectionMatrix(PCAProjectionMatrix.GetData(), NumDimensions, NumberOfPrincipalComponents);
 	const RowMajorVectorMapConst MapMean(Mean.GetData(), 1, NumDimensions);
-	const RowMajorMatrixMapConst MapPCAValues(PCAValues.GetData(), NumPoses, NumberOfPrincipalComponents);
 
-	const RowMajorVector ReciprocalWeightsSqrt = MapWeightsSqrt.cwiseInverse();
-	const RowMajorVector WeightedReconstructedValues = MapPCAValues.row(PoseIdx) * MapPCAProjectionMatrix.transpose() + MapMean;
+	check(PCAValues.Num() % NumberOfPrincipalComponents == 0);
+	const int32 NumPCAValuesVectors = PCAValues.Num() / NumberOfPrincipalComponents;
+	const RowMajorMatrixMapConst MapPCAValues(PCAValues.GetData(), NumPCAValuesVectors, NumberOfPrincipalComponents);
+
+	const int32 PCAValueIdx = GetPCAValueIndex(PoseIdx);
+	if (PCAValueIdx != INDEX_NONE)
+	{
+		check(PCAValueIdx >= 0 && PCAValueIdx < NumPCAValuesVectors);
+		const RowMajorVector WeightedReconstructedValues = MapPCAValues.row(PCAValueIdx) * MapPCAProjectionMatrix.transpose() + MapMean;
 
-	RowMajorVectorMap ReconstructedPoseValues(BufferUsedForReconstruction.GetData(), 1, NumDimensions);
-	ReconstructedPoseValues = WeightedReconstructedValues.array() * ReciprocalWeightsSqrt.array();
+		const RowMajorVectorMapConst MapWeightsSqrt(WeightsSqrt.GetData(), 1, NumDimensions);
+		const RowMajorVector ReciprocalWeightsSqrt = MapWeightsSqrt.cwiseInverse();
 
-	return BufferUsedForReconstruction;
+		RowMajorVectorMap ReconstructedPoseValues(BufferUsedForReconstruction.GetData(), 1, NumDimensions);
+		ReconstructedPoseValues = WeightedReconstructedValues.array() * ReciprocalWeightsSqrt.array();
+
+		return BufferUsedForReconstruction;
+	}
+
+	return TConstArrayView<float>();
 }
 
 int32 FSearchIndex::GetNumDimensions() const
@@ -791,11 +796,7 @@ void FSearchIndex::PrunePCAValuesFromBlockTransitionPoses(int32 NumberOfPrincipa
 	}
 
 	PCAValues = PrunedPCAValues;
-	PCAValuesVectorToPoseIndexes = FSparsePoseMultiMap<int32>(PrunedPCAValuesVectorToPoseIndexes.Num(), GetNumPoses() - 1);
-	for (const TPair<int32, TArray<int32>>& Pair : PrunedPCAValuesVectorToPoseIndexes)
-	{
-		PCAValuesVectorToPoseIndexes.Insert(Pair.Key, Pair.Value);
-	}
+	PCAValuesVectorToPoseIndexes.Init(PrunedPCAValuesVectorToPoseIndexes, GetNumPoses() - 1);
 }
 
 void FSearchIndex::PruneDuplicatePCAValues(float SimilarityThreshold, int32 NumberOfPrincipalComponents)
@@ -946,9 +947,6 @@ void FSearchIndex::GetPoseToPCAValuesVectorIndexes(TArray<uint32>& PoseToPCAValu
 				PoseToPCAValuesVectorIndexes[PoseIdx] = PCAValuesVectorIdx;
 			}
 		}
-
-		// if block transition pruning is enabled we could have poses without their PCAValuesVectorIdx counterpart
-		// check(!PoseToPCAValuesVectorIndexes.Contains(INDEX_NONE));
 	}
 	else
 	{
@@ -956,6 +954,27 @@ void FSearchIndex::GetPoseToPCAValuesVectorIndexes(TArray<uint32>& PoseToPCAValu
 	}
 }
 
+int32 FSearchIndex::GetPCAValueIndex(int32 PoseIdx) const
+{
+	if (PCAValuesVectorToPoseIndexes.Num() > 0)
+	{
+		for (int32 PCAValuesVectorIdx = 0; PCAValuesVectorIdx < PCAValuesVectorToPoseIndexes.Num(); ++PCAValuesVectorIdx)
+		{
+			for (int32 PoseIdxForPCAValuesVectorIdx : PCAValuesVectorToPoseIndexes[PCAValuesVectorIdx])
+			{
+				if (PoseIdxForPCAValuesVectorIdx == PoseIdx)
+				{
+					return PCAValuesVectorIdx;
+				}
+			}
+		}
+
+		return INDEX_NONE;
+	}
+	
+	return PoseIdx;
+}
+
 bool FSearchIndex::operator==(const FSearchIndex& Other) const
 {
 	return FSearchIndexBase::operator==(Other)
diff --git a/Engine/Plugins/Animation/PoseSearch/Source/Runtime/Private/PoseSearchLibrary.cpp b/Engine/Plugins/Animation/PoseSearch/Source/Runtime/Private/PoseSearchLibrary.cpp
--- a/Engine/Plugins/Animation/PoseSearch/Source/Runtime/Private/PoseSearchLibrary.cpp
+++ b/Engine/Plugins/Animation/PoseSearch/Source/Runtime/Private/PoseSearchLibrary.cpp
@@ -460,10 +460,38 @@ namespace UE::PoseSearch
 				{
 					if (const UPoseSearchFeatureChannel_Trajectory* TrajectoryChannel = SearchResult.Database->Schema->FindFirstChannelOfType<UPoseSearchFeatureChannel_Trajectory>())
 					{
-						TConstArrayView<float> ResultData = SearchResult.Database->GetSearchIndex().GetPoseValues(SearchResult.PoseIdx);
-						const float EstimatedSpeedRatio = TrajectoryChannel->GetEstimatedSpeedRatio(QueryData, ResultData);
+						const FSearchIndex& SearchIndex = SearchResult.Database->GetSearchIndex();
+						const bool bReconstructPoseValues = SearchIndex.IsValuesEmpty();
 
-						WantedPlayRate = FMath::Clamp(EstimatedSpeedRatio, PlayRate.Min, PlayRate.Max);
+						if (bReconstructPoseValues)
+						{
+							const int32 NumDimensions = SearchIndex.GetNumDimensions();
+
+							// FMemory_Alloca is forced 16 bytes aligned and its allocated memory is in scope until the end of the function scope,
+							// not the statement scope, so it's safe to use it in GetEstimatedSpeedRatio
+							TArrayView<float> ReconstructedPoseValuesBuffer((float*)FMemory_Alloca(NumDimensions * sizeof(float)), NumDimensions);
+							check(IsAligned(ReconstructedPoseValuesBuffer.GetData(), alignof(VectorRegister4Float)));
+
+							const TConstArrayView<float> ResultData = SearchIndex.GetReconstructedPoseValues(SearchResult.PoseIdx, ReconstructedPoseValuesBuffer);
+							if (ResultData.IsEmpty())
+							{
+								UE_LOG(LogPoseSearch, Warning,
+									TEXT("Couldn't update the WantedPlayRate in CalculateWantedPlayRate, because couldn't reconstruct the pose value in GetReconstructedPoseValues for pose %d"),
+									SearchResult.PoseIdx);
+								WantedPlayRate = FMath::Clamp(1.f, PlayRate.Min, PlayRate.Max);
+							}
+							else
+							{
+								const float EstimatedSpeedRatio = TrajectoryChannel->GetEstimatedSpeedRatio(QueryData, ResultData);
+								WantedPlayRate = FMath::Clamp(EstimatedSpeedRatio, PlayRate.Min, PlayRate.Max);
+							}
+						}
+						else
+						{
+							const TConstArrayView<float> ResultData = SearchIndex.GetPoseValues(SearchResult.PoseIdx);
+							const float EstimatedSpeedRatio = TrajectoryChannel->GetEstimatedSpeedRatio(QueryData, ResultData);
+							WantedPlayRate = FMath::Clamp(EstimatedSpeedRatio, PlayRate.Min, PlayRate.Max);
+						}
 					}
 					else
 					{
diff --git a/Engine/Plugins/Animation/PoseSearch/Source/Runtime/Public/PoseSearch/PoseSearchContext.h b/Engine/Plugins/Animation/PoseSearch/Source/Runtime/Public/PoseSearch/PoseSearchContext.h
--- a/Engine/Plugins/Animation/PoseSearch/Source/Runtime/Public/PoseSearch/PoseSearchContext.h
+++ b/Engine/Plugins/Animation/PoseSearch/Source/Runtime/Public/PoseSearch/PoseSearchContext.h
@@ -142,6 +142,22 @@ private:
 	const UPoseSearchSchema* Schema;
 };
 
+// Experimental, this feature might be removed without warning, not for production use
+struct FCachedPCAQuery
+{
+public:
+	explicit FCachedPCAQuery(const UPoseSearchDatabase* InDatabase);
+	const UPoseSearchDatabase* GetDatabase() const { return Database; }
+	TArrayView<float> EditValues() { return Values; }
+	TConstArrayView<float> GetValues() const { return Values; }
+
+private:
+	TStackAlignedArray<float> Values;
+	
+	// no need for a TWeakObjectPtr since it doesn't persist across multiple frames (same lifespan as FSearchContext)
+	const UPoseSearchDatabase* Database;
+};
+
 // CachedChannels uses hashed unique identifiers to determine channels that can share feature vector data during the building of the query
 struct FCachedChannel
 {
@@ -201,8 +217,10 @@ struct FSearchContext
 	UE_DEPRECATED(5.7, "No longer necessary API. Use the FSearchResult::PoseCost from UPoseSearchDatabase::Search if a search cost is needed")
 	float GetCurrentBestTotalCost() const { return MAX_FLT; }
 
-	UE_API TConstArrayView<float> GetOrBuildQuery(const UPoseSearchSchema* Schema);
 	UE_API TConstArrayView<float> GetCachedQuery(const UPoseSearchSchema* Schema) const;
+	UE_API TConstArrayView<float> GetOrBuildQuery(const UPoseSearchSchema* Schema);
+	// Experimental, this feature might be removed without warning, not for production use
+	UE_API TConstArrayView<float> GetOrBuildPCAQuery(const UPoseSearchDatabase* Database);
 
 	UE_DEPRECATED(5.7, "compare Database against GetContinuingPoseSearchResult().Database instead (the internal replacement IsContinuingPoseDatabase is not going to be dll exported)")
 	UE_API bool IsCurrentResultFromDatabase(const UPoseSearchDatabase* Database) const;
@@ -340,6 +358,7 @@ private:
 	TMap<uint32, FTransform, TInlineSetAllocator<64, TMemStackSetAllocator<>>> CachedTransforms;
 
 	TArray<FCachedQuery, TInlineAllocator<PreallocatedCachedQueriesNum, TMemStackAllocator<>>> CachedQueries;
+	TArray<FCachedPCAQuery, TInlineAllocator<PreallocatedCachedQueriesNum, TMemStackAllocator<>>> CachedPCAQueries;
 
 	// mapping channel unique identifier (hash) to FCachedChannel
 	TMap<uint32, FCachedChannel, TInlineSetAllocator<PreallocatedCachedChannelDataNum, TMemStackSetAllocator<>>> CachedChannels;
@@ -406,15 +425,8 @@ public:
 		TSet<FPoseCandidate, FPoseCandidatePoseIdxKeyFunc> PermanentPoseCandidates;
 	};
 	
-	void Track(const UPoseSearchDatabase* Database, int32 PoseIdx = INDEX_NONE, EPoseCandidateFlags PoseCandidateFlags = EPoseCandidateFlags::None, const FPoseSearchCost& Cost = FPoseSearchCost())
-	{
-		check(Database);
-		FBestPoseCandidates& BestPoseCandidates = BestPoseCandidatesMap.FindOrAdd(Database);
-		if (PoseIdx != INDEX_NONE)
-		{
-			BestPoseCandidates.Add(PoseIdx, PoseCandidateFlags, Cost);
-		}
-	}
+	UE_API void Track(const UPoseSearchDatabase* Database, int32 PoseIdx = INDEX_NONE, EPoseCandidateFlags PoseCandidateFlags = EPoseCandidateFlags::None, const FPoseSearchCost& Cost = FPoseSearchCost());
+	UE_API void Track(const UPoseSearchDatabase* Database, int32 PoseIdx, EPoseCandidateFlags PoseCandidateFlags, TConstArrayView<float> DynamicWeightsSqrt, float ContinuingPoseCostAddend, float ContinuingInteractionCostAddend);
 
 	const TMap<const UPoseSearchDatabase*, FBestPoseCandidates>& GetBestPoseCandidatesMap() const
 	{
diff --git a/Engine/Plugins/Animation/PoseSearch/Source/Runtime/Public/PoseSearch/PoseSearchDatabase.h b/Engine/Plugins/Animation/PoseSearch/Source/Runtime/Public/PoseSearch/PoseSearchDatabase.h
--- a/Engine/Plugins/Animation/PoseSearch/Source/Runtime/Public/PoseSearch/PoseSearchDatabase.h
+++ b/Engine/Plugins/Animation/PoseSearch/Source/Runtime/Public/PoseSearch/PoseSearchDatabase.h
@@ -738,22 +738,22 @@ public:
 #endif // WITH_EDITORONLY_DATA
 
 private:
-	UE_API void SearchPCAKDTree(UE::PoseSearch::FSearchContext& SearchContext, UE::PoseSearch::FSearchResults& SearchResults) const;
-	UE_API void SearchVPTree(UE::PoseSearch::FSearchContext& SearchContext, UE::PoseSearch::FSearchResults& SearchResults) const;
-	UE_API void SearchBruteForce(UE::PoseSearch::FSearchContext& SearchContext, UE::PoseSearch::FSearchResults& SearchResults) const;
-	UE_API void SearchEvent(UE::PoseSearch::FSearchContext& SearchContext, UE::PoseSearch::FSearchResults& SearchResults) const;
+	void SearchPCAKDTree(UE::PoseSearch::FSearchContext& SearchContext, UE::PoseSearch::FSearchResults& SearchResults) const;
+	void SearchVPTree(UE::PoseSearch::FSearchContext& SearchContext, UE::PoseSearch::FSearchResults& SearchResults) const;
+	void SearchBruteForce(UE::PoseSearch::FSearchContext& SearchContext, UE::PoseSearch::FSearchResults& SearchResults) const;
+	void SearchEvent(UE::PoseSearch::FSearchContext& SearchContext, UE::PoseSearch::FSearchResults& SearchResults) const;
 
 	typedef TArray<int32, TInlineAllocator<256, TMemStackAllocator<>>> FSelectableAssetIdx;
-	UE_API void PopulateSelectableAssetIdx(FSelectableAssetIdx& SelectableAssetIdx, TConstArrayView<const UObject*> AssetsToConsider) const;
+	void PopulateSelectableAssetIdx(FSelectableAssetIdx& SelectableAssetIdx, TConstArrayView<const UObject*> AssetsToConsider) const;
 
 	typedef TArray<int32, TInlineAllocator<256, TMemStackAllocator<>>> FNonSelectableIdx;
-	UE_API void PopulateNonSelectableIdx(FNonSelectableIdx& NonSelectableIdx, UE::PoseSearch::FSearchContext& SearchContext
+	void PopulateNonSelectableIdx(FNonSelectableIdx& NonSelectableIdx, UE::PoseSearch::FSearchContext& SearchContext
 #if UE_POSE_SEARCH_TRACE_ENABLED
-		, float ContinuingPoseCostAddend, float ContinuingInteractionCostAddend, TConstArrayView<float> QueryValues, TConstArrayView<float> DynamicWeightsSqrt
+		, float ContinuingPoseCostAddend, float ContinuingInteractionCostAddend, TConstArrayView<float> DynamicWeightsSqrt
 #endif //UE_POSE_SEARCH_TRACE_ENABLED
 	) const;
 
-	UE_API void UpdateCachedProperties();
+	void UpdateCachedProperties();
 };
 
 template<typename TDatabaseAnimationAsset>
diff --git a/Engine/Plugins/Animation/PoseSearch/Source/Runtime/Public/PoseSearch/PoseSearchIndex.h b/Engine/Plugins/Animation/PoseSearch/Source/Runtime/Public/PoseSearch/PoseSearchIndex.h
--- a/Engine/Plugins/Animation/PoseSearch/Source/Runtime/Public/PoseSearch/PoseSearchIndex.h
+++ b/Engine/Plugins/Animation/PoseSearch/Source/Runtime/Public/PoseSearch/PoseSearchIndex.h
@@ -314,17 +314,40 @@ struct FSparsePoseMultiMap
 	FSparsePoseMultiMap& operator=(FSparsePoseMultiMap&& Other) = default;
 
 	FSparsePoseMultiMap(Type InMaxKey, Type InMaxValue)
-	: MaxKey(InMaxKey)
-	, MaxValue(InMaxValue)
-	, DeltaKeyValue(InMaxValue >= InMaxKey ? InMaxValue - InMaxKey + 1 : 0)
 	{
+		Init(InMaxKey, InMaxValue);
+	}
+
+	FSparsePoseMultiMap(TConstArrayView<TPair<Type, TArray<Type>>> Pairs, Type InMaxValue)
+	{
+		Init(Pairs, InMaxValue);
+	}
+
+	void Init(Type InMaxKey, Type InMaxValue)
+	{
+		MaxKey = InMaxKey;
+		MaxValue = InMaxValue;
+		DeltaKeyValue = InMaxValue >= InMaxKey ? InMaxValue - InMaxKey + 1 : 0;
+	
+		DataValues.Reset();
 		DataValues.Reserve(InMaxKey * 2);
+		
 		for (Type Index = 0; Index < InMaxKey; ++Index)
 		{
 			DataValues.Add(Type(INDEX_NONE));
 		}
 	}
 
+	void Init(TConstArrayView<TPair<Type, TArray<Type>>> Pairs, Type InMaxValue)
+	{
+		Init(Pairs.Num(), InMaxValue);
+		
+		for (const TPair<Type, TArray<Type>>& Pair : Pairs)
+		{
+			Insert(Pair.Key, Pair.Value);
+		}
+	}
+	
 	void Insert(Type Key, TConstArrayView<Type> Values)
 	{
 #if DO_CHECK
@@ -555,6 +578,11 @@ struct FSearchIndex : public FSearchIndexBase
 	// ((VA - VB).square() * VW).sum()
 	// since (VA - VB).square() could lead to big numbers, and VW being multiplied by the variance of the dataset
 	TAlignedArray<float> WeightsSqrt;
+
+	// array containing the data of FSearchIndexBase::Values encoded in PCA space. used to create the kdtree.
+	// Not all the poses have a related PCAValue(s) array since, for instance, poses marked as block transition 
+	// don't have to be searchable via kdtree, but only via continuing pose search, so method like 
+	// GetPCAValueIndex or GetPoseToPCAValuesVectorIndexes could not find the respective PCAValueIndex of a requested pose
 	TAlignedArray<float> PCAValues;
 	FSparsePoseMultiMap<int32> PCAValuesVectorToPoseIndexes;
 	TAlignedArray<float> PCAProjectionMatrix;
@@ -610,6 +638,9 @@ struct FSearchIndex : public FSearchIndexBase
 
 	// returns the inverse mapping of PCAValuesVectorToPoseIndexes
 	POSESEARCH_API void GetPoseToPCAValuesVectorIndexes(TArray<uint32>& PoseToPCAValuesVectorIndexes) const;
+	
+	// returns the PCAValueIndex from the pose index. the method is very slow. consider calling GetPoseToPCAValuesVectorIndexes to cache ALL the PCAValueIndex from the pose index mapping
+	int32 GetPCAValueIndex(int32 PoseIdx) const;
 
 	bool operator==(const FSearchIndex& Other) const;
 	friend FArchive& operator<<(FArchive& Ar, FSearchIndex& Index);
