Chooser
- introduced temporary scratch area in UChooserTable::EvaluateChooser when calling FChooserColumnBase::Filter and FChooserColumnBase::SetOutputs to support FPoseSearchColumn without needs for TThreadSingleton
- added FPoseSearchColumn rewind debugger support

diff --git a/Engine/Plugins/Animation/PoseSearch/Source/Runtime/Private/Chooser/PoseSearchChooserColumn.cpp b/Engine/Plugins/Animation/PoseSearch/Source/Runtime/Private/Chooser/PoseSearchChooserColumn.cpp
--- a/Engine/Plugins/Animation/PoseSearch/Source/Runtime/Private/Chooser/PoseSearchChooserColumn.cpp
+++ b/Engine/Plugins/Animation/PoseSearch/Source/Runtime/Private/Chooser/PoseSearchChooserColumn.cpp
@@ -23,16 +23,23 @@
 
 namespace UE::PoseSearch
 {
-	// @todo: we should rework how to pass temporary data from Filter to SetOutputs, but for now this makes things working
-	struct FPoseSearchColumnContextScratchArea : public TThreadSingleton<FPoseSearchColumnContextScratchArea>
+	// Experimental, this feature might be removed without warning, not for production use
+	struct FPoseSearchColumnScratchArea
 	{
 		TArray<UE::PoseSearch::FSearchResult> SearchResults;
 		const UE::PoseSearch::IPoseHistory* PoseHistory = nullptr;
+
 #if DO_CHECK
 		const FPoseSearchColumn* DebugOwner = nullptr;
 #endif // DO_CHECK
 	};
 
+	FArchive& operator<<(FArchive& Ar, FActiveColumnCost& ActiveColumnCost)
+	{
+		Ar << ActiveColumnCost.RowIndex;
+		Ar << ActiveColumnCost.RowCost;
+		return Ar;
+	}
 } // namespace UE::PoseSearch
 
 bool FPoseHistoryContextProperty::GetValue(FChooserEvaluationContext& Context, FPoseHistoryReference& OutResult) const
@@ -228,21 +235,56 @@ void FPoseSearchColumn::AutoPopulate(int32 RowIndex, UObject* OutputObject)
 
 bool FPoseSearchColumn::EditorTestFilter(int32 RowIndex) const
 {
-	// todo: not sure how to do editor display for pose search weights (might need to make a copy of the pose history struct)
-	return true;
+	using namespace UE::PoseSearch;
+
+	for (FActiveColumnCost& ActiveColumnCost : ActiveColumnCosts)
+	{
+		if (ActiveColumnCost.RowIndex == RowIndex)
+		{
+			return true;
+		}
+	}
+	return false;
 }
 
 float FPoseSearchColumn::EditorTestCost(int32 RowIndex) const
 {
-	// todo: not sure how to do editor display for pose search weights (might need to make a copy of the pose history struct)
-	return 0.0f;
+	using namespace UE::PoseSearch;
+
+	for (FActiveColumnCost& ActiveColumnCost : ActiveColumnCosts)
+	{
+		if (ActiveColumnCost.RowIndex == RowIndex)
+		{
+			return ActiveColumnCost.RowCost;
+		}
+	}
+	return 0.f;
+}
+
+void FPoseSearchColumn::SetTestValue(TArrayView<const uint8> Value)
+{
+    FMemoryReaderView Reader(Value);
+    Reader << ActiveColumnCosts;
 }
 
 #endif // WITH_EDITOR
 
+const FChooserParameterBase* FPoseSearchColumn::GetInputValue() const
+{
+	if (const FChooserParameterBase* ChooserParameterBase = InputValue.GetPtr<FChooserParameterBase>())
+	{
+		return ChooserParameterBase;
+	}
+	return &NullInputValue;
+}
+
 FChooserParameterBase* FPoseSearchColumn::GetInputValue()
 {
-	return InputValue.GetMutablePtr<FChooserParameterBase>();
+	if (FChooserParameterBase* ChooserParameterBase = InputValue.GetMutablePtr<FChooserParameterBase>())
+	{
+		return ChooserParameterBase;
+	}
+	return &NullInputValue;
 }
 
 bool FPoseSearchColumn::HasFilters() const
@@ -250,17 +292,22 @@ bool FPoseSearchColumn::HasFilters() const
 	return true;
 }
 
-void FPoseSearchColumn::Filter(FChooserEvaluationContext& Context, const FChooserIndexArray& IndexListIn, FChooserIndexArray& IndexListOut) const
+void FPoseSearchColumn::Filter(FChooserEvaluationContext& Context, const FChooserIndexArray& IndexListIn, FChooserIndexArray& IndexListOut, TArrayView<uint8> ScratchArea) const
 {
 	using namespace UE::PoseSearch;
 
 	CheckConsistency();
 
-	FPoseSearchColumnContextScratchArea& ScratchArea = FPoseSearchColumnContextScratchArea::Get();
+	check(ScratchArea.Num() == GetScratchAreaSize());
+	FPoseSearchColumnScratchArea* PoseSearchColumnScratchArea = reinterpret_cast<FPoseSearchColumnScratchArea*>(ScratchArea.begin());
 #if DO_CHECK
-	ScratchArea.DebugOwner = this;
+	check(PoseSearchColumnScratchArea->DebugOwner == this);
+	check(PoseSearchColumnScratchArea->SearchResults.IsEmpty());
 #endif // DO_CHECK
-	ScratchArea.SearchResults.Reset();
+
+#if WITH_EDITOR
+	TArray<FActiveColumnCost> TracedActiveColumnCosts;
+#endif // WITH_EDITOR
 
 	const UE::PoseSearch::IPoseHistory* PoseHistory = nullptr;
 	FPoseHistoryReference PoseHistoryReference;
@@ -297,7 +344,7 @@ void FPoseSearchColumn::Filter(FChooserEvaluationContext& Context, const FChoose
 	}
 	else
 	{
-		ScratchArea.PoseHistory = PoseHistory;
+		PoseSearchColumnScratchArea->PoseHistory = PoseHistory;
 
 		TArray<const UObject*, TInlineAllocator<128>> AssetsToConsider;
 		for (const FChooserIndexArray::FIndexData& IndexData : IndexListIn)
@@ -347,8 +394,11 @@ void FPoseSearchColumn::Filter(FChooserEvaluationContext& Context, const FChoose
 				// @todo: handle duplicate entries, by providing to the UPoseSearchLibrary::MotionMatch a SourceAssetIdxsToConsider rather than AssetsToConsider
 				const int32 SearchResultSourceAssetIdx = SearchIndexAsset->GetSourceAssetIdx();
 				check(SearchResultSourceAssetIdx >= 0);
-				ScratchArea.SearchResults.Add(BestSearchResult);
-				IndexListOut.Push({ static_cast<uint32>(SearchResultSourceAssetIdx), 0.f});
+				PoseSearchColumnScratchArea->SearchResults.Add(BestSearchResult);
+				IndexListOut.Push({ static_cast<uint32>(SearchResultSourceAssetIdx), 0.f });
+#if WITH_EDITOR
+				TracedActiveColumnCosts.Add({SearchResultSourceAssetIdx, 0.f});
+#endif // WITH_EDITOR
 			}
 		}
 		else
@@ -366,8 +416,11 @@ void FPoseSearchColumn::Filter(FChooserEvaluationContext& Context, const FChoose
 			{
 				if (const FSearchResult* FoundSearchResult = SearchResults.FindSearchResultFor(InternalDatabase.Get(), IndexData.Index))
 				{
-					ScratchArea.SearchResults.Add(*FoundSearchResult);
+					PoseSearchColumnScratchArea->SearchResults.Add(*FoundSearchResult);
 					IndexListOut.Push({ IndexData.Index, FoundSearchResult->PoseCost });
+#if WITH_EDITOR
+					TracedActiveColumnCosts.Add({ (int32)IndexData.Index, FoundSearchResult->PoseCost });
+#endif // WITH_EDITOR
 				}
 			}
 		}
@@ -378,9 +431,17 @@ void FPoseSearchColumn::Filter(FChooserEvaluationContext& Context, const FChoose
 		// if nothing passed the pose match, or if PoseHistory was not found, then ignore pose matching and just pass through everything
 		IndexListOut = IndexListIn;
 	}
+
+#if WITH_EDITOR
+	TRACE_CHOOSER_VALUE(Context, ToCStr(GetInputValue()->GetDebugName()), TracedActiveColumnCosts);
+	if (Context.DebuggingInfo.bCurrentDebugTarget)
+	{
+		ActiveColumnCosts = TracedActiveColumnCosts;
+	}
+#endif // WITH_EDITOR
 }
 
-void FPoseSearchColumn::SetOutputs(FChooserEvaluationContext& Context, int RowIndex) const
+void FPoseSearchColumn::SetOutputs(FChooserEvaluationContext& Context, int RowIndex, TArrayView<uint8> ScratchArea) const
 {
 	using namespace UE::PoseSearch;
 
@@ -390,12 +451,13 @@ void FPoseSearchColumn::SetOutputs(FChooserEvaluationContext& Context, int RowIn
 	bool bMirrorValue = false;
 	bool bForceBlendToValue = false;
 
-	FPoseSearchColumnContextScratchArea& ScratchArea = FPoseSearchColumnContextScratchArea::Get();
+	check(ScratchArea.Num() == GetScratchAreaSize());
+	FPoseSearchColumnScratchArea* PoseSearchColumnScratchArea = reinterpret_cast<FPoseSearchColumnScratchArea*>(ScratchArea.begin());
 #if DO_CHECK
-	check(ScratchArea.DebugOwner == this);
+	check(PoseSearchColumnScratchArea->DebugOwner == this);
 #endif // DO_CHECK
 
-	for (const FSearchResult& SearchResult : ScratchArea.SearchResults)
+	for (const FSearchResult& SearchResult : PoseSearchColumnScratchArea->SearchResults)
 	{
 		if (const FSearchIndexAsset* SearchIndexAsset = SearchResult.GetSearchIndexAsset())
 		{
@@ -407,9 +469,9 @@ void FPoseSearchColumn::SetOutputs(FChooserEvaluationContext& Context, int RowIn
 				bMirrorValue = SearchIndexAsset->IsMirrored();
 				bForceBlendToValue = SearchResult.bIsContinuingPoseSearch;
 
-				if (ScratchArea.PoseHistory)
+				if (PoseSearchColumnScratchArea->PoseHistory)
 				{
-					if (const FPoseIndicesHistory* PoseIndicesHistory = ScratchArea.PoseHistory->GetPoseIndicesHistory())
+					if (const FPoseIndicesHistory* PoseIndicesHistory = PoseSearchColumnScratchArea->PoseHistory->GetPoseIndicesHistory())
 					{
 						// @todo: integrate DeltaTime into SearchContext, and implement it for UAF as well
 						float DeltaTime = FiniteDelta;
@@ -427,12 +489,6 @@ void FPoseSearchColumn::SetOutputs(FChooserEvaluationContext& Context, int RowIn
 		}
 	}
 	
-#if DO_CHECK
-	ScratchArea.DebugOwner = nullptr;
-#endif // DO_CHECK
-	ScratchArea.SearchResults.Reset();
-	ScratchArea.PoseHistory = nullptr;
-
 	if (const FChooserParameterFloatBase* StartTime = OutputStartTime.GetPtr<FChooserParameterFloatBase>())
 	{
 		StartTime->SetValue(Context, StartTimeValue);
@@ -454,6 +510,33 @@ void FPoseSearchColumn::SetOutputs(FChooserEvaluationContext& Context, int RowIn
 	}
 }
 
+int32 FPoseSearchColumn::GetScratchAreaSize() const
+{
+	return sizeof(UE::PoseSearch::FPoseSearchColumnScratchArea);
+}
+
+void FPoseSearchColumn::InitializeScratchArea(TArrayView<uint8> ScratchArea) const
+{
+	using namespace UE::PoseSearch;
+	check(ScratchArea.Num() == GetScratchAreaSize());
+	FPoseSearchColumnScratchArea* PoseSearchColumnScratchArea = new(ScratchArea.begin()) FPoseSearchColumnScratchArea();
+
+#if DO_CHECK
+	PoseSearchColumnScratchArea->DebugOwner = this;
+#endif // DO_CHECK
+}
+
+void FPoseSearchColumn::DeinitializeScratchArea(TArrayView<uint8> ScratchArea) const
+{
+	using namespace UE::PoseSearch;
+	check(ScratchArea.Num() == GetScratchAreaSize());
+	FPoseSearchColumnScratchArea* PoseSearchColumnScratchArea = reinterpret_cast<FPoseSearchColumnScratchArea*>(ScratchArea.begin());
+#if DO_CHECK
+	check(PoseSearchColumnScratchArea->DebugOwner == this);
+#endif // DO_CHECK
+	PoseSearchColumnScratchArea->~FPoseSearchColumnScratchArea();
+}
+
 #if WITH_EDITOR
 void FPoseSearchColumn::Initialize(UChooserTable* OuterChooser)
 {
diff --git a/Engine/Plugins/Animation/PoseSearch/Source/Runtime/Public/PoseSearch/Chooser/PoseSearchChooserColumn.h b/Engine/Plugins/Animation/PoseSearch/Source/Runtime/Public/PoseSearch/Chooser/PoseSearchChooserColumn.h
--- a/Engine/Plugins/Animation/PoseSearch/Source/Runtime/Public/PoseSearch/Chooser/PoseSearchChooserColumn.h
+++ b/Engine/Plugins/Animation/PoseSearch/Source/Runtime/Public/PoseSearch/Chooser/PoseSearchChooserColumn.h
@@ -17,6 +17,18 @@
 
 class UAnimationAsset;
 
+namespace UE::PoseSearch
+{
+	// Experimental, this feature might be removed without warning, not for production use
+	struct FActiveColumnCost
+	{
+		int32 RowIndex = INDEX_NONE;
+		float RowCost = 0.f;
+	};
+	FArchive& operator<<(FArchive& Ar, FActiveColumnCost& ActiveColumnCost);
+} // namespace UE::PoseSearch
+
+
 USTRUCT(Experimental, DisplayName = "Pose History Property Binding")
 struct FPoseHistoryContextProperty :  public FChooserParameterPoseHistoryBase
 {
@@ -34,6 +46,18 @@ struct FPoseHistoryContextProperty :  public FChooserParameterPoseHistoryBase
 	CHOOSER_PARAMETER_BOILERPLATE();
 };
 
+// reintroduced just to make the TRACE_CHOOSER_VALUE work, since it relies on ToCStr(InputValue.Get<FChooserParameterBase>().GetDebugName()) to work
+USTRUCT(Experimental)
+struct FPoseSearchColumnNullInputValue : public FChooserParameterBase
+{
+	GENERATED_BODY()
+
+	virtual void GetDisplayName(FText& OutName) const override
+	{
+		OutName = FText::FromString("FPoseSearchColumn");
+	}
+};
+
 USTRUCT(DisplayName = "Pose Match", Meta = (Experimental, Category = "Experimental", Tooltip = "This column filters out all assets except the one which is selected by motion matching query.  Results must be AnimationAssets with a PoseSearchBranchIn notify state.  It also outputs OutputStartTime to specify the frame which matched pose best.  To work as intended it must be placed last (furthest right) in the Chooser so that other filters are applied first."))
 struct FPoseSearchColumn : public FChooserColumnBase
 {
@@ -65,6 +89,11 @@ struct FPoseSearchColumn : public FChooserColumnBase
 	UPROPERTY(EditAnywhere, DisplayName = "Pose History", NoClear, Meta = (ExcludeBaseStruct, BaseStruct = "/Script/PoseSearch.ChooserParameterPoseHistoryBase"), Category = "Data")
 	FInstancedStruct InputValue;
 	
+	// Experimental, this feature might be removed without warning, not for production use
+	// reintroduced just to make the TRACE_CHOOSER_VALUE work, since it relies on ToCStr(InputValue.Get<FChooserParameterBase>().GetDebugName()) to work
+	UPROPERTY()
+	FPoseSearchColumnNullInputValue NullInputValue;
+
 	// Float output for the start time with the best matching pose
 	UPROPERTY(EditAnywhere, NoClear, Meta = (ExcludeBaseStruct, BaseStruct = "/Script/Chooser.ChooserParameterFloatBase"), Category = "Output")
 	FInstancedStruct OutputStartTime;
@@ -81,23 +110,24 @@ struct FPoseSearchColumn : public FChooserColumnBase
 	UPROPERTY(EditAnywhere, NoClear, Meta = (ExcludeBaseStruct, BaseStruct = "/Script/Chooser.ChooserParameterFloatBase"), Category = "Output")
 	FInstancedStruct OutputCost;
 
+private:
 	virtual bool HasFilters() const override;
-	virtual void Filter(FChooserEvaluationContext& Context, const FChooserIndexArray& IndexListIn, FChooserIndexArray& IndexListOut) const override;
+	virtual void Filter(FChooserEvaluationContext& Context, const FChooserIndexArray& IndexListIn, FChooserIndexArray& IndexListOut, TArrayView<uint8> ScratchArea) const override;
 	virtual bool HasCosts() const override { return true; }
-	virtual void SetOutputs(FChooserEvaluationContext& Context, int RowIndex) const override;
+	virtual void SetOutputs(FChooserEvaluationContext& Context, int RowIndex, TArrayView<uint8> ScratchArea) const override;
+
+	virtual int32 GetScratchAreaSize() const override;
+	virtual void InitializeScratchArea(TArrayView<uint8> ScratchArea) const override;
+	virtual void DeinitializeScratchArea(TArrayView<uint8> ScratchArea) const override;
 
 #if WITH_EDITOR
 	virtual void Initialize(UChooserTable* OuterChooser) override;
 	
-	mutable double TestValue = 0.0;
+	mutable TArray<UE::PoseSearch::FActiveColumnCost> ActiveColumnCosts;
+
 	virtual bool EditorTestFilter(int32 RowIndex) const override;
 	virtual float EditorTestCost(int32 RowIndex) const override;
-	
-	virtual void SetTestValue(TArrayView<const uint8> Value) override
-	{
-		FMemoryReaderView Reader(Value);
-		Reader << TestValue;
-	}
+	virtual void SetTestValue(TArrayView<const uint8> Value) override;
 
 	virtual void AddToDetails(FInstancedPropertyBag& PropertyBag, int32 ColumnIndex, int32 RowIndex) override;
 	virtual void SetFromDetails(FInstancedPropertyBag& PropertyBag, int32 ColumnIndex, int32 RowIndex) override;
@@ -117,8 +147,8 @@ struct FPoseSearchColumn : public FChooserColumnBase
 
 #endif // WITH_EDITOR
 	virtual FChooserParameterBase* GetInputValue() override;
+	const FChooserParameterBase* GetInputValue() const;
 
-private:
 #if WITH_EDITOR
 	UObject* GetReferencedObject(int32 RowIndex) const;
 #endif // WITH_EDITOR
diff --git a/Engine/Plugins/Chooser/Source/Chooser/Private/Chooser.cpp b/Engine/Plugins/Chooser/Source/Chooser/Private/Chooser.cpp
--- a/Engine/Plugins/Chooser/Source/Chooser/Private/Chooser.cpp
+++ b/Engine/Plugins/Chooser/Source/Chooser/Private/Chooser.cpp
@@ -466,12 +466,48 @@ FObjectChooserBase::EIteratorStatus UChooserTable::EvaluateChooser(FChooserEvalu
 			Indices1.Push({i, 0});
 		}
 	}
+
+	// calculating the necessary scratch area size and allocating it on the stack
+	int32 TotalScratchAreaSize = 0;
+	for (const FInstancedStruct& ColumnData : Chooser->ColumnsStructs)
+	{
+		const FChooserColumnBase& Column = ColumnData.Get<FChooserColumnBase>();
+		TotalScratchAreaSize += Column.GetScratchAreaSize();
+	}
+	TArrayView<uint8> ScratchArea((uint8*)FMemory_Alloca(TotalScratchAreaSize * sizeof(uint8)), TotalScratchAreaSize);
+
+	auto DeinitializeScratchAreas = [Chooser, ScratchArea, TotalScratchAreaSize]()
+		{
+			if (TotalScratchAreaSize > 0)
+			{
+				// deinitializing all the ScratchArea(s)
+				int32 ScratchAreaStart = 0;
+				for (const FInstancedStruct& ColumnData : Chooser->ColumnsStructs)
+				{
+					const FChooserColumnBase& Column = ColumnData.Get<FChooserColumnBase>();
+					const int32 ColumnScratchAreaSize = Column.GetScratchAreaSize();
+
+					// @todo: maybe initialize ScratchArea only for column that has filters and are not bDisabled
+					TArrayView<uint8> ColumnScratchArea = ScratchArea.Slice(ScratchAreaStart, ColumnScratchAreaSize);
+					Column.DeinitializeScratchArea(ColumnScratchArea);
+					ScratchAreaStart += ColumnScratchAreaSize;
+				}
+			}
+		};
+
 	FChooserIndexArray* IndicesOut = &Indices1;
 	FChooserIndexArray* IndicesIn = &Indices2;
 
+	int32 ScratchAreaStart = 0;
 	for (const FInstancedStruct& ColumnData : Chooser->ColumnsStructs)
 	{
 		const FChooserColumnBase& Column = ColumnData.Get<FChooserColumnBase>();
+		const int32 ColumnScratchAreaSize = Column.GetScratchAreaSize();
+
+		// @todo: maybe initialize ScratchArea only for column that has filters and are not bDisabled
+		TArrayView<uint8> ColumnScratchArea = ScratchArea.Slice(ScratchAreaStart, ColumnScratchAreaSize);
+		Column.InitializeScratchArea(ColumnScratchArea);
+		ScratchAreaStart += ColumnScratchAreaSize;
 
 #if WITH_EDITORONLY_DATA
 		if (Column.bDisabled)
@@ -484,7 +520,7 @@ FObjectChooserBase::EIteratorStatus UChooserTable::EvaluateChooser(FChooserEvalu
 		{
 			Swap(IndicesIn, IndicesOut);
 			IndicesOut->SetNum(0);
-			Column.Filter(Context, *IndicesIn, *IndicesOut);
+			Column.Filter(Context, *IndicesIn, *IndicesOut, ColumnScratchArea);
 			
 			if (IndicesIn->HasCosts() || Column.HasCosts())
 			{
@@ -492,6 +528,7 @@ FObjectChooserBase::EIteratorStatus UChooserTable::EvaluateChooser(FChooserEvalu
 			}
 		}
 	}
+	check(ScratchAreaStart == TotalScratchAreaSize);
 
 	//	No need to score with only one valid option
 	if (IndicesOut->Num() > 1 && IndicesOut->HasCosts())
@@ -512,9 +549,15 @@ FObjectChooserBase::EIteratorStatus UChooserTable::EvaluateChooser(FChooserEvalu
 			{
 				bSetOutputs = true;
 				// trigger all output columns
+			
+				ScratchAreaStart = 0;
 				for (const FInstancedStruct& ColumnData : Chooser->ColumnsStructs)
 				{
 					const FChooserColumnBase& Column = ColumnData.Get<FChooserColumnBase>();
+					const int32 ColumnScratchAreaSize = Column.GetScratchAreaSize();
+
+					TArrayView<uint8> ColumnScratchArea = ScratchArea.Slice(ScratchAreaStart, ColumnScratchAreaSize);
+					ScratchAreaStart += ColumnScratchAreaSize;
 
 					#if WITH_EDITORONLY_DATA
 					if (Column.bDisabled)
@@ -523,8 +566,10 @@ FObjectChooserBase::EIteratorStatus UChooserTable::EvaluateChooser(FChooserEvalu
 					}
 					#endif
 					
-					Column.SetOutputs(Context, SelectedIndexData.Index);
+					Column.SetOutputs(Context, SelectedIndexData.Index, ColumnScratchArea);
 				}
+				check(ScratchAreaStart == TotalScratchAreaSize);
+
 				#if WITH_EDITOR
 				if (Context.DebuggingInfo.bCurrentDebugTarget)
 				{
@@ -535,6 +580,7 @@ FObjectChooserBase::EIteratorStatus UChooserTable::EvaluateChooser(FChooserEvalu
 			}
 			if (Status == FObjectChooserBase::EIteratorStatus::Stop)
 			{
+				DeinitializeScratchAreas();
 				return FObjectChooserBase::EIteratorStatus::Stop;
 			}
 		}
@@ -560,18 +606,28 @@ FObjectChooserBase::EIteratorStatus UChooserTable::EvaluateChooser(FChooserEvalu
 			{
 				bSetOutputs = true;
 				// trigger all output columns to set their default output value
+				ScratchAreaStart = 0;
 				for (const FInstancedStruct& ColumnData : Chooser->ColumnsStructs)
 				{
 					const FChooserColumnBase& Column = ColumnData.Get<FChooserColumnBase>();
-					Column.SetOutputs(Context, ChooserColumn_SpecialIndex_Fallback);
+					const int32 ColumnScratchAreaSize = Column.GetScratchAreaSize();
+
+					TArrayView<uint8> ColumnScratchArea = ScratchArea.Slice(ScratchAreaStart, ColumnScratchAreaSize);
+					ScratchAreaStart += ColumnScratchAreaSize;
+					Column.SetOutputs(Context, ChooserColumn_SpecialIndex_Fallback, ColumnScratchArea);
 				}
+				check(ScratchAreaStart == TotalScratchAreaSize);
 			}
 			if (Status == FObjectChooserBase::EIteratorStatus::Stop)
 			{
+				DeinitializeScratchAreas();
 				return FObjectChooserBase::EIteratorStatus::Stop;
 			}
 		}
 	}
+
+	DeinitializeScratchAreas();
+
 	// If this is a nested chooser make sure the parent also sets the output values from the row that contained this chooser
 	return bSetOutputs ? FObjectChooserBase::EIteratorStatus::ContinueWithOutputs : FObjectChooserBase::EIteratorStatus::Continue;
 }
diff --git a/Engine/Plugins/Chooser/Source/Chooser/Public/IChooserColumn.h b/Engine/Plugins/Chooser/Source/Chooser/Public/IChooserColumn.h
--- a/Engine/Plugins/Chooser/Source/Chooser/Public/IChooserColumn.h
+++ b/Engine/Plugins/Chooser/Source/Chooser/Public/IChooserColumn.h
@@ -38,13 +38,26 @@ struct FChooserColumnBase
 public:
 	virtual ~FChooserColumnBase() {}
 	virtual void PostLoad() {}
+
 	virtual void Filter(FChooserEvaluationContext& Context, const FChooserIndexArray& IndexListIn, FChooserIndexArray& IndexListOut) const {}
+	
+	// Experimental, this feature might be removed without warning, not for production use
+	virtual void Filter(FChooserEvaluationContext& Context, const FChooserIndexArray& IndexListIn, FChooserIndexArray& IndexListOut, TArrayView<uint8> ScratchArea) const
+	{
+		Filter(Context, IndexListIn, IndexListOut);
+	}
 
 	virtual bool HasCosts() const { return false; }
 	virtual bool HasFilters() const { return true; }
 	virtual bool HasOutputs() const { return false; }
 	
 	virtual void SetOutputs(FChooserEvaluationContext& Context, int RowIndex) const { }
+	
+	// Experimental, this feature might be removed without warning, not for production use
+	virtual void SetOutputs(FChooserEvaluationContext& Context, int RowIndex, TArrayView<uint8> ScratchArea) const
+	{
+		SetOutputs(Context, RowIndex);
+	}
 
 	virtual void Compile(IHasContextClass* Owner, bool bForce)
 	{
@@ -62,6 +75,14 @@ public:
 
 	virtual void SetTestValue(TArrayView<const uint8> Value) {}
 	virtual FChooserParameterBase* GetInputValue() { return nullptr; }
+
+	// Experimental, this feature might be removed without warning, not for production use
+	virtual int32 GetScratchAreaSize() const { return 0; }
+	// Experimental, this feature might be removed without warning, not for production use
+	virtual void InitializeScratchArea(TArrayView<uint8> ScratchArea) const {}
+	// Experimental, this feature might be removed without warning, not for production use
+	virtual void DeinitializeScratchArea(TArrayView<uint8> ScratchArea) const {}
+
 #if WITH_EDITOR
 	virtual void Initialize(UChooserTable* OuterChooser) {}
 	
