PoseSearch - tool side performance optimizations

diff --git a/Engine/Plugins/Animation/PoseSearch/Source/Runtime/Private/PoseSearchDerivedData.cpp b/Engine/Plugins/Animation/PoseSearch/Source/Runtime/Private/PoseSearchDerivedData.cpp
--- a/Engine/Plugins/Animation/PoseSearch/Source/Runtime/Private/PoseSearchDerivedData.cpp
+++ b/Engine/Plugins/Animation/PoseSearch/Source/Runtime/Private/PoseSearchDerivedData.cpp
@@ -1622,7 +1622,7 @@ void FPoseSearchDatabaseAsyncCacheTask::PreCancelIfDependsOn(const UObject* Obje
 {
 	check(IsInGameThread());
 
-	if (Object)
+	if (FKeyBuilder::ShouldInclude(Object))
 	{
 		// DatabaseDependencies is updated only in StartNewRequestIfNeeded when there are no active requests, so it's thread safe to access it 
 		if (DatabaseDependencies.Contains(Object))
@@ -2342,36 +2342,42 @@ void FAsyncPoseSearchDatabasesManagement::PreModified(UObject* Object)
 {
 	check(IsInGameThread());
 
-	FScopeLock Lock(&Mutex);
+	if (FKeyBuilder::ShouldInclude(Object))
+	{
+		FScopeLock Lock(&Mutex);
 
-	PartialKeyHashes.Remove(Object);
+		PartialKeyHashes.Remove(Object);
 
-	for (TUniquePtr<FPoseSearchDatabaseAsyncCacheTask>& TaskPtr : Tasks)
-	{
-		TaskPtr->PreCancelIfDependsOn(Object);
-	}
+		for (TUniquePtr<FPoseSearchDatabaseAsyncCacheTask>& TaskPtr : Tasks)
+		{
+			TaskPtr->PreCancelIfDependsOn(Object);
+		}
 
-	// collecting databases to synchronize prior modifying the Object
-	CollectDatabasesToSynchronize(Object);
+		// collecting databases to synchronize prior modifying the Object
+		CollectDatabasesToSynchronize(Object);
+	}
 }
 
 void FAsyncPoseSearchDatabasesManagement::PostModified(UObject* Object)
 {
 	check(IsInGameThread());
 
-	FScopeLock Lock(&Mutex);
+	if (FKeyBuilder::ShouldInclude(Object))
+	{
+		FScopeLock Lock(&Mutex);
 
-	// collecting databases to synchronize, and merging the results with the PreModified collection
-	CollectDatabasesToSynchronize(Object);
+		// collecting databases to synchronize, and merging the results with the PreModified collection
+		CollectDatabasesToSynchronize(Object);
 
-	SynchronizeDatabases();
+		SynchronizeDatabases();
 
-	// promoting the EState::PreCancelled into EState::Cancelled
-	for (TUniquePtr<FPoseSearchDatabaseAsyncCacheTask>& TaskPtr : Tasks)
-	{
-		if (TaskPtr->GetState() == FPoseSearchDatabaseAsyncCacheTask::EState::PreCancelled)
+		// promoting the EState::PreCancelled into EState::Cancelled
+		for (TUniquePtr<FPoseSearchDatabaseAsyncCacheTask>& TaskPtr : Tasks)
 		{
-			TaskPtr->Cancel();
+			if (TaskPtr->GetState() == FPoseSearchDatabaseAsyncCacheTask::EState::PreCancelled)
+			{
+				TaskPtr->Cancel();
+			}
 		}
 	}
 }
diff --git a/Engine/Plugins/Animation/PoseSearch/Source/Runtime/Private/PoseSearchDerivedDataKey.cpp b/Engine/Plugins/Animation/PoseSearch/Source/Runtime/Private/PoseSearchDerivedDataKey.cpp
--- a/Engine/Plugins/Animation/PoseSearch/Source/Runtime/Private/PoseSearchDerivedDataKey.cpp
+++ b/Engine/Plugins/Animation/PoseSearch/Source/Runtime/Private/PoseSearchDerivedDataKey.cpp
@@ -217,14 +217,14 @@ FArchive& FKeyBuilder::TryAddDependency(UObject* Object, bool bAddToPartialKeyHa
 
 FArchive& FKeyBuilder::operator<<(class UObject*& Object)
 {
-	if (!Object)
+	if (ShouldInclude(Object))
 	{
-		return *this;
+		return TryAddDependency(Object, true);
 	}
 
-	if (Object->HasAnyFlags(RF_Transient | RF_MirroredGarbage | RF_NewerVersionExists | RF_BeginDestroyed | RF_FinishDestroyed | RF_DuplicateTransient | RF_NonPIEDuplicateTransient) || IsExcludedType(Object))
-	{
 #if UE_POSE_SEARCH_DERIVED_DATA_LOGGING
+	if (Object)
+	{
 		if (Object->HasAnyFlags(RF_Transient))
 		{
 			UE_LOG(LogPoseSearch, Log, TEXT("Transient '%s' (%s)"), *Object->GetName(), *Object->GetClass()->GetName());
@@ -257,11 +257,9 @@ FArchive& FKeyBuilder::operator<<(class UObject*& Object)
 		{
 			UE_LOG(LogPoseSearch, Log, TEXT("Excluded '%s' (%s)"), *Object->GetName(), *Object->GetClass()->GetName());
 		}
-#endif // UE_POSE_SEARCH_DERIVED_DATA_LOGGING
-		return *this;
 	}
-
-	return TryAddDependency(Object, true);
+#endif // UE_POSE_SEARCH_DERIVED_DATA_LOGGING
+	return *this;
 }
 
 void FKeyBuilder::SerializeObjectInternal(UObject* Object)
@@ -425,7 +423,7 @@ const TSet<const UObject*>& FKeyBuilder::GetDependencies() const
 }
 
 // to keep the key generation lightweight, we don't hash these types
-bool FKeyBuilder::IsExcludedType(class UObject* Object)
+bool FKeyBuilder::IsExcludedType(const UObject* Object)
 {
 	if (Object->IsA<UAnimationModifier>())
 	{
@@ -447,8 +445,28 @@ bool FKeyBuilder::IsExcludedType(class UObject* Object)
 	return false;
 }
 
+bool FKeyBuilder::ShouldInclude(const UObject* Object)
+{
+	if (!Object)
+	{
+		return false;
+	}
+
+	if (Object->HasAnyFlags(RF_Transient | RF_MirroredGarbage | RF_NewerVersionExists | RF_BeginDestroyed | RF_FinishDestroyed | RF_DuplicateTransient | RF_NonPIEDuplicateTransient))
+	{
+		return false;
+	}
+
+	if (IsExcludedType(Object))
+	{
+		return false;
+	}
+
+	return true;
+}
+
 // to keep the key generation lightweight, we hash only the full names for these types. Object(s) will be added to Dependencies
-bool FKeyBuilder::IsAddNameOnlyType(class UObject* Object)
+bool FKeyBuilder::IsAddNameOnlyType(const UObject* Object)
 {
 	check(Object);
 
diff --git a/Engine/Plugins/Animation/PoseSearch/Source/Runtime/Private/Trace/PoseSearchTraceLogger.cpp b/Engine/Plugins/Animation/PoseSearch/Source/Runtime/Private/Trace/PoseSearchTraceLogger.cpp
--- a/Engine/Plugins/Animation/PoseSearch/Source/Runtime/Private/Trace/PoseSearchTraceLogger.cpp
+++ b/Engine/Plugins/Animation/PoseSearch/Source/Runtime/Private/Trace/PoseSearchTraceLogger.cpp
@@ -44,7 +44,7 @@ FArchive& operator<<(FArchive& Ar, FTraceMessage& State)
 FArchive& operator<<(FArchive& Ar, FTraceMotionMatchingStatePoseEntry& Entry)
 {
 	Ar << Entry.DbPoseIdx;
-	FPoseSearchCost::StaticStruct()->SerializeItem(Ar, &Entry.Cost, nullptr);
+	Ar << Entry.Cost;
 	Ar << Entry.PoseCandidateFlags;
 	return Ar;
 }
diff --git a/Engine/Plugins/Animation/PoseSearch/Source/Runtime/Public/PoseSearch/PoseSearchCost.h b/Engine/Plugins/Animation/PoseSearch/Source/Runtime/Public/PoseSearch/PoseSearchCost.h
--- a/Engine/Plugins/Animation/PoseSearch/Source/Runtime/Public/PoseSearch/PoseSearchCost.h
+++ b/Engine/Plugins/Animation/PoseSearch/Source/Runtime/Public/PoseSearch/PoseSearchCost.h
@@ -50,6 +50,24 @@ PRAGMA_ENABLE_DEPRECATION_WARNINGS
 	float GetContinuingInteractionCostAddend() const { return ContinuingInteractionCostAddend; }
 #endif // WITH_EDITORONLY_DATA
 
+	friend FArchive& operator<<(FArchive& Ar, FPoseSearchCost& PoseSearchCost)
+	{
+		Ar << PoseSearchCost.TotalCost;
+
+#if WITH_EDITORONLY_DATA
+		Ar << PoseSearchCost.NotifyCostAddend;
+		Ar << PoseSearchCost.ContinuingPoseCostAddend;
+		Ar << PoseSearchCost.ContinuingInteractionCostAddend;
+
+		PRAGMA_DISABLE_DEPRECATION_WARNINGS
+		Ar << PoseSearchCost.CostAddend;
+		PRAGMA_ENABLE_DEPRECATION_WARNINGS
+
+#endif // WITH_EDITORONLY_DATA
+			
+		return Ar;
+	}
+
 protected:
 	// TotalCost is the sum of all the Cost contributions (dissimilarity, notifies, continuing pose, continuing interaction costs)
 	UPROPERTY()
@@ -77,3 +95,4 @@ public:
 	float CostAddend = 0.f;
 #endif // WITH_EDITORONLY_DATA
 };
+
diff --git a/Engine/Plugins/Animation/PoseSearch/Source/Runtime/Public/PoseSearch/PoseSearchDerivedDataKey.h b/Engine/Plugins/Animation/PoseSearch/Source/Runtime/Public/PoseSearch/PoseSearchDerivedDataKey.h
--- a/Engine/Plugins/Animation/PoseSearch/Source/Runtime/Public/PoseSearch/PoseSearchDerivedDataKey.h
+++ b/Engine/Plugins/Animation/PoseSearch/Source/Runtime/Public/PoseSearch/PoseSearchDerivedDataKey.h
@@ -59,12 +59,15 @@ public:
 	POSESEARCH_API FIoHash Finalize() const;
 	POSESEARCH_API const TSet<const UObject*>& GetDependencies() const;
 
+	// Experimental, this feature might be removed without warning, not for production use
+	static POSESEARCH_API bool ShouldInclude(const UObject* Object);
+
 protected:
 	// to keep the key generation lightweight, we don't hash these types
-	static POSESEARCH_API bool IsExcludedType(class UObject* Object);
+	static POSESEARCH_API bool IsExcludedType(const UObject* Object);
 
 	// to keep the key generation lightweight, we hash only the full names for these types. Object(s) will be added to Dependencies
-	static POSESEARCH_API bool IsAddNameOnlyType(class UObject* Object);
+	static POSESEARCH_API bool IsAddNameOnlyType(const UObject* Object);
 
 	HashBuilderType Hasher;
 
