PoseSearch - fix for logic in FAsyncPoseSearchDatabasesManagement::CollectDatabasesToSynchronize when input object is a UAnimNotifyState_PoseSearchBranchIn

diff --git a/Engine/Plugins/Animation/PoseSearch/Source/Runtime/Private/PoseSearchDerivedData.cpp b/Engine/Plugins/Animation/PoseSearch/Source/Runtime/Private/PoseSearchDerivedData.cpp
--- a/Engine/Plugins/Animation/PoseSearch/Source/Runtime/Private/PoseSearchDerivedData.cpp
+++ b/Engine/Plugins/Animation/PoseSearch/Source/Runtime/Private/PoseSearchDerivedData.cpp
@@ -2271,7 +2271,7 @@ void FAsyncPoseSearchDatabasesManagement::CollectDatabasesToSynchronize(UObject*
 		return;
 	}
 	
-	if (!Object->IsA<UAnimSequenceBase>() && Object->IsA<UAnimNotifyState_PoseSearchBranchIn>())
+	if (!Object->IsA<UAnimSequenceBase>() && !Object->IsA<UAnimNotifyState_PoseSearchBranchIn>())
 	{
 		return;
 	}
