PoseSearch - UPoseSearchDatabase::GetAssetIndexesForSourceAsset perfromance improvements to speed up FPoseSearchColumn::Filter

diff --git a/Engine/Plugins/Animation/PoseSearch/Source/Runtime/Private/PoseSearchDatabase.cpp b/Engine/Plugins/Animation/PoseSearch/Source/Runtime/Private/PoseSearchDatabase.cpp
--- a/Engine/Plugins/Animation/PoseSearch/Source/Runtime/Private/PoseSearchDatabase.cpp
+++ b/Engine/Plugins/Animation/PoseSearch/Source/Runtime/Private/PoseSearchDatabase.cpp
@@ -1066,7 +1066,7 @@ void UPoseSearchDatabase::UpdateCachedProperties()
 		}
 	}
 
-	for (TPair<FObjectKey, TArray<int32>>& CachedAssetMapPair : CachedAssetMap)
+	for (FCachedAssetMapPair& CachedAssetMapPair : CachedAssetMap)
 	{
 		CachedAssetMapPair.Value.Sort();
 	}
@@ -1074,24 +1074,8 @@ void UPoseSearchDatabase::UpdateCachedProperties()
 
 TConstArrayView<int32> UPoseSearchDatabase::GetAssetIndexesForSourceAsset(const UObject* SourceAsset) const
 {
-	using namespace UE::PoseSearch;
-
 	if (const TArray<int32>* IndexesForSourceAsset = CachedAssetMap.Find(SourceAsset))
 	{
-#if DO_CHECK
-		// validating the consistency of IndexesForSourceAsset retrieved from SourceAsset
-		const FSearchIndex& SearchIndex = GetSearchIndex();
-		for (int32 AssetIndex : *IndexesForSourceAsset)
-		{
-			const FSearchIndexAsset& SearchIndexAsset = SearchIndex.Assets[AssetIndex];
-			const FPoseSearchDatabaseAnimationAssetBase* DatabaseAnimationAssetBase = GetDatabaseAnimationAsset(SearchIndexAsset);
-	
-			// if those checks fail the calling code hasn't been protected by FAsyncPoseSearchDatabasesManagement::RequestAsyncBuildIndex
-			check(DatabaseAnimationAssetBase);
-			check(DatabaseAnimationAssetBase->GetAnimationAsset() == SourceAsset);
-		}
-#endif // DO_CHECK
-
 		return *IndexesForSourceAsset;
 	}
 	return TConstArrayView<int32>();
diff --git a/Engine/Plugins/Animation/PoseSearch/Source/Runtime/Public/PoseSearch/PoseSearchDatabase.h b/Engine/Plugins/Animation/PoseSearch/Source/Runtime/Public/PoseSearch/PoseSearchDatabase.h
--- a/Engine/Plugins/Animation/PoseSearch/Source/Runtime/Public/PoseSearch/PoseSearchDatabase.h
+++ b/Engine/Plugins/Animation/PoseSearch/Source/Runtime/Public/PoseSearch/PoseSearchDatabase.h
@@ -631,7 +631,10 @@ private:
 	UE::PoseSearch::FSearchIndex SearchIndexPrivate;
 	
 	// CachedAssetMap is NOT serialized in operator<< but recalculated by UpdateCachedProperties every time SearchIndexPrivate changes
-	TMap<FObjectKey, TArray<int32>> CachedAssetMap;
+	// accessing the key const UObject* is protected by FAsyncPoseSearchDatabasesManagement::RequestAsyncBuildIndex so it doesn't have to be TWeakObjectPtr<const UObject> / FObjectKey for GC protection
+	typedef TMap<const UObject*, TArray<int32>> FCachedAssetMap;
+	typedef TPair<const UObject*, TArray<int32>> FCachedAssetMapPair;
+	FCachedAssetMap CachedAssetMap;
 
 #if WITH_EDITOR
 	DECLARE_MULTICAST_DELEGATE(FOnDerivedDataRebuildMulticaster);
