PoseSearch - making FPoseSearchColumn::InterruptMode pinnable

diff --git a/Engine/Plugins/Animation/PoseSearch/Source/Runtime/Private/Chooser/PoseSearchChooserColumn.cpp b/Engine/Plugins/Animation/PoseSearch/Source/Runtime/Private/Chooser/PoseSearchChooserColumn.cpp
--- a/Engine/Plugins/Animation/PoseSearch/Source/Runtime/Private/Chooser/PoseSearchChooserColumn.cpp
+++ b/Engine/Plugins/Animation/PoseSearch/Source/Runtime/Private/Chooser/PoseSearchChooserColumn.cpp
@@ -8,6 +8,7 @@
 #include "ChooserTrace.h"
 #include "ChooserTypes.h"
 #include "IChooserParameterBool.h"
+#include "IChooserParameterEnum.h"
 #include "IChooserParameterFloat.h"
 #include "PoseSearch/PoseSearchContext.h"
 #include "PoseSearch/PoseSearchDatabase.h"
@@ -371,7 +372,12 @@ void FPoseSearchColumn::Filter(FChooserEvaluationContext& Context, const FChoose
 			ContinuingProperties.PlayingAssetAccumulatedTime = ChooserPlayerSettings->PlayingAssetAccumulatedTime;
 			ContinuingProperties.bIsPlayingAssetMirrored = ChooserPlayerSettings->bIsPlayingAssetMirrored;
 			ContinuingProperties.PlayingAssetBlendParameters = ChooserPlayerSettings->PlayingAssetBlendParameters;
-			ContinuingProperties.InterruptMode = InterruptMode;
+
+			uint8 InterruptModeValue = 0;
+			if (InterruptMode.IsValid() && InterruptMode.Get<FChooserParameterEnumBase>().GetValue(Context, InterruptModeValue))
+			{
+				ContinuingProperties.InterruptMode = (EPoseSearchInterruptMode)InterruptModeValue;
+			}
 		}
 
 		FSearchResult BestSearchResult;
diff --git a/Engine/Plugins/Animation/PoseSearch/Source/Runtime/Public/PoseSearch/AnimNode_MotionMatching.h b/Engine/Plugins/Animation/PoseSearch/Source/Runtime/Public/PoseSearch/AnimNode_MotionMatching.h
--- a/Engine/Plugins/Animation/PoseSearch/Source/Runtime/Public/PoseSearch/AnimNode_MotionMatching.h
+++ b/Engine/Plugins/Animation/PoseSearch/Source/Runtime/Public/PoseSearch/AnimNode_MotionMatching.h
@@ -172,7 +172,7 @@ private:
 	UPROPERTY(Transient)
 	TArray<TObjectPtr<const UPoseSearchDatabase>> DatabasesToSearch;
 
-	// Applied EPoseSearchInterruptMode on the next update that controls the continuing pose search eveluation. This is set back to EPoseSearchInterruptMode::DoNotInterrupt after each update.
+	// Applied EPoseSearchInterruptMode on the next update that controls the continuing pose search evaluation. This is set back to EPoseSearchInterruptMode::DoNotInterrupt after each update.
 	EPoseSearchInterruptMode NextUpdateInterruptMode = EPoseSearchInterruptMode::DoNotInterrupt;
 
 	// True if the Database property on this node has been overridden by SetDatabaseToSearch/SetDatabasesToSearch.
diff --git a/Engine/Plugins/Animation/PoseSearch/Source/Runtime/Public/PoseSearch/Chooser/PoseSearchChooserColumn.h b/Engine/Plugins/Animation/PoseSearch/Source/Runtime/Public/PoseSearch/Chooser/PoseSearchChooserColumn.h
--- a/Engine/Plugins/Animation/PoseSearch/Source/Runtime/Public/PoseSearch/Chooser/PoseSearchChooserColumn.h
+++ b/Engine/Plugins/Animation/PoseSearch/Source/Runtime/Public/PoseSearch/Chooser/PoseSearchChooserColumn.h
@@ -81,9 +81,9 @@ struct FPoseSearchColumn : public FChooserColumnBase
 	UPROPERTY(EditAnywhere, Category = "Data")
 	int32 MaxNumberOfResults = 1;
 
-	// @todo: expose it as pinnable parameter
-	UPROPERTY(EditAnywhere, Category = "Data")
-	EPoseSearchInterruptMode InterruptMode = EPoseSearchInterruptMode::DoNotInterrupt;
+	// @todo: filter to allow only enums of type EPoseSearchInterruptMode
+	UPROPERTY(EditAnywhere, NoClear, Meta = (ExcludeBaseStruct, BaseStruct = "/Script/Chooser.ChooserParameterEnumBase", ToolTip="Applied EPoseSearchInterruptMode controlling the continuing pose search evaluation. Defaulted to EPoseSearchInterruptMode::DoNotInterrupt if not set"), Category = "Data")
+	FInstancedStruct InterruptMode;
 
 	// Pose History
 	UPROPERTY(EditAnywhere, DisplayName = "Pose History", NoClear, Meta = (ExcludeBaseStruct, BaseStruct = "/Script/PoseSearch.ChooserParameterPoseHistoryBase"), Category = "Data")
diff --git a/Engine/Plugins/Experimental/UAF/UAFPoseSearch/Source/UAFPoseSearch/Internal/MotionMatchingTraitData.h b/Engine/Plugins/Experimental/UAF/UAFPoseSearch/Source/UAFPoseSearch/Internal/MotionMatchingTraitData.h
--- a/Engine/Plugins/Experimental/UAF/UAFPoseSearch/Source/UAFPoseSearch/Internal/MotionMatchingTraitData.h
+++ b/Engine/Plugins/Experimental/UAF/UAFPoseSearch/Source/UAFPoseSearch/Internal/MotionMatchingTraitData.h
@@ -77,7 +77,7 @@ struct FMotionMatchingTraitSharedData : public FAnimNextTraitSharedData
 
 	// @todo: should this be in FMotionMatchingState?
 	UPROPERTY(EditAnywhere, Category = Settings, meta = (PinHiddenByDefault))
-	// Applied EPoseSearchInterruptMode on the next update that controls the continuing pose search eveluation. This is set back to EPoseSearchInterruptMode::DoNotInterrupt after each update.
+	// Applied EPoseSearchInterruptMode on the next update that controls the continuing pose search evaluation. This is set back to EPoseSearchInterruptMode::DoNotInterrupt after each update.
 	EPoseSearchInterruptMode NextUpdateInterruptMode = EPoseSearchInterruptMode::DoNotInterrupt;
 
 	// If set, the motion matching result will be written to this variable
