PoseSearch - fix for FSearchResults_AssetBests::Shrink actually shrinking the number of preserved results

diff --git a/Engine/Plugins/Animation/PoseSearch/Source/Runtime/Private/PoseSearchResult.cpp b/Engine/Plugins/Animation/PoseSearch/Source/Runtime/Private/PoseSearchResult.cpp
--- a/Engine/Plugins/Animation/PoseSearch/Source/Runtime/Private/PoseSearchResult.cpp
+++ b/Engine/Plugins/Animation/PoseSearch/Source/Runtime/Private/PoseSearchResult.cpp
@@ -468,6 +468,16 @@ void FSearchResults_AssetBests::Shrink(int32 MaxNumResults)
 				}
 			}
 		}
+
+		PerDatabaseSourceAssetIdxToBestSearchResult.Reset();
+		for (const FSearchResult& BestSearchResult : BestSearchResults)
+		{
+			const FSearchIndexAsset* SearchIndexAsset = BestSearchResult.GetSearchIndexAsset();
+			check(SearchIndexAsset);
+			FSourceAssetIdxToBestSearchResult& SourceAssetIdxToBestSearchResult = PerDatabaseSourceAssetIdxToBestSearchResult.FindOrAdd(BestSearchResult.Database.Get());
+			FSearchResult& BestSearchResultForSourceAssetIdx = SourceAssetIdxToBestSearchResult.FindOrAdd(SearchIndexAsset->GetSourceAssetIdx());
+			BestSearchResultForSourceAssetIdx = BestSearchResult;
+		}
 	}
 }
 
