PoseSearch - PopulateSearches performance improvement to speed up FPoseSearchColumn::Filter

diff --git a/Engine/Plugins/Animation/PoseSearch/Source/Runtime/Private/Chooser/PoseSearchChooserColumn.cpp b/Engine/Plugins/Animation/PoseSearch/Source/Runtime/Private/Chooser/PoseSearchChooserColumn.cpp
--- a/Engine/Plugins/Animation/PoseSearch/Source/Runtime/Private/Chooser/PoseSearchChooserColumn.cpp
+++ b/Engine/Plugins/Animation/PoseSearch/Source/Runtime/Private/Chooser/PoseSearchChooserColumn.cpp
@@ -341,7 +341,7 @@ void FPoseSearchColumn::Filter(FChooserEvaluationContext& Context, const FChoose
 	{
 		PoseSearchColumnScratchArea->PoseHistory = PoseHistory;
 
-		TArray<const UObject*, TInlineAllocator<128>> AssetsToConsider;
+		FStackAssetSet AssetsToConsider;
 		for (const FChooserIndexArray::FIndexData& IndexData : IndexListIn)
 		{
 			const FPoseSearchDatabaseAnimationAsset* Asset = InternalDatabase->GetDatabaseAnimationAsset(IndexData.Index);
@@ -363,7 +363,7 @@ void FPoseSearchColumn::Filter(FChooserEvaluationContext& Context, const FChoose
 
 		const FRole Role = InternalDatabase->Schema ? InternalDatabase->Schema->GetDefaultRole() : DefaultRole;
 		SearchContext.AddRole(Role, &Context, PoseHistory);
-		SearchContext.SetAssetsToConsider(AssetsToConsider);
+		SearchContext.SetAssetsToConsiderSet(&AssetsToConsider);
 
 		FPoseSearchContinuingProperties ContinuingProperties;
 		if (ChooserPlayerSettings)
diff --git a/Engine/Plugins/Animation/PoseSearch/Source/Runtime/Private/PoseSearchDatabase.cpp b/Engine/Plugins/Animation/PoseSearch/Source/Runtime/Private/PoseSearchDatabase.cpp
--- a/Engine/Plugins/Animation/PoseSearch/Source/Runtime/Private/PoseSearchDatabase.cpp
+++ b/Engine/Plugins/Animation/PoseSearch/Source/Runtime/Private/PoseSearchDatabase.cpp
@@ -1725,6 +1725,32 @@ void UPoseSearchDatabase::PopulateSelectableAssetIdx(FSelectableAssetIdx& Select
 	}
 }
 
+void UPoseSearchDatabase::PopulateSelectableAssetIdx(FSelectableAssetIdx& SelectableAssetIdx, const UE::PoseSearch::FStackAssetSet* AssetsToConsider) const
+{
+	SelectableAssetIdx.Reset();
+	if (AssetsToConsider && !AssetsToConsider->IsEmpty())
+	{
+		for (const UObject* AssetToConsider : *AssetsToConsider)
+		{
+			SelectableAssetIdx.Append(GetAssetIndexesForSourceAsset(AssetToConsider));
+		}
+
+		if (!SelectableAssetIdx.IsEmpty())
+		{
+			if (SelectableAssetIdx.Num() != GetSearchIndex().Assets.Num())
+			{
+				SelectableAssetIdx.Sort();
+			}
+			else
+			{
+				// SelectableAssetIdx contains ALL the Database->GetSearchIndex().Assets. 
+				// We reset SelectableAssetIdx since it has the same meaning, and it'll perform better
+				SelectableAssetIdx.Reset();
+			}
+		}
+	}
+}
+
 UE::PoseSearch::FSearchResult UPoseSearchDatabase::Search(UE::PoseSearch::FSearchContext& SearchContext) const
 {
 	using namespace UE::PoseSearch;
@@ -1890,7 +1916,16 @@ void UPoseSearchDatabase::SearchPCAKDTree(UE::PoseSearch::FSearchContext& Search
 		check(QueryValues.Num() == NumDimensions);
 
 		FSelectableAssetIdx SelectableAssetIdx;
-		PopulateSelectableAssetIdx(SelectableAssetIdx, SearchContext.GetAssetsToConsider());
+		PRAGMA_DISABLE_DEPRECATION_WARNINGS
+		if (!SearchContext.GetInternalDeprecatedAssetsToConsider().IsEmpty())
+		{
+			PopulateSelectableAssetIdx(SelectableAssetIdx, SearchContext.GetInternalDeprecatedAssetsToConsider());
+		}
+		else
+		PRAGMA_ENABLE_DEPRECATION_WARNINGS
+		{
+			PopulateSelectableAssetIdx(SelectableAssetIdx, SearchContext.GetAssetsToConsiderSet());
+		}
 
 		FNonSelectableIdx NonSelectableIdx;
 		PopulateNonSelectableIdx(NonSelectableIdx, SearchContext
@@ -2245,7 +2280,16 @@ void UPoseSearchDatabase::SearchVPTree(UE::PoseSearch::FSearchContext& SearchCon
 		check(QueryValues.Num() == NumDimensions);
 
 		FSelectableAssetIdx SelectableAssetIdx;
-		PopulateSelectableAssetIdx(SelectableAssetIdx, SearchContext.GetAssetsToConsider());
+		PRAGMA_DISABLE_DEPRECATION_WARNINGS
+		if (!SearchContext.GetInternalDeprecatedAssetsToConsider().IsEmpty())
+		{
+			PopulateSelectableAssetIdx(SelectableAssetIdx, SearchContext.GetInternalDeprecatedAssetsToConsider());
+		}
+		else
+		PRAGMA_ENABLE_DEPRECATION_WARNINGS
+		{
+			PopulateSelectableAssetIdx(SelectableAssetIdx, SearchContext.GetAssetsToConsiderSet());
+		}
 
 		// @todo: implement filtering within the VPTree as KDTree does
 		FNonSelectableIdx NonSelectableIdx;
@@ -2351,7 +2395,16 @@ void UPoseSearchDatabase::SearchBruteForce(UE::PoseSearch::FSearchContext& Searc
 		TConstArrayView<float> QueryValues = SearchContext.GetOrBuildQuery(Schema);
 
 		FSelectableAssetIdx SelectableAssetIdx;
-		PopulateSelectableAssetIdx(SelectableAssetIdx, SearchContext.GetAssetsToConsider());
+		PRAGMA_DISABLE_DEPRECATION_WARNINGS
+		if (!SearchContext.GetInternalDeprecatedAssetsToConsider().IsEmpty())
+		{
+			PopulateSelectableAssetIdx(SelectableAssetIdx, SearchContext.GetInternalDeprecatedAssetsToConsider());
+		}
+		else
+		PRAGMA_ENABLE_DEPRECATION_WARNINGS
+		{
+			PopulateSelectableAssetIdx(SelectableAssetIdx, SearchContext.GetAssetsToConsiderSet());
+		}
 
 		FNonSelectableIdx NonSelectableIdx;
 		PopulateNonSelectableIdx(NonSelectableIdx, SearchContext
@@ -2538,7 +2591,16 @@ void UPoseSearchDatabase::SearchEvent(UE::PoseSearch::FSearchContext& SearchCont
 			const TConstArrayView<float> DynamicWeightsSqrt = CalculateDynamicWeightsSqrt(DynamicWeightsSqrtBuffer);
 
 			FSelectableAssetIdx SelectableAssetIdx;
-			PopulateSelectableAssetIdx(SelectableAssetIdx, SearchContext.GetAssetsToConsider());
+			PRAGMA_DISABLE_DEPRECATION_WARNINGS
+			if (!SearchContext.GetInternalDeprecatedAssetsToConsider().IsEmpty())
+			{
+				PopulateSelectableAssetIdx(SelectableAssetIdx, SearchContext.GetInternalDeprecatedAssetsToConsider());
+			}
+			else
+			PRAGMA_ENABLE_DEPRECATION_WARNINGS
+			{
+				PopulateSelectableAssetIdx(SelectableAssetIdx, SearchContext.GetAssetsToConsiderSet());
+			}
 
 			FNonSelectableIdx NonSelectableIdx;
 			if (EventToSearch.bEnablePoseFilters)
diff --git a/Engine/Plugins/Animation/PoseSearch/Source/Runtime/Private/PoseSearchInteractionIsland.cpp b/Engine/Plugins/Animation/PoseSearch/Source/Runtime/Private/PoseSearchInteractionIsland.cpp
--- a/Engine/Plugins/Animation/PoseSearch/Source/Runtime/Private/PoseSearchInteractionIsland.cpp
+++ b/Engine/Plugins/Animation/PoseSearch/Source/Runtime/Private/PoseSearchInteractionIsland.cpp
@@ -143,9 +143,7 @@ static FInteractionSearchResult InitSearchResult(const FSearchResult& SearchResu
 	return InteractionSearchResult;
 }
 
-typedef TSet<const UObject*, DefaultKeyFuncs<const UObject*>, TInlineSetAllocator<64, TMemStackSetAllocator<>>> FVisitedAnimContexts;
-
-static bool IsPoseSearchResultUsable(int32 SearchIndex, TConstArrayView<FSearchResult> PoseSearchResults, TConstArrayView<FInteractionSearchContext> SearchContexts, const FVisitedAnimContexts& VisitedAnimContexts)
+static bool IsPoseSearchResultUsable(int32 SearchIndex, TConstArrayView<FSearchResult> PoseSearchResults, TConstArrayView<FInteractionSearchContext> SearchContexts, const FStackAssetSet& VisitedAnimContexts)
 {
 	if (!PoseSearchResults[SearchIndex].IsValid())
 	{
@@ -220,7 +218,7 @@ static void InitSearchResults(TArray<FInteractionSearchResult>& SearchResults, T
 			});
 
 		// assign from best to worst result
-		FVisitedAnimContexts VisitedAnimContexts;
+		FStackAssetSet VisitedAnimContexts;
 		for (int32 SearchIndex : SortedPoseSearchResults)
 		{
 			if (IsPoseSearchResultUsable(SearchIndex, PoseSearchResults, SearchContexts, VisitedAnimContexts))
diff --git a/Engine/Plugins/Animation/PoseSearch/Source/Runtime/Private/PoseSearchLibrary.cpp b/Engine/Plugins/Animation/PoseSearch/Source/Runtime/Private/PoseSearchLibrary.cpp
--- a/Engine/Plugins/Animation/PoseSearch/Source/Runtime/Private/PoseSearchLibrary.cpp
+++ b/Engine/Plugins/Animation/PoseSearch/Source/Runtime/Private/PoseSearchLibrary.cpp
@@ -42,9 +42,8 @@ static FAutoConsoleVariableRef CVarAnimMotionMatchDrawMatchEnable(TEXT("a.Motion
 
 namespace UE::PoseSearch
 {
-	// an empty FAssetsToSearch in any of the FAssetsToSearchPerDatabasePair entries means we need to search ALL the assets for the associated Database
-	typedef	TArray<const UObject*, TInlineAllocator<8, TMemStackAllocator<>>> FAssetsToSearch;
-	typedef TPair<const UPoseSearchDatabase*, FAssetsToSearch> FAssetsToSearchPerDatabasePair;
+	// an empty FStackAssetSet in any of the FAssetsToSearchPerDatabasePair entries means we need to search ALL the assets for the associated Database
+	typedef TPair<const UPoseSearchDatabase*, FStackAssetSet> FAssetsToSearchPerDatabasePair;
 	struct FAssetsToSearchPerDatabase : public TArray<FAssetsToSearchPerDatabasePair, TInlineAllocator<8, TMemStackAllocator<>>>
 	{
 		FAssetsToSearchPerDatabasePair* Find(const UPoseSearchDatabase* Database)
@@ -124,11 +123,11 @@ namespace UE::PoseSearch
 
 		if (FAssetsToSearchPerDatabasePair* AssetsToSearchPerDatabasePair = AssetsToSearchPerDatabase.Find(Database))
 		{
-			// an empty FAssetsToSearch associated to Database means we need to search ALL the assets, so we don't need to add this AssetToSearch
-			FAssetsToSearch& AssetsToSearch = AssetsToSearchPerDatabasePair->Value;
+			// an empty FStackAssetSet associated to Database means we need to search ALL the assets, so we don't need to add this AssetToSearch
+			FStackAssetSet& AssetsToSearch = AssetsToSearchPerDatabasePair->Value;
 			if (!AssetsToSearch.IsEmpty())
 			{
-				AssetsToSearch.AddUnique(AssetToSearch);
+				AssetsToSearch.Add(AssetToSearch);
 			}
 		}
 		else
@@ -186,13 +185,13 @@ namespace UE::PoseSearch
 			// we already added Database to AssetsToSearchPerDatabase, so it already successfully passed RequestAsyncBuildIndex
 			if (FAssetsToSearchPerDatabasePair* AssetsToSearchPerDatabasePair = AssetsToSearchPerDatabase.Find(Database))
 			{
-				// an empty FAssetsToSearch associated to Database means we need to search ALL the assets
-				FAssetsToSearch& AssetsToSearch = AssetsToSearchPerDatabasePair->Value;
+				// an empty FStackAssetSet associated to Database means we need to search ALL the assets
+				FStackAssetSet& AssetsToSearch = AssetsToSearchPerDatabasePair->Value;
 				AssetsToSearch.Reset();
 			}
 			else
 			{
-				// an empty FAssetsToSearch associated to Database means we need to search ALL the assets
+				// an empty FStackAssetSet associated to Database means we need to search ALL the assets
 				AssetsToSearchPerDatabase.AddDefaulted_GetRef().Key = Database;
 			}
 			return false;
@@ -207,13 +206,28 @@ namespace UE::PoseSearch
 		{
 			// checking if PlayingAnimationAsset can be considered or filtered out
 			bool bCanBeConsidered = true;
-			if (!SearchContext.GetAssetsToConsider().IsEmpty())
+			
+			// if !AssetsToConsider or AssetsToConsider->IsEmpty(), we need to consider all the assets!
+			PRAGMA_DISABLE_DEPRECATION_WARNINGS
+			if (!SearchContext.GetInternalDeprecatedAssetsToConsider().IsEmpty())
 			{
-				if (!SearchContext.GetAssetsToConsider().Contains(PlayingAnimationAsset))
+				// backward compatible path to support deprecated API FSearchContext::SetAssetsToConsider
+				if (!SearchContext.GetInternalDeprecatedAssetsToConsider().Contains(PlayingAnimationAsset))
 				{
 					bCanBeConsidered = false;
 				}
 			}
+			else
+			PRAGMA_ENABLE_DEPRECATION_WARNINGS
+			{
+				if (const FStackAssetSet* AssetsToConsider = SearchContext.GetAssetsToConsiderSet())
+				{
+					if (!AssetsToConsider->IsEmpty() && !AssetsToConsider->Contains(PlayingAnimationAsset))
+					{
+						bCanBeConsidered = false;
+					}
+				}
+			}
 
 			if (bCanBeConsidered)
 			{
@@ -266,22 +280,24 @@ namespace UE::PoseSearch
 		}
 
 		// intersecting AssetsToSearchPerDatabase with asset to consider
-		if (!SearchContext.GetAssetsToConsider().IsEmpty())
+		PRAGMA_DISABLE_DEPRECATION_WARNINGS
+		if (!SearchContext.GetInternalDeprecatedAssetsToConsider().IsEmpty())
 		{
+			// backward compatible path to support deprecated API FSearchContext::SetAssetsToConsider
 			for (FAssetsToSearchPerDatabasePair& AssetsToSearchPerDatabasePair : AssetsToSearchPerDatabase)
 			{
 				const UPoseSearchDatabase* Database = AssetsToSearchPerDatabasePair.Key;
 				check(Database);
 
-				FAssetsToSearch& AssetsToSearchForDatabase = AssetsToSearchPerDatabasePair.Value;
+				FStackAssetSet& AssetsToSearchForDatabase = AssetsToSearchPerDatabasePair.Value;
 				if (!AssetsToSearchForDatabase.IsEmpty())
 				{
 					// doing an intersection between AssetsToSearchForDatabase and asset to consider
-					FAssetsToSearch AssetsToSearchForDatabaseIntersection;
+					FStackAssetSet AssetsToSearchForDatabaseIntersection;
 					AssetsToSearchForDatabaseIntersection.Reserve(AssetsToSearchForDatabase.Num());
 					for (const UObject* AssetToSearchForDatabase : AssetsToSearchForDatabase)
 					{
-						if (SearchContext.GetAssetsToConsider().Contains(AssetToSearchForDatabase))
+						if (SearchContext.GetInternalDeprecatedAssetsToConsider().Contains(AssetToSearchForDatabase))
 						{
 							AssetsToSearchForDatabaseIntersection.Add(AssetToSearchForDatabase);
 						}
@@ -291,11 +307,51 @@ namespace UE::PoseSearch
 				else
 				{
 					// since all the database assets can be searched, we need to add only the one we can consider:
-					for (const UObject* AssetToConsider : SearchContext.GetAssetsToConsider())
+					for (const UObject* AssetToConsider : SearchContext.GetInternalDeprecatedAssetsToConsider())
 					{
 						if (Database->Contains(AssetToConsider))
 						{
-							AssetsToSearchForDatabase.AddUnique(AssetToConsider);
+							AssetsToSearchForDatabase.Add(AssetToConsider);
+						}
+					}
+				}
+			}
+		}
+		PRAGMA_ENABLE_DEPRECATION_WARNINGS
+		else if (const FStackAssetSet* AssetsToConsider = SearchContext.GetAssetsToConsiderSet())
+		{
+			if (!AssetsToConsider->IsEmpty())
+			{
+				for (FAssetsToSearchPerDatabasePair& AssetsToSearchPerDatabasePair : AssetsToSearchPerDatabase)
+				{
+					const UPoseSearchDatabase* Database = AssetsToSearchPerDatabasePair.Key;
+					check(Database);
+
+					FStackAssetSet& AssetsToSearchForDatabase = AssetsToSearchPerDatabasePair.Value;
+					if (!AssetsToSearchForDatabase.IsEmpty())
+					{
+						// doing an intersection between AssetsToSearchForDatabase and asset to consider
+						FStackAssetSet AssetsToSearchForDatabaseIntersection;
+						AssetsToSearchForDatabaseIntersection.Reserve(AssetsToSearchForDatabase.Num());
+						for (const UObject* AssetToSearchForDatabase : AssetsToSearchForDatabase)
+						{
+							if (AssetsToConsider->Contains(AssetToSearchForDatabase))
+							{
+								AssetsToSearchForDatabaseIntersection.Add(AssetToSearchForDatabase);
+							}
+						}
+						AssetsToSearchForDatabase = AssetsToSearchForDatabaseIntersection;
+					}
+					else
+					{
+						// since all the database assets can be searched, we need to add only the one we can consider:
+						// @todo: should we just do a copy? AssetsToSearchForDatabase = AssetsToConsider??
+						for (const UObject* AssetToConsider : *AssetsToConsider)
+						{
+							if (Database->Contains(AssetToConsider))
+							{
+								AssetsToSearchForDatabase.Add(AssetToConsider);
+							}
 						}
 					}
 				}
@@ -1378,7 +1434,13 @@ void UPoseSearchLibrary::MotionMatch(
 {
 	using namespace UE::PoseSearch;
 
-	TConstArrayView<const UObject*> CurrentAssetToConsider = SearchContext.GetAssetsToConsider();
+	const FStackAssetSet* CurrentAssetToConsider = nullptr;
+	PRAGMA_DISABLE_DEPRECATION_WARNINGS
+	if (SearchContext.GetInternalDeprecatedAssetsToConsider().IsEmpty())
+	PRAGMA_ENABLE_DEPRECATION_WARNINGS
+	{
+		CurrentAssetToConsider = SearchContext.GetAssetsToConsiderSet();
+	}
 	SearchContext.SetIsContinuingInteraction(ContinuingProperties.bIsContinuingInteraction);
 
 	// collecting all the databases searches in AssetsToSearchPerDatabase
@@ -1425,14 +1487,24 @@ void UPoseSearchLibrary::MotionMatch(
 		const UPoseSearchDatabase* Database = AssetsToSearchPerDatabasePair.Key;
 		check(Database);
 
+		PRAGMA_DISABLE_DEPRECATION_WARNINGS
+		TArray<const UObject*> InternalDeprecatedAssetsToConsider = SearchContext.GetInternalDeprecatedAssetsToConsider();
+		TArray<const UObject*> EmptyAssetsToConsider;
+		SearchContext.SetInternalDeprecatedAssetsToConsider(EmptyAssetsToConsider);
+		PRAGMA_ENABLE_DEPRECATION_WARNINGS
+
 		// in case we haven't searched the continuing pose for this Database, we haven't created and cached the query yet,
 		// but if we didn't invalidated the continuing pose (when IsInvalidatingContinuingPose is true), we still can reuse
 		// the updated FirstInstanceOfReconstructedContinuingPoseSearchResult data, and by calling UpdateContinuingPoseSearchResult we set the 
 		// SearchContext to be able to create a query for Database using the continuing pose data.
 		SearchContext.UpdateContinuingPoseSearchResult(CachedContinuingPoseSearchResults.FindOrDefault(Database), CachedContinuingPoseSearchResults.FindOrDefault(Database->Schema));
-		SearchContext.SetAssetsToConsider(AssetsToSearchPerDatabasePair.Value);
+		SearchContext.SetAssetsToConsiderSet(&AssetsToSearchPerDatabasePair.Value);
 		
 		Database->Search(SearchContext, SearchResults);
+		
+		PRAGMA_DISABLE_DEPRECATION_WARNINGS
+		SearchContext.SetInternalDeprecatedAssetsToConsider(InternalDeprecatedAssetsToConsider);
+		PRAGMA_ENABLE_DEPRECATION_WARNINGS
 	}
 
 #if (ENABLE_DRAW_DEBUG && ENABLE_ANIM_DEBUG) || UE_POSE_SEARCH_TRACE_ENABLED
@@ -1466,7 +1538,12 @@ void UPoseSearchLibrary::MotionMatch(
 #endif // UE_POSE_SEARCH_TRACE_ENABLED
 
 	// restoring the asset to consider array view
-	SearchContext.SetAssetsToConsider(CurrentAssetToConsider);
+	PRAGMA_DISABLE_DEPRECATION_WARNINGS
+	if (SearchContext.GetInternalDeprecatedAssetsToConsider().IsEmpty())
+	PRAGMA_ENABLE_DEPRECATION_WARNINGS
+	{
+		SearchContext.SetAssetsToConsiderSet(CurrentAssetToConsider);
+	}
 }
 
 const FAnimNode_PoseSearchHistoryCollector_Base* UPoseSearchLibrary::FindPoseHistoryNode(
diff --git a/Engine/Plugins/Animation/PoseSearch/Source/Runtime/Public/PoseSearch/PoseSearchContext.h b/Engine/Plugins/Animation/PoseSearch/Source/Runtime/Public/PoseSearch/PoseSearchContext.h
--- a/Engine/Plugins/Animation/PoseSearch/Source/Runtime/Public/PoseSearch/PoseSearchContext.h
+++ b/Engine/Plugins/Animation/PoseSearch/Source/Runtime/Public/PoseSearch/PoseSearchContext.h
@@ -20,6 +20,7 @@ class UPoseSearchFeatureChannel_Position;
 namespace UE::PoseSearch
 {
 	
+typedef TSet<const UObject*, DefaultKeyFuncs<const UObject*>, TInlineSetAllocator<8, TMemStackSetAllocator<>>> FStackAssetSet;
 
 // Utils functions
 const FTransform& GetContextTransform(const UObject* AnimContext);
@@ -288,8 +289,60 @@ struct FSearchContext
 
 	const FRoleToIndex& GetRoleToIndex() const { return RoleToIndex; }
 
-	void SetAssetsToConsider(TConstArrayView<const UObject*> InAssetsToConsider) { AssetsToConsider = InAssetsToConsider; }
-	TConstArrayView<const UObject*> GetAssetsToConsider() const { return AssetsToConsider; }
+	UE_DEPRECATED(5.7, "Use SetAssetsToConsiderSet instead")
+	void SetAssetsToConsider(TConstArrayView<const UObject*> InAssetsToConsider)
+	{
+		// you cannot call GetAssetsToConsiderSet / SetAssetsToConsiderSet with GetAssetsToConsider / SetAssetsToConsider
+		check(AssetsToConsider == nullptr);
+		PRAGMA_DISABLE_DEPRECATION_WARNINGS
+		AssetsToConsider_DEPRECATED = InAssetsToConsider;
+		PRAGMA_ENABLE_DEPRECATION_WARNINGS
+	}
+
+	UE_DEPRECATED(5.7, "Use GetAssetsToConsiderSet instead")
+	TConstArrayView<const UObject*> GetAssetsToConsider() const
+	{
+		// you cannot call GetAssetsToConsiderSet / SetAssetsToConsiderSet with GetAssetsToConsider / SetAssetsToConsider
+		check(AssetsToConsider == nullptr);
+		PRAGMA_DISABLE_DEPRECATION_WARNINGS
+		return AssetsToConsider_DEPRECATED;
+		PRAGMA_ENABLE_DEPRECATION_WARNINGS
+	}
+
+	UE_DEPRECATED(5.7, "Unnecessary once deprecated GetAssetsToConsiderSet / SetAssetsToConsiderSet have been removed")
+	const TArray<const UObject*>& GetInternalDeprecatedAssetsToConsider() const
+	{
+		PRAGMA_DISABLE_DEPRECATION_WARNINGS
+		return AssetsToConsider_DEPRECATED;
+		PRAGMA_ENABLE_DEPRECATION_WARNINGS
+	}
+
+	UE_DEPRECATED(5.7, "Unnecessary once deprecated GetAssetsToConsiderSet / SetAssetsToConsiderSet have been removed")
+	void SetInternalDeprecatedAssetsToConsider(TArray<const UObject*>& InAssetsToConsider)
+	{
+		PRAGMA_DISABLE_DEPRECATION_WARNINGS
+		AssetsToConsider_DEPRECATED = InAssetsToConsider;
+		PRAGMA_ENABLE_DEPRECATION_WARNINGS
+	}
+
+	void SetAssetsToConsiderSet(const FStackAssetSet* InAssetsToConsider)
+	{
+		PRAGMA_DISABLE_DEPRECATION_WARNINGS
+		// you cannot call GetAssetsToConsiderSet / SetAssetsToConsiderSet with GetAssetsToConsider / SetAssetsToConsider
+		check(AssetsToConsider_DEPRECATED.IsEmpty());
+		PRAGMA_ENABLE_DEPRECATION_WARNINGS
+		AssetsToConsider = InAssetsToConsider;
+	}
+
+	const FStackAssetSet* GetAssetsToConsiderSet() const
+	{
+		PRAGMA_DISABLE_DEPRECATION_WARNINGS
+		// you cannot call GetAssetsToConsiderSet / SetAssetsToConsiderSet with GetAssetsToConsider / SetAssetsToConsider
+		check(AssetsToConsider_DEPRECATED.IsEmpty());
+		PRAGMA_ENABLE_DEPRECATION_WARNINGS
+		return AssetsToConsider;
+	}
+
 	// Experimental, this feature might be removed without warning, not for production use
 	const FPoseSearchEvent& GetEventToSearch() const { return EventToSearch; }
 	
@@ -327,7 +380,11 @@ private:
 	FRoleToIndex RoleToIndex;
 	
 	// if AssetsToConsider is not empty, we'll search only for poses from UObject(s) that are in the AssetsToConsider
-	TConstArrayView<const UObject*> AssetsToConsider;
+	const FStackAssetSet* AssetsToConsider = nullptr;
+	
+	UE_DEPRECATED(5.7, "Unnecessary once deprecated GetAssetsToConsiderSet / SetAssetsToConsiderSet have been removed")
+	TArray<const UObject*> AssetsToConsider_DEPRECATED;
+	
 	const FPoseSearchEvent EventToSearch;
 
 	const float DesiredPermutationTimeOffset = 0.f;
diff --git a/Engine/Plugins/Animation/PoseSearch/Source/Runtime/Public/PoseSearch/PoseSearchDatabase.h b/Engine/Plugins/Animation/PoseSearch/Source/Runtime/Public/PoseSearch/PoseSearchDatabase.h
--- a/Engine/Plugins/Animation/PoseSearch/Source/Runtime/Public/PoseSearch/PoseSearchDatabase.h
+++ b/Engine/Plugins/Animation/PoseSearch/Source/Runtime/Public/PoseSearch/PoseSearchDatabase.h
@@ -3,12 +3,9 @@
 #pragma once
 
 #include "Engine/DataAsset.h"
-#include "StructUtils/InstancedStruct.h"
 #include "PoseSearch/PoseSearchAssetSampler.h"
-#include "PoseSearch/PoseSearchCost.h"
-#include "PoseSearch/PoseSearchIndex.h"
-#include "PoseSearch/PoseSearchResult.h"
-#include "PoseSearch/PoseSearchRole.h"
+#include "PoseSearch/PoseSearchContext.h"
+#include "StructUtils/InstancedStruct.h"
 #include "PoseSearchDatabase.generated.h"
 
 #define UE_API POSESEARCH_API
@@ -25,8 +22,6 @@ class UPoseSearchNormalizationSet;
 
 namespace UE::PoseSearch
 {
-	struct FSearchContext;
-
 	UE_DEPRECATED(5.7, "UpgradeToDatabaseAnimationAsset will be removed once all the classes inheriting from FPoseSearchDatabaseAnimationAssetBase are removed")
 	UE_API FPoseSearchDatabaseAnimationAsset UpgradeToDatabaseAnimationAsset(FInstancedStruct& AnimationAsset);
 } // namespace UE::PoseSearch
@@ -747,7 +742,9 @@ private:
 	void SearchEvent(UE::PoseSearch::FSearchContext& SearchContext, UE::PoseSearch::FSearchResults& SearchResults) const;
 
 	typedef TArray<int32, TInlineAllocator<256, TMemStackAllocator<>>> FSelectableAssetIdx;
+	UE_DEPRECATED(5.7, "Use PopulateSelectableAssetIdx taking FStackAssetSet as input instead")
 	void PopulateSelectableAssetIdx(FSelectableAssetIdx& SelectableAssetIdx, TConstArrayView<const UObject*> AssetsToConsider) const;
+	void PopulateSelectableAssetIdx(FSelectableAssetIdx& SelectableAssetIdx, const UE::PoseSearch::FStackAssetSet* AssetsToConsider) const;
 
 	typedef TArray<int32, TInlineAllocator<256, TMemStackAllocator<>>> FNonSelectableIdx;
 	void PopulateNonSelectableIdx(FNonSelectableIdx& NonSelectableIdx, UE::PoseSearch::FSearchContext& SearchContext
