PoseSearch - fix for crash in UPoseSearchInteractionAsset when unsupported UAnimationAsset derived types are used in items

diff --git a/Engine/Plugins/Animation/PoseSearch/Source/Runtime/Private/PoseSearchInteractionAsset.cpp b/Engine/Plugins/Animation/PoseSearch/Source/Runtime/Private/PoseSearchInteractionAsset.cpp
--- a/Engine/Plugins/Animation/PoseSearch/Source/Runtime/Private/PoseSearchInteractionAsset.cpp
+++ b/Engine/Plugins/Animation/PoseSearch/Source/Runtime/Private/PoseSearchInteractionAsset.cpp
@@ -31,7 +31,8 @@ bool UPoseSearchInteractionAsset::IsLooping() const
 			}
 			else
 			{
-				unimplemented();
+				UE_LOG(LogPoseSearch, Warning, TEXT("UPoseSearchInteractionAsset::IsLooping non fully supported UAnimationAsset derived type '%s' used as item in '%s'"), *AnimationAsset->GetClass()->GetName(), *GetName());
+				return false;
 			}
 
 			if (CommonPlayLength < 0.f)
@@ -72,7 +73,7 @@ bool UPoseSearchInteractionAsset::HasRootMotion() const
 			}
 			else
 			{
-				unimplemented();
+				UE_LOG(LogPoseSearch, Warning, TEXT("UPoseSearchInteractionAsset::HasRootMotion non fully supported UAnimationAsset derived type '%s' used as item in '%s'"), *AnimationAsset->GetClass()->GetName(), *GetName());
 			}
 			bHasAtLeastOneValidItem = true;
 		}
