PoseSearch - fix for UPoseSearchFeatureChannel_Phase::IndexAsset crashing while indexing assets with a single pose

diff --git a/Engine/Plugins/Animation/PoseSearch/Source/Runtime/Private/PoseSearchDerivedData.cpp b/Engine/Plugins/Animation/PoseSearch/Source/Runtime/Private/PoseSearchDerivedData.cpp
--- a/Engine/Plugins/Animation/PoseSearch/Source/Runtime/Private/PoseSearchDerivedData.cpp
+++ b/Engine/Plugins/Animation/PoseSearch/Source/Runtime/Private/PoseSearchDerivedData.cpp
@@ -1977,7 +1977,7 @@ void FPoseSearchDatabaseAsyncCacheTask::OnGetComplete(UE::DerivedData::FCacheGet
 
 					if (!IndexDatabase(SearchIndexBase, DependentDatabase, DependentDatabaseSchema, DependentSamplingContext, DependentAdditionalExtrapolationTime, Owner))
 					{
-						UE_LOG(LogPoseSearch, Log, TEXT("%s - %s BuildIndex Cancelled"), *LexToString(FullIndexKey.Hash), *MainDatabaseName);
+						UE_LOG(LogPoseSearch, Log, TEXT("%s - %s BuildIndex Failed in IndexDatabase"), *LexToString(FullIndexKey.Hash), *MainDatabaseName);
 						ResetSearchIndex();
 						return;
 					}
diff --git a/Engine/Plugins/Animation/PoseSearch/Source/Runtime/Private/PoseSearchFeatureChannel_Phase.cpp b/Engine/Plugins/Animation/PoseSearch/Source/Runtime/Private/PoseSearchFeatureChannel_Phase.cpp
--- a/Engine/Plugins/Animation/PoseSearch/Source/Runtime/Private/PoseSearchFeatureChannel_Phase.cpp
+++ b/Engine/Plugins/Animation/PoseSearch/Source/Runtime/Private/PoseSearchFeatureChannel_Phase.cpp
@@ -398,6 +398,13 @@ bool UPoseSearchFeatureChannel_Phase::IndexAsset(UE::PoseSearch::FAssetIndexer&
 {
 	using namespace UE::PoseSearch;
 
+	if (Indexer.GetNumIndexedPoses() <= 1)
+	{
+		// We must have at least 2 poses in order to index phase properly.
+		UE_LOG(LogPoseSearch, Error, TEXT("UPoseSearchFeatureChannel_Phase::IndexAsset - Attempt to index animation asset with single pose is not supported. Asset name: %s"), *Indexer.GetAssetName());
+		return false;
+	}
+
 	// @todo: do we really need to use double(s) in all this math?
 	// @todo: expose them via UI
 	static float BoneSamplingCentralDifferencesTime = 0.2f; // seconds
