UPoseSearchDatabase::SynchronizeWithExternalDependencies - only enumerate on-disk assets when cooking, to ensure deterministic behaviour

diff --git a/Engine/Plugins/Animation/PoseSearch/Source/Runtime/Private/PoseSearchDatabase.cpp b/Engine/Plugins/Animation/PoseSearch/Source/Runtime/Private/PoseSearchDatabase.cpp
--- a/Engine/Plugins/Animation/PoseSearch/Source/Runtime/Private/PoseSearchDatabase.cpp
+++ b/Engine/Plugins/Animation/PoseSearch/Source/Runtime/Private/PoseSearchDatabase.cpp
@@ -1331,7 +1331,8 @@ void UPoseSearchDatabase::SynchronizeWithExternalDependencies()
 	for (const FAssetIdentifier& Referencer : Referencers)
 	{
 		TArray<FAssetData> Assets;
-		AssetRegistry.GetAssetsByPackageName(Referencer.PackageName, Assets);
+		const bool bSearchOnDiskAssetsOnly = IsRunningCookCommandlet();		// When cooking, only enumerate on-disk assets to ensure deterministic results
+		AssetRegistry.GetAssetsByPackageName(Referencer.PackageName, Assets, bSearchOnDiskAssetsOnly);
 
 		for (const FAssetData& Asset : Assets)
 		{
