PoseSearch - avoid opening the schema asset selector window for every row in the chooser table

diff --git a/Engine/Plugins/Animation/PoseSearch/Source/Editor/Private/Chooser/PoseSearchColumnEditor.cpp b/Engine/Plugins/Animation/PoseSearch/Source/Editor/Private/Chooser/PoseSearchColumnEditor.cpp
--- a/Engine/Plugins/Animation/PoseSearch/Source/Editor/Private/Chooser/PoseSearchColumnEditor.cpp
+++ b/Engine/Plugins/Animation/PoseSearch/Source/Editor/Private/Chooser/PoseSearchColumnEditor.cpp
@@ -24,51 +24,52 @@ TSharedRef<SWidget> CreatePoseSearchColumnWidget(UChooserTable* Chooser, FChoose
 {
 	FPoseSearchColumn* PoseSearchColumn = static_cast<FPoseSearchColumn*>(Column);
 	
-	if (!PoseSearchColumn->InternalDatabase->Schema)
+	if (Row == UE::ChooserEditor::ColumnWidget_SpecialIndex_Fallback)
+	{
+		return SNullWidget::NullWidget;
+	}
+	else if (Row == UE::ChooserEditor::ColumnWidget_SpecialIndex_Header)
 	{
-		TSharedPtr<SWindow> PickerWindow;
+		// @todo: avoid reopening the asset picker after tyhe first user choice, or prevent the FPoseSearchColumn creation if no schemas has been selected!
+		if (!PoseSearchColumn->InternalDatabase->Schema)
+		{
+			TSharedPtr<SWindow> PickerWindow;
 
-		// asset picker config to select Schema
-		FAssetPickerConfig AssetPickerConfig;
-		AssetPickerConfig.Filter.ClassPaths.Add(UPoseSearchSchema::StaticClass()->GetClassPathName());
-		AssetPickerConfig.Filter.bRecursiveClasses = true;
-		AssetPickerConfig.InitialAssetViewType = EAssetViewType::List;
-		AssetPickerConfig.OnAssetSelected = FOnAssetSelected::CreateLambda([PoseSearchColumn, &PickerWindow](const FAssetData& SelectedAsset)
-			{
-				if (UPoseSearchSchema* Schema = Cast<UPoseSearchSchema>(SelectedAsset.GetAsset()))
+			// asset picker config to select Schema
+			FAssetPickerConfig AssetPickerConfig;
+			AssetPickerConfig.Filter.ClassPaths.Add(UPoseSearchSchema::StaticClass()->GetClassPathName());
+			AssetPickerConfig.Filter.bRecursiveClasses = true;
+			AssetPickerConfig.InitialAssetViewType = EAssetViewType::List;
+			AssetPickerConfig.OnAssetSelected = FOnAssetSelected::CreateLambda([PoseSearchColumn, &PickerWindow](const FAssetData& SelectedAsset)
 				{
-					FScopedTransaction ScopedTransaction(LOCTEXT("AssignPoseSearchColumnDatabaseSchema", "Assign Pose Search Column Database Schema"));
-					PoseSearchColumn->InternalDatabase->Modify();
-					PoseSearchColumn->InternalDatabase->Schema = Schema;
-				}
-				PickerWindow->RequestDestroyWindow();
-			});
+					if (UPoseSearchSchema* Schema = Cast<UPoseSearchSchema>(SelectedAsset.GetAsset()))
+					{
+						FScopedTransaction ScopedTransaction(LOCTEXT("AssignPoseSearchColumnDatabaseSchema", "Assign Pose Search Column Database Schema"));
+						PoseSearchColumn->InternalDatabase->Modify();
+						PoseSearchColumn->InternalDatabase->Schema = Schema;
+					}
+					PickerWindow->RequestDestroyWindow();
+				});
 
-		// Load the content browser module to display an asset picker
-		FContentBrowserModule& ContentBrowserModule = FModuleManager::LoadModuleChecked<FContentBrowserModule>("ContentBrowser");
-		PickerWindow = SNew(SWindow)
-			.Title(LOCTEXT("SelectPoseSearchColumnDatabaseSchema", "Pick Schema for owned Pose Search Database"))
-			.ClientSize(FVector2D(500, 600))
-			.SupportsMinimize(false)
-			.SupportsMaximize(false)
-			[
-				SNew(SBorder)
-					.BorderImage(FAppStyle::GetBrush("Menu.Background"))
-					[
-						ContentBrowserModule.Get().CreateAssetPicker(AssetPickerConfig)
-					]
-			];
+			// Load the content browser module to display an asset picker
+			FContentBrowserModule& ContentBrowserModule = FModuleManager::LoadModuleChecked<FContentBrowserModule>("ContentBrowser");
+			PickerWindow = SNew(SWindow)
+				.Title(LOCTEXT("SelectPoseSearchColumnDatabaseSchema", "Pick Schema for owned Pose Search Database"))
+				.ClientSize(FVector2D(500, 600))
+				.SupportsMinimize(false)
+				.SupportsMaximize(false)
+				[
+					SNew(SBorder)
+						.BorderImage(FAppStyle::GetBrush("Menu.Background"))
+						[
+							ContentBrowserModule.Get().CreateAssetPicker(AssetPickerConfig)
+						]
+				];
 		
-		GEditor->EditorAddModalWindow(PickerWindow.ToSharedRef());
-		PickerWindow.Reset();
-	}
+			GEditor->EditorAddModalWindow(PickerWindow.ToSharedRef());
+			PickerWindow.Reset();
+		}
 
-	if (Row == UE::ChooserEditor::ColumnWidget_SpecialIndex_Fallback)
-	{
-		return SNullWidget::NullWidget;
-	}
-	if (Row == UE::ChooserEditor::ColumnWidget_SpecialIndex_Header)
-	{
 		// create column header widget
 		const FSlateBrush* ColumnIcon = FAppStyle::Get().GetBrush("Icons.Search");
 		const FText ColumnTooltip = LOCTEXT("Pose Match Tooltip", "Pose Match: Selects a single result based on the animation with the best matching pose, and outputs the StartTime for the frame with that pose. Animation Assets must contain \"Pose Match Branch In\" Notify State. AutoPopulate will fill in Column data with result Animation Assets.");
