PoseSearch - removing unnecessary transformation from quaternion to rotator and vice versa

diff --git a/Engine/Plugins/Animation/PoseSearch/Source/Runtime/Private/PoseSearchTrajectoryLibrary.cpp b/Engine/Plugins/Animation/PoseSearch/Source/Runtime/Private/PoseSearchTrajectoryLibrary.cpp
--- a/Engine/Plugins/Animation/PoseSearch/Source/Runtime/Private/PoseSearchTrajectoryLibrary.cpp
+++ b/Engine/Plugins/Animation/PoseSearch/Source/Runtime/Private/PoseSearchTrajectoryLibrary.cpp
@@ -95,7 +95,7 @@ bool FPoseSearchTrajectoryData::UpdateData(
 		
 	if (TrajectoryDataDerived.bOrientRotationToMovement)
 	{
-		TrajectoryDataDerived.Facing = MeshComp->GetComponentRotation().Quaternion();
+		TrajectoryDataDerived.Facing = MeshComp->GetComponentTransform().GetRotation();
 	}
 	else
 	{
