Chooser - PoseSearchColumn initial support for fallback row

diff --git a/Engine/Plugins/Animation/PoseSearch/Source/Editor/Private/Chooser/PoseSearchColumnEditor.cpp b/Engine/Plugins/Animation/PoseSearch/Source/Editor/Private/Chooser/PoseSearchColumnEditor.cpp
--- a/Engine/Plugins/Animation/PoseSearch/Source/Editor/Private/Chooser/PoseSearchColumnEditor.cpp
+++ b/Engine/Plugins/Animation/PoseSearch/Source/Editor/Private/Chooser/PoseSearchColumnEditor.cpp
@@ -20,18 +20,18 @@
 namespace UE::PoseSearchEditor
 {
 
-TSharedRef<SWidget> CreatePoseSearchColumnWidget(UChooserTable* Chooser, FChooserColumnBase* Column, int Row)
+TSharedRef<SWidget> CreatePoseSearchColumnWidget(UChooserTable* Chooser, FChooserColumnBase* Column, int RowIndex)
 {
 	FPoseSearchColumn* PoseSearchColumn = static_cast<FPoseSearchColumn*>(Column);
 	
-	if (Row == UE::ChooserEditor::ColumnWidget_SpecialIndex_Fallback)
+	if (RowIndex == UE::ChooserEditor::ColumnWidget_SpecialIndex_Fallback)
 	{
 		return SNullWidget::NullWidget;
 	}
-	else if (Row == UE::ChooserEditor::ColumnWidget_SpecialIndex_Header)
+	else if (RowIndex == UE::ChooserEditor::ColumnWidget_SpecialIndex_Header)
 	{
 		// @todo: avoid reopening the asset picker after tyhe first user choice, or prevent the FPoseSearchColumn creation if no schemas has been selected!
-		if (!PoseSearchColumn->InternalDatabase->Schema)
+		if (!PoseSearchColumn->GetDatabaseSchema())
 		{
 			TSharedPtr<SWindow> PickerWindow;
 
@@ -45,8 +45,7 @@ TSharedRef<SWidget> CreatePoseSearchColumnWidget(UChooserTable* Chooser, FChoose
 					if (UPoseSearchSchema* Schema = Cast<UPoseSearchSchema>(SelectedAsset.GetAsset()))
 					{
 						FScopedTransaction ScopedTransaction(LOCTEXT("AssignPoseSearchColumnDatabaseSchema", "Assign Pose Search Column Database Schema"));
-						PoseSearchColumn->InternalDatabase->Modify();
-						PoseSearchColumn->InternalDatabase->Schema = Schema;
+						PoseSearchColumn->SetDatabaseSchema(Schema);
 					}
 					PickerWindow->RequestDestroyWindow();
 				});
@@ -83,19 +82,16 @@ TSharedRef<SWidget> CreatePoseSearchColumnWidget(UChooserTable* Chooser, FChoose
 	// create cell widget
 	return
 		SNew(STextBlock)
-			.Text_Lambda([Row, PoseSearchColumn]()
+			.Text_Lambda([RowIndex, PoseSearchColumn]()
 			{
-				if (FPoseSearchDatabaseAnimationAsset* Asset = PoseSearchColumn->InternalDatabase->GetMutableDatabaseAnimationAsset(Row))
+				if (const UObject* Asset = PoseSearchColumn->GetDatabaseAsset(RowIndex))
 				{
-					if (UObject* ResultAsset = Asset->GetAnimationAsset())
-					{
-						return FText::FromString(ResultAsset->GetName());
-					}
+					return FText::FromString(Asset->GetName());
 				}
 				return FText();
 			});
 }
-	
+
 void RegisterPoseSearchChooserWidgets()
 {
 	UE::ChooserEditor::FObjectChooserWidgetFactories::RegisterColumnWidgetCreator(FPoseSearchColumn::StaticStruct(), CreatePoseSearchColumnWidget);
diff --git a/Engine/Plugins/Animation/PoseSearch/Source/Runtime/Private/Chooser/PoseSearchChooserColumn.cpp b/Engine/Plugins/Animation/PoseSearch/Source/Runtime/Private/Chooser/PoseSearchChooserColumn.cpp
--- a/Engine/Plugins/Animation/PoseSearch/Source/Runtime/Private/Chooser/PoseSearchChooserColumn.cpp
+++ b/Engine/Plugins/Animation/PoseSearch/Source/Runtime/Private/Chooser/PoseSearchChooserColumn.cpp
@@ -24,11 +24,49 @@
 
 namespace UE::PoseSearch
 {
+	int32 GetDatabaseAssetIndex(int32 RowIndex)
+	{
+		if (RowIndex == ChooserColumn_SpecialIndex_Fallback)
+		{
+			return 0;
+		}
+
+		check(RowIndex >= 0);
+
+		return RowIndex + 1;
+	}
+
+	int32 GetRowIndex(int32 DatabaseAssetIndex)
+	{
+		check(DatabaseAssetIndex >= 0);
+
+		if (DatabaseAssetIndex == 0)
+		{
+			return ChooserColumn_SpecialIndex_Fallback;
+		}
+		
+		return DatabaseAssetIndex - 1;
+	}
+
+	// searching for FChooserPlayerSettings in the Context.Params
+	static const FChooserPlayerSettings* GetChooserPlayerSettings(const FChooserEvaluationContext& Context)
+	{
+		for (const FStructView& Param : Context.Params)
+		{
+			if (const FChooserPlayerSettings* ChooserPlayerSettings = Param.GetPtr<FChooserPlayerSettings>())
+			{
+				return ChooserPlayerSettings;
+			}
+		}
+		return nullptr;
+	}
+
 	// Experimental, this feature might be removed without warning, not for production use
 	struct FPoseSearchColumnScratchArea
 	{
 		TArray<UE::PoseSearch::FSearchResult> SearchResults;
 		const UE::PoseSearch::IPoseHistory* PoseHistory = nullptr;
+		const FChooserPlayerSettings* ChooserPlayerSettings = nullptr;
 
 #if DO_CHECK
 		const FPoseSearchColumn* DebugOwner = nullptr;
@@ -56,20 +94,23 @@ void FPoseSearchColumn::CheckConsistency() const
 {
 #if DO_CHECK && WITH_EDITORONLY_DATA
 
+	using namespace UE::PoseSearch;
+
 	check(InternalDatabase);
 	UChooserTable* OuterChooser = Cast<UChooserTable>(InternalDatabase->GetOuter());
 	check(OuterChooser);
 
 	const int32 NumAnimationAssets = InternalDatabase->GetNumAnimationAssets();
-	check(NumAnimationAssets == OuterChooser->ResultsStructs.Num());
+	const int32 NumRows = GetRowIndex(NumAnimationAssets);
+	check(NumRows == OuterChooser->ResultsStructs.Num());
 
 	for (int32 AnimationAssetIndex = 0; AnimationAssetIndex < NumAnimationAssets; ++AnimationAssetIndex)
 	{
 		const FPoseSearchDatabaseAnimationAsset* DatabaseAnimationAsset = InternalDatabase->GetDatabaseAnimationAsset(AnimationAssetIndex);
 		check(DatabaseAnimationAsset);
-
-		const UObject* ReferencedObject = OuterChooser->ResultsStructs[AnimationAssetIndex].Get<FObjectChooserBase>().GetReferencedObject();
+		
 		const UObject* AnimationAsset = DatabaseAnimationAsset->GetAnimationAsset();
+		const UObject* ReferencedObject = GetReferencedObject(GetRowIndex(AnimationAssetIndex));
 		check(AnimationAsset == ReferencedObject);
 	}
 
@@ -80,9 +121,25 @@ void FPoseSearchColumn::CheckConsistency() const
 
 UObject* FPoseSearchColumn::GetReferencedObject(int32 RowIndex) const
 {
+	using namespace UE::PoseSearch;
+
 	UChooserTable* OuterChooser = Cast<UChooserTable>(InternalDatabase->GetOuter());
 	check(OuterChooser);
-	return OuterChooser->ResultsStructs[RowIndex].Get<FObjectChooserBase>().GetReferencedObject();
+
+	const UObject* ReferencedObject = nullptr;
+
+	if (RowIndex != ChooserColumn_SpecialIndex_Fallback)
+	{
+		check(RowIndex >= 0);
+		return OuterChooser->ResultsStructs[RowIndex].Get<FObjectChooserBase>().GetReferencedObject();
+	}
+
+	if (OuterChooser->FallbackResult.IsValid())
+	{
+		return OuterChooser->FallbackResult.Get<FObjectChooserBase>().GetReferencedObject();
+	}
+
+	return nullptr;
 }
 
 FName FPoseSearchColumn::RowValuesPropertyName()
@@ -92,16 +149,19 @@ FName FPoseSearchColumn::RowValuesPropertyName()
 
 void FPoseSearchColumn::SetNumRows(int32 NumRows)
 {
-	check(InternalDatabase);
-	
-	if (NumRows == InternalDatabase->GetNumAnimationAssets())
+	using namespace UE::PoseSearch;
+
+	check(InternalDatabase && NumRows >= 0);
+
+	const int32 NumDatabaseEntries = GetDatabaseAssetIndex(NumRows);
+	if (NumDatabaseEntries == InternalDatabase->GetNumAnimationAssets())
 	{
 		// nothing to do
 	}
-	else if (NumRows < InternalDatabase->GetNumAnimationAssets())
+	else if (NumDatabaseEntries < InternalDatabase->GetNumAnimationAssets())
 	{
 		InternalDatabase->Modify();
-		while (InternalDatabase->GetNumAnimationAssets() > 0 && InternalDatabase->GetNumAnimationAssets() > NumRows)
+		while (InternalDatabase->GetNumAnimationAssets() > 0 && InternalDatabase->GetNumAnimationAssets() > NumDatabaseEntries)
 		{
 			InternalDatabase->RemoveAnimationAssetAt(InternalDatabase->GetNumAnimationAssets() - 1);
 		}
@@ -109,10 +169,11 @@ void FPoseSearchColumn::SetNumRows(int32 NumRows)
 	else
 	{
 		InternalDatabase->Modify();
-		while (InternalDatabase->GetNumAnimationAssets() < NumRows)
+		while (InternalDatabase->GetNumAnimationAssets() < NumDatabaseEntries)
 		{
 			FPoseSearchDatabaseAnimationAsset NewAsset;
-			NewAsset.AnimAsset = GetReferencedObject(InternalDatabase->GetNumAnimationAssets());
+			const int32 RowIndex = GetRowIndex(InternalDatabase->GetNumAnimationAssets());
+			NewAsset.AnimAsset = GetReferencedObject(RowIndex);
 			InternalDatabase->AddAnimationAsset(NewAsset);
 		}
 	}
@@ -128,6 +189,8 @@ void FPoseSearchColumn::SetNumRows(int32 NumRows)
 
 void FPoseSearchColumn::InsertRows(int Index, int Count)
 {
+	using namespace UE::PoseSearch;
+
 	check(InternalDatabase);
 	if (Count > 0)
 	{
@@ -135,8 +198,10 @@ void FPoseSearchColumn::InsertRows(int Index, int Count)
 		for (int i = 0; i < Count; i++)
 		{
 			FPoseSearchDatabaseAnimationAsset NewAsset;
-			NewAsset.AnimAsset = GetReferencedObject(Index + i);
-			InternalDatabase->InsertAnimationAssetAt(NewAsset, Index + i);
+			const int32 RowIndex = Index + i;
+			const int32 DatabaseAssetIndex = GetDatabaseAssetIndex(RowIndex);
+			NewAsset.AnimAsset = GetReferencedObject(RowIndex);
+			InternalDatabase->InsertAnimationAssetAt(NewAsset, DatabaseAssetIndex);
 		}
 	}
 
@@ -146,11 +211,21 @@ void FPoseSearchColumn::InsertRows(int Index, int Count)
 
 void FPoseSearchColumn::DeleteRows(const TArray<uint32>& RowIndices)
 {
+	using namespace UE::PoseSearch;
+
 	check(InternalDatabase);
 	InternalDatabase->Modify();
-	for(uint32 RowIndex : RowIndices)
+	
+	for (uint32 RowIndex : RowIndices)
 	{
-		InternalDatabase->RemoveAnimationAssetAt(RowIndex);
+		const int32 DatabaseAssetIndex = GetDatabaseAssetIndex(RowIndex);
+		InternalDatabase->RemoveAnimationAssetAt(DatabaseAssetIndex);
+	}
+
+	// resynchronize the fallback row asset
+	if (FPoseSearchDatabaseAnimationAsset* FallbackDatabaseAnimationAsset = InternalDatabase->GetMutableDatabaseAnimationAsset(GetDatabaseAssetIndex(ChooserColumn_SpecialIndex_Fallback)))
+	{
+		FallbackDatabaseAnimationAsset->AnimAsset = GetReferencedObject(ChooserColumn_SpecialIndex_Fallback);
 	}
 
 	InternalDatabase->NotifySynchronizeWithExternalDependencies();
@@ -159,19 +234,25 @@ void FPoseSearchColumn::DeleteRows(const TArray<uint32>& RowIndices)
 
 void FPoseSearchColumn::MoveRow(int SourceRowIndex, int TargetRowIndex)
 {
+	using namespace UE::PoseSearch;
+
 	check(InternalDatabase);
-	if (const FPoseSearchDatabaseAnimationAsset* SourceAsset = InternalDatabase->GetDatabaseAnimationAsset(SourceRowIndex))
+
+	const int32 SourceDatabaseAssetIndex = GetDatabaseAssetIndex(SourceRowIndex);
+	int32 TargetDatabaseAssetIndex = GetDatabaseAssetIndex(TargetRowIndex);
+
+	if (const FPoseSearchDatabaseAnimationAsset* SourceAsset = InternalDatabase->GetDatabaseAnimationAsset(SourceDatabaseAssetIndex))
 	{
 		InternalDatabase->Modify();
 
 		FPoseSearchDatabaseAnimationAsset CopySourceAsset = *SourceAsset;
-		InternalDatabase->RemoveAnimationAssetAt(SourceRowIndex);
+		InternalDatabase->RemoveAnimationAssetAt(SourceDatabaseAssetIndex);
 
-		if (SourceRowIndex < TargetRowIndex)
+		if (SourceDatabaseAssetIndex < TargetDatabaseAssetIndex)
 		{
-			TargetRowIndex--;
+			TargetDatabaseAssetIndex--;
 		}
-		InternalDatabase->InsertAnimationAssetAt(CopySourceAsset, TargetRowIndex);
+		InternalDatabase->InsertAnimationAssetAt(CopySourceAsset, TargetDatabaseAssetIndex);
 
 		InternalDatabase->NotifySynchronizeWithExternalDependencies();
 	}
@@ -181,13 +262,17 @@ void FPoseSearchColumn::MoveRow(int SourceRowIndex, int TargetRowIndex)
 
 void FPoseSearchColumn::CopyRow(FChooserColumnBase& SourceColumn, int SourceRowIndex, int TargetRowIndex)
 {
+	using namespace UE::PoseSearch;
+
 	check(InternalDatabase);
 
 	const FPoseSearchColumn& SourcePoseSearchColumn = static_cast<FPoseSearchColumn&>(SourceColumn);
 
-	if (const FPoseSearchDatabaseAnimationAsset* SourceAsset = SourcePoseSearchColumn.InternalDatabase->GetDatabaseAnimationAsset(SourceRowIndex))
+	const int32 SourceDatabaseAssetIndex = GetDatabaseAssetIndex(SourceRowIndex);
+	const int32 TargetDatabaseAssetIndex = GetDatabaseAssetIndex(TargetRowIndex);
+	if (const FPoseSearchDatabaseAnimationAsset* SourceAsset = SourcePoseSearchColumn.InternalDatabase->GetDatabaseAnimationAsset(SourceDatabaseAssetIndex))
 	{
-		if (FPoseSearchDatabaseAnimationAsset* TargetAsset = InternalDatabase->GetMutableDatabaseAnimationAsset(TargetRowIndex))
+		if (FPoseSearchDatabaseAnimationAsset* TargetAsset = InternalDatabase->GetMutableDatabaseAnimationAsset(TargetDatabaseAssetIndex))
 		{
 			InternalDatabase->Modify();
 			*TargetAsset = *SourceAsset;
@@ -216,7 +301,10 @@ void FPoseSearchColumn::SetInputType(const UScriptStruct* Type)
 
 void FPoseSearchColumn::AutoPopulate(int32 RowIndex, UObject* OutputObject)
 {
-	if (FPoseSearchDatabaseAnimationAsset* Asset = InternalDatabase->GetMutableDatabaseAnimationAsset(RowIndex))
+	using namespace UE::PoseSearch;
+
+	const int32 DatabaseAssetIndex = GetDatabaseAssetIndex(RowIndex);
+	if (FPoseSearchDatabaseAnimationAsset* Asset = InternalDatabase->GetMutableDatabaseAnimationAsset(DatabaseAssetIndex))
 	{
 		if (Asset->AnimAsset != OutputObject)
 		{
@@ -282,6 +370,31 @@ FChooserParameterBase* FPoseSearchColumn::GetInputValue()
 	return &NullInputValue;
 }
 
+const UObject* FPoseSearchColumn::GetDatabaseAsset(int32 RowIndex) const
+{
+	using namespace UE::PoseSearch;
+	const int32 DatabaseAssetIndex = GetDatabaseAssetIndex(RowIndex);
+
+	if (const FPoseSearchDatabaseAnimationAsset* Asset = InternalDatabase->GetDatabaseAnimationAsset(DatabaseAssetIndex))
+	{
+		return Asset->GetAnimationAsset();
+	}
+	return nullptr;
+}
+
+const UPoseSearchSchema* FPoseSearchColumn::GetDatabaseSchema() const
+{
+	check(InternalDatabase);
+	return InternalDatabase->Schema;
+}
+	
+void FPoseSearchColumn::SetDatabaseSchema(const UPoseSearchSchema* Schema)
+{
+	check(InternalDatabase);
+	InternalDatabase->Modify();
+	InternalDatabase->Schema = Schema;
+}
+
 bool FPoseSearchColumn::HasFilters() const
 {
 	return true;
@@ -300,10 +413,6 @@ void FPoseSearchColumn::Filter(FChooserEvaluationContext& Context, const FChoose
 	check(PoseSearchColumnScratchArea->SearchResults.IsEmpty());
 #endif // DO_CHECK
 
-#if WITH_EDITOR
-	TArray<FActiveColumnCost> TracedActiveColumnCosts;
-#endif // WITH_EDITOR
-
 	const UE::PoseSearch::IPoseHistory* PoseHistory = nullptr;
 	FPoseHistoryReference PoseHistoryReference;
 	if (const FChooserParameterPoseHistoryBase* PoseHistoryParameter = InputValue.GetPtr<FChooserParameterPoseHistoryBase>())
@@ -314,25 +423,19 @@ void FPoseSearchColumn::Filter(FChooserEvaluationContext& Context, const FChoose
 		}
 	}
 
-	// searching for FChooserPlayerSettings in the Context.Params
-	const FChooserPlayerSettings* ChooserPlayerSettings = nullptr;
-	for (const FStructView& Param : Context.Params)
-	{
-		ChooserPlayerSettings = Param.GetPtr<FChooserPlayerSettings>();
-		if (ChooserPlayerSettings)
-		{
-			break;
-		}
-	}
-
-	if (!PoseHistory && ChooserPlayerSettings && ChooserPlayerSettings->AnimationUpdateContext)
+	PoseSearchColumnScratchArea->ChooserPlayerSettings = GetChooserPlayerSettings(Context);
+	if (!PoseHistory && PoseSearchColumnScratchArea->ChooserPlayerSettings && PoseSearchColumnScratchArea->ChooserPlayerSettings->AnimationUpdateContext)
 	{
-		if (const FPoseHistoryProvider* PoseHistoryProvider = ChooserPlayerSettings->AnimationUpdateContext->GetMessage<FPoseHistoryProvider>())
+		if (const FPoseHistoryProvider* PoseHistoryProvider = PoseSearchColumnScratchArea->ChooserPlayerSettings->AnimationUpdateContext->GetMessage<FPoseHistoryProvider>())
 		{
 			PoseHistory = &PoseHistoryProvider->GetPoseHistory();
 		}
 	}
 
+#if WITH_EDITOR
+	TArray<FActiveColumnCost> TracedActiveColumnCosts;
+#endif // WITH_EDITOR
+
 	if (!PoseHistory)
 	{
 		UE_LOG(LogPoseSearch, Error, TEXT("FPoseSearchColumn::Filter, missing IPoseHistory"));
@@ -344,94 +447,98 @@ void FPoseSearchColumn::Filter(FChooserEvaluationContext& Context, const FChoose
 		FStackAssetSet AssetsToConsider;
 		for (const FChooserIndexArray::FIndexData& IndexData : IndexListIn)
 		{
-			const FPoseSearchDatabaseAnimationAsset* Asset = InternalDatabase->GetDatabaseAnimationAsset(IndexData.Index);
+			const int32 RowIndex = IndexData.Index;
+			const int32 DatabaseAssetIndex = GetDatabaseAssetIndex(RowIndex);
+			const FPoseSearchDatabaseAnimationAsset* Asset = InternalDatabase->GetDatabaseAnimationAsset(DatabaseAssetIndex);
 			if (ensure(Asset))
 			{
 				const UObject* AnimationAsset = Asset->GetAnimationAsset();
 #if DO_CHECK && WITH_EDITOR
 				// making sure InternalDatabase and the OuterChooser are still in synch
-				const UObject* ReferencedObject = GetReferencedObject(IndexData.Index);
+				const UObject* ReferencedObject = GetReferencedObject(RowIndex);
 				check(ReferencedObject == AnimationAsset);
 #endif // DO_CHECK && WITH_EDITOR
 				AssetsToConsider.Add(AnimationAsset);
 			}
 		}
 
-		const FFloatInterval PoseJumpThresholdTime(0.f, 0.f);
-		const UObject* DatabasesToSearch[] = { InternalDatabase.Get() };
-		FSearchContext SearchContext(0.f, PoseJumpThresholdTime, FPoseSearchEvent());
-
-		const FRole Role = InternalDatabase->Schema ? InternalDatabase->Schema->GetDefaultRole() : DefaultRole;
-		SearchContext.AddRole(Role, &Context, PoseHistory);
-		SearchContext.SetAssetsToConsiderSet(&AssetsToConsider);
-
-		FPoseSearchContinuingProperties ContinuingProperties;
-		if (ChooserPlayerSettings)
+		if (!AssetsToConsider.IsEmpty())
 		{
-			ContinuingProperties.PlayingAsset = ChooserPlayerSettings->PlayingAsset;
-			ContinuingProperties.PlayingAssetAccumulatedTime = ChooserPlayerSettings->PlayingAssetAccumulatedTime;
-			ContinuingProperties.bIsPlayingAssetMirrored = ChooserPlayerSettings->bIsPlayingAssetMirrored;
-			ContinuingProperties.PlayingAssetBlendParameters = ChooserPlayerSettings->PlayingAssetBlendParameters;
+			const FFloatInterval PoseJumpThresholdTime(0.f, 0.f);
+			const UObject* DatabasesToSearch[] = { InternalDatabase.Get() };
+			FSearchContext SearchContext(0.f, PoseJumpThresholdTime, FPoseSearchEvent());
 
-			uint8 InterruptModeValue = 0;
-			if (InterruptMode.IsValid() && InterruptMode.Get<FChooserParameterEnumBase>().GetValue(Context, InterruptModeValue))
+			const FRole Role = InternalDatabase->Schema ? InternalDatabase->Schema->GetDefaultRole() : DefaultRole;
+			SearchContext.AddRole(Role, &Context, PoseHistory);
+			SearchContext.SetAssetsToConsiderSet(&AssetsToConsider);
+
+			FPoseSearchContinuingProperties ContinuingProperties;
+			if (PoseSearchColumnScratchArea->ChooserPlayerSettings)
 			{
-				ContinuingProperties.InterruptMode = (EPoseSearchInterruptMode)InterruptModeValue;
-			}
-		}
+				ContinuingProperties.PlayingAsset = PoseSearchColumnScratchArea->ChooserPlayerSettings->PlayingAsset;
+				ContinuingProperties.PlayingAssetAccumulatedTime = PoseSearchColumnScratchArea->ChooserPlayerSettings->PlayingAssetAccumulatedTime;
+				ContinuingProperties.bIsPlayingAssetMirrored = PoseSearchColumnScratchArea->ChooserPlayerSettings->bIsPlayingAssetMirrored;
+				ContinuingProperties.PlayingAssetBlendParameters = PoseSearchColumnScratchArea->ChooserPlayerSettings->PlayingAssetBlendParameters;
 
-		FSearchResult BestSearchResult;
-		if (MaxNumberOfResults == 1)
-		{
-			FSearchResults_Single SearchResults;
-			UPoseSearchLibrary::MotionMatch(SearchContext, DatabasesToSearch, ContinuingProperties, SearchResults);
-			BestSearchResult = SearchResults.GetBestResult();
+				uint8 InterruptModeValue = 0;
+				if (InterruptMode.IsValid() && InterruptMode.Get<FChooserParameterEnumBase>().GetValue(Context, InterruptModeValue))
+				{
+					ContinuingProperties.InterruptMode = (EPoseSearchInterruptMode)InterruptModeValue;
+				}
+			}
 
-			if (const FSearchIndexAsset* SearchIndexAsset = BestSearchResult.GetSearchIndexAsset())
+			FSearchResult BestSearchResult;
+			if (MaxNumberOfResults == 1)
 			{
-				check(BestSearchResult.Database == InternalDatabase);
+				FSearchResults_Single SearchResults;
+				UPoseSearchLibrary::MotionMatch(SearchContext, DatabasesToSearch, ContinuingProperties, SearchResults);
+				BestSearchResult = SearchResults.GetBestResult();
 
-				// @todo: handle duplicate entries, by providing to the UPoseSearchLibrary::MotionMatch a SourceAssetIdxsToConsider rather than AssetsToConsider
-				const int32 SearchResultSourceAssetIdx = SearchIndexAsset->GetSourceAssetIdx();
-				check(SearchResultSourceAssetIdx >= 0);
-				PoseSearchColumnScratchArea->SearchResults.Add(BestSearchResult);
-				IndexListOut.Push({ static_cast<uint32>(SearchResultSourceAssetIdx), 0.f });
+				if (const FSearchIndexAsset* SearchIndexAsset = BestSearchResult.GetSearchIndexAsset())
+				{
+					check(BestSearchResult.Database == InternalDatabase);
+
+					// @todo: handle duplicate entries, by providing to the UPoseSearchLibrary::MotionMatch a SourceAssetIdxsToConsider rather than AssetsToConsider
+					const int32 SearchResultSourceAssetIdx = SearchIndexAsset->GetSourceAssetIdx();
+					check(SearchResultSourceAssetIdx >= 0);
+
+					const int32 RowIndex = GetRowIndex(SearchResultSourceAssetIdx);
+					check(RowIndex >= 0);
+					PoseSearchColumnScratchArea->SearchResults.Add(BestSearchResult);
+					IndexListOut.Push({ static_cast<uint32>(RowIndex), BestSearchResult.PoseCost });
 #if WITH_EDITOR
-				TracedActiveColumnCosts.Add({SearchResultSourceAssetIdx, 0.f});
+					TracedActiveColumnCosts.Add({ RowIndex, BestSearchResult.PoseCost });
 #endif // WITH_EDITOR
+				}
 			}
-		}
-		else
-		{
-			FSearchResults_AssetBests SearchResults;
-			UPoseSearchLibrary::MotionMatch(SearchContext, DatabasesToSearch, ContinuingProperties, SearchResults);
-			BestSearchResult = SearchResults.GetBestResult();
-			if (MaxNumberOfResults > 0)
+			else
 			{
-				// keeping only up to MaxNumberOfResults results
-				SearchResults.Shrink(MaxNumberOfResults);
-			}
+				FSearchResults_AssetBests SearchResults;
+				UPoseSearchLibrary::MotionMatch(SearchContext, DatabasesToSearch, ContinuingProperties, SearchResults);
+				BestSearchResult = SearchResults.GetBestResult();
+				if (MaxNumberOfResults > 0)
+				{
+					// keeping only up to MaxNumberOfResults results
+					SearchResults.Shrink(MaxNumberOfResults);
+				}
 
-			for (const FChooserIndexArray::FIndexData& IndexData : IndexListIn)
-			{
-				if (const FSearchResult* FoundSearchResult = SearchResults.FindSearchResultFor(InternalDatabase.Get(), IndexData.Index))
+				for (const FChooserIndexArray::FIndexData& IndexData : IndexListIn)
 				{
-					PoseSearchColumnScratchArea->SearchResults.Add(*FoundSearchResult);
-					IndexListOut.Push({ IndexData.Index, FoundSearchResult->PoseCost });
+					const uint32 RowIndex = IndexData.Index;
+					const int32 DatabaseAssetIndex = GetDatabaseAssetIndex(RowIndex);
+					if (const FSearchResult* FoundSearchResult = SearchResults.FindSearchResultFor(InternalDatabase.Get(), DatabaseAssetIndex))
+					{
+						PoseSearchColumnScratchArea->SearchResults.Add(*FoundSearchResult);
+						IndexListOut.Push({ RowIndex, FoundSearchResult->PoseCost });
 #if WITH_EDITOR
-					TracedActiveColumnCosts.Add({ (int32)IndexData.Index, FoundSearchResult->PoseCost });
+						TracedActiveColumnCosts.Add({ static_cast<int32>(RowIndex), FoundSearchResult->PoseCost });
 #endif // WITH_EDITOR
+					}
 				}
 			}
 		}
 	}
 
-	if (IndexListOut.IsEmpty())
-	{
-		// if nothing passed the pose match, or if PoseHistory was not found, then ignore pose matching and just pass through everything
-		IndexListOut = IndexListIn;
-	}
-
 #if WITH_EDITOR
 	TRACE_CHOOSER_VALUE(Context, ToCStr(GetInputValue()->GetDebugName()), TracedActiveColumnCosts);
 	if (Context.DebuggingInfo.bCurrentDebugTarget)
@@ -457,20 +564,66 @@ void FPoseSearchColumn::SetOutputs(FChooserEvaluationContext& Context, int RowIn
 	check(PoseSearchColumnScratchArea->DebugOwner == this);
 #endif // DO_CHECK
 
-	for (const FSearchResult& SearchResult : PoseSearchColumnScratchArea->SearchResults)
+	if (RowIndex == ChooserColumn_SpecialIndex_Fallback && PoseSearchColumnScratchArea->PoseHistory)
 	{
-		if (const FSearchIndexAsset* SearchIndexAsset = SearchResult.GetSearchIndexAsset())
+		// making sure that if we search again with the fallback column we didn't find ANY previous result!
+		check(PoseSearchColumnScratchArea->SearchResults.IsEmpty());
+
+		// do the search again only with the fallback row..
+		FStackAssetSet AssetsToConsider;
+		const int32 DatabaseAssetIndex = GetDatabaseAssetIndex(RowIndex);
+		const FPoseSearchDatabaseAnimationAsset* Asset = InternalDatabase->GetDatabaseAnimationAsset(DatabaseAssetIndex);
+		if (ensure(Asset))
 		{
-			const int32 SearchResultSourceAssetIdx = SearchIndexAsset->GetSourceAssetIdx();
-			if (SearchResultSourceAssetIdx == RowIndex)
+			if (const UObject* AnimationAsset = Asset->GetAnimationAsset())
 			{
-				StartTimeValue = SearchResult.GetAssetTime();
-				CostValue = SearchResult.PoseCost;
-				bMirrorValue = SearchIndexAsset->IsMirrored();
-				bForceBlendToValue = SearchResult.bIsContinuingPoseSearch;
+#if DO_CHECK && WITH_EDITOR
+				// making sure InternalDatabase and the OuterChooser are still in synch
+				const UObject* ReferencedObject = GetReferencedObject(RowIndex);
+				check(ReferencedObject == AnimationAsset);
+#endif // DO_CHECK && WITH_EDITOR
+				AssetsToConsider.Add(AnimationAsset);
+
+				const FFloatInterval PoseJumpThresholdTime(0.f, 0.f);
+				const UObject* DatabasesToSearch[] = { InternalDatabase.Get() };
+				FSearchContext SearchContext(0.f, PoseJumpThresholdTime, FPoseSearchEvent());
+
+				const FRole Role = InternalDatabase->Schema ? InternalDatabase->Schema->GetDefaultRole() : DefaultRole;
+				SearchContext.AddRole(Role, &Context, PoseSearchColumnScratchArea->PoseHistory);
+				SearchContext.SetAssetsToConsiderSet(&AssetsToConsider);
 
-				if (PoseSearchColumnScratchArea->PoseHistory)
+				FPoseSearchContinuingProperties ContinuingProperties;
+				if (PoseSearchColumnScratchArea->ChooserPlayerSettings)
 				{
+					ContinuingProperties.PlayingAsset = PoseSearchColumnScratchArea->ChooserPlayerSettings->PlayingAsset;
+					ContinuingProperties.PlayingAssetAccumulatedTime = PoseSearchColumnScratchArea->ChooserPlayerSettings->PlayingAssetAccumulatedTime;
+					ContinuingProperties.bIsPlayingAssetMirrored = PoseSearchColumnScratchArea->ChooserPlayerSettings->bIsPlayingAssetMirrored;
+					ContinuingProperties.PlayingAssetBlendParameters = PoseSearchColumnScratchArea->ChooserPlayerSettings->PlayingAssetBlendParameters;
+
+					uint8 InterruptModeValue = 0;
+					if (InterruptMode.IsValid() && InterruptMode.Get<FChooserParameterEnumBase>().GetValue(Context, InterruptModeValue))
+					{
+						ContinuingProperties.InterruptMode = (EPoseSearchInterruptMode)InterruptModeValue;
+					}
+				}
+
+				FSearchResults_Single SearchResults;
+				UPoseSearchLibrary::MotionMatch(SearchContext, DatabasesToSearch, ContinuingProperties, SearchResults);
+				const FSearchResult BestSearchResult = SearchResults.GetBestResult();
+
+				if (const FSearchIndexAsset* SearchIndexAsset = BestSearchResult.GetSearchIndexAsset())
+				{
+					check(BestSearchResult.Database == InternalDatabase);
+
+					// @todo: handle duplicate entries, by providing to the UPoseSearchLibrary::MotionMatch a SourceAssetIdxsToConsider rather than AssetsToConsider
+					const int32 SearchResultSourceAssetIdx = SearchIndexAsset->GetSourceAssetIdx();
+					check(SearchResultSourceAssetIdx == 0 && RowIndex == GetRowIndex(SearchResultSourceAssetIdx));
+
+					StartTimeValue = BestSearchResult.GetAssetTime();
+					CostValue = BestSearchResult.PoseCost;
+					bMirrorValue = SearchIndexAsset->IsMirrored();
+					bForceBlendToValue = BestSearchResult.bIsContinuingPoseSearch;
+
 					if (const FPoseIndicesHistory* PoseIndicesHistory = PoseSearchColumnScratchArea->PoseHistory->GetPoseIndicesHistory())
 					{
 						// @todo: integrate DeltaTime into SearchContext, and implement it for UAF as well
@@ -481,14 +634,48 @@ void FPoseSearchColumn::SetOutputs(FChooserEvaluationContext& Context, int RowIn
 						}
 
 						// const casting here is safe since we're in the thread owning the pose history, and it's the correct place to update the previously selected poses
-						const_cast<FPoseIndicesHistory*>(PoseIndicesHistory)->Update(SearchResult, DeltaTime, PoseReselectHistory);
+						const_cast<FPoseIndicesHistory*>(PoseIndicesHistory)->Update(BestSearchResult, DeltaTime, PoseReselectHistory);
 					}
 				}
-				break;
 			}
 		}
 	}
-	
+	else
+	{
+		for (const FSearchResult& SearchResult : PoseSearchColumnScratchArea->SearchResults)
+		{
+			if (const FSearchIndexAsset* SearchIndexAsset = SearchResult.GetSearchIndexAsset())
+			{
+				const int32 SearchResultSourceAssetIdx = SearchIndexAsset->GetSourceAssetIdx();
+				const int32 SearchResultRowIndex = GetRowIndex(SearchResultSourceAssetIdx);
+				if (SearchResultRowIndex == RowIndex)
+				{
+					StartTimeValue = SearchResult.GetAssetTime();
+					CostValue = SearchResult.PoseCost;
+					bMirrorValue = SearchIndexAsset->IsMirrored();
+					bForceBlendToValue = SearchResult.bIsContinuingPoseSearch;
+
+					if (PoseSearchColumnScratchArea->PoseHistory)
+					{
+						if (const FPoseIndicesHistory* PoseIndicesHistory = PoseSearchColumnScratchArea->PoseHistory->GetPoseIndicesHistory())
+						{
+							// @todo: integrate DeltaTime into SearchContext, and implement it for UAF as well
+							float DeltaTime = FiniteDelta;
+							if (const UAnimInstance* AnimInstance = Cast<UAnimInstance>(Context.GetFirstObjectParam()))
+							{
+								DeltaTime = AnimInstance->GetDeltaSeconds();
+							}
+
+							// const casting here is safe since we're in the thread owning the pose history, and it's the correct place to update the previously selected poses
+							const_cast<FPoseIndicesHistory*>(PoseIndicesHistory)->Update(SearchResult, DeltaTime, PoseReselectHistory);
+						}
+					}
+					break;
+				}
+			}
+		}
+	}
+
 	if (const FChooserParameterFloatBase* StartTime = OutputStartTime.GetPtr<FChooserParameterFloatBase>())
 	{
 		StartTime->SetValue(Context, StartTimeValue);
diff --git a/Engine/Plugins/Animation/PoseSearch/Source/Runtime/Public/PoseSearch/Chooser/PoseSearchChooserColumn.h b/Engine/Plugins/Animation/PoseSearch/Source/Runtime/Public/PoseSearch/Chooser/PoseSearchChooserColumn.h
--- a/Engine/Plugins/Animation/PoseSearch/Source/Runtime/Public/PoseSearch/Chooser/PoseSearchChooserColumn.h
+++ b/Engine/Plugins/Animation/PoseSearch/Source/Runtime/Public/PoseSearch/Chooser/PoseSearchChooserColumn.h
@@ -65,9 +65,11 @@ struct FPoseSearchColumn : public FChooserColumnBase
 	
 	UE_API FPoseSearchColumn();
 
+private:
 	UPROPERTY(VisibleAnywhere, Category = "Data")
 	TObjectPtr<UPoseSearchDatabase> InternalDatabase;
-	
+
+public:
 	// @todo: expose it as pinnable parameter
 	// Prevent re-selection of poses that have been selected previously within this much time (in seconds) in the past. This is across all animation segments that have been selected within this time range.
 	UPROPERTY(EditAnywhere, Category = Settings, meta = (PinHiddenByDefault, ClampMin = "0"))
@@ -110,6 +112,10 @@ struct FPoseSearchColumn : public FChooserColumnBase
 	UPROPERTY(EditAnywhere, NoClear, Meta = (ExcludeBaseStruct, BaseStruct = "/Script/Chooser.ChooserParameterFloatBase"), Category = "Output")
 	FInstancedStruct OutputCost;
 
+	UE_API const UObject* GetDatabaseAsset(int32 RowIndex) const;
+	UE_API const UPoseSearchSchema* GetDatabaseSchema() const;
+	UE_API void SetDatabaseSchema(const UPoseSearchSchema* Schema);
+
 private:
 	virtual bool HasFilters() const override;
 	virtual void Filter(FChooserEvaluationContext& Context, const FChooserIndexArray& IndexListIn, FChooserIndexArray& IndexListOut, TArrayView<uint8> ScratchArea) const override;
@@ -152,6 +158,7 @@ private:
 #if WITH_EDITOR
 	UObject* GetReferencedObject(int32 RowIndex) const;
 #endif // WITH_EDITOR
+	
 	void CheckConsistency() const;
 };
 
diff --git a/Engine/Plugins/Chooser/Source/Chooser/Private/Chooser.cpp b/Engine/Plugins/Chooser/Source/Chooser/Private/Chooser.cpp
--- a/Engine/Plugins/Chooser/Source/Chooser/Private/Chooser.cpp
+++ b/Engine/Plugins/Chooser/Source/Chooser/Private/Chooser.cpp
@@ -468,13 +468,13 @@ FObjectChooserBase::EIteratorStatus UChooserTable::EvaluateChooser(FChooserEvalu
 	}
 
 	// calculating the necessary scratch area size and allocating it on the stack
+	static constexpr int32 ScratchAreaAlignment = 16;
 	int32 TotalScratchAreaSize = 0;
 	for (const FInstancedStruct& ColumnData : Chooser->ColumnsStructs)
 	{
 		const FChooserColumnBase& Column = ColumnData.Get<FChooserColumnBase>();
-		TotalScratchAreaSize += Column.GetScratchAreaSize();
+		TotalScratchAreaSize += Align(Column.GetScratchAreaSize(), ScratchAreaAlignment);
 	}
-	static constexpr int32 ScratchAreaAlignment = 16;
 	TArrayView<uint8> ScratchArea((uint8*)FMemory_Alloca_Aligned(TotalScratchAreaSize * sizeof(uint8), ScratchAreaAlignment), TotalScratchAreaSize);
 	auto DeinitializeScratchAreas = [Chooser, ScratchArea, TotalScratchAreaSize]()
 		{
diff --git a/Engine/Plugins/Chooser/Source/ChooserEditor/Private/ChooserTableEditor.cpp b/Engine/Plugins/Chooser/Source/ChooserEditor/Private/ChooserTableEditor.cpp
--- a/Engine/Plugins/Chooser/Source/ChooserEditor/Private/ChooserTableEditor.cpp
+++ b/Engine/Plugins/Chooser/Source/ChooserEditor/Private/ChooserTableEditor.cpp
@@ -1968,6 +1968,17 @@ void FChooserTableEditor::AutoPopulateRow(int Index)
 				}
 			}
 		}
+		else if (Index == ColumnWidget_SpecialIndex_Fallback && Chooser->FallbackResult.IsValid())
+		{
+			if (UObject* ReferencedObject = Chooser->FallbackResult.Get<FObjectChooserBase>().GetReferencedObject())
+			{
+				for (FInstancedStruct& ColumnData : Chooser->ColumnsStructs)
+				{
+					FChooserColumnBase& Column = ColumnData.GetMutable<FChooserColumnBase>();
+					Column.AutoPopulate(Index, ReferencedObject);
+				}
+			}
+		}
 	}
 }
 	
diff --git a/Engine/Plugins/Chooser/Source/ChooserEditor/Private/SChooserTableRow.cpp b/Engine/Plugins/Chooser/Source/ChooserEditor/Private/SChooserTableRow.cpp
--- a/Engine/Plugins/Chooser/Source/ChooserEditor/Private/SChooserTableRow.cpp
+++ b/Engine/Plugins/Chooser/Source/ChooserEditor/Private/SChooserTableRow.cpp
@@ -300,19 +300,41 @@ namespace UE::ChooserEditor
 			}
 			else if (ColumnName == Result) 
 			{
+				FChooserWidgetValueChanged ChooserWidgetValueChanged = FChooserWidgetValueChanged::CreateLambda([this]()
+					{
+						check(Editor);
+						Editor->AutoPopulateRow(RowIndex->RowIndex);
+					});
+
 				UChooserTable* ContextOwner = Chooser->GetRootChooser();
-				TSharedPtr<SWidget> ResultWidget = FObjectChooserWidgetFactories::CreateWidget(false, Chooser, FObjectChooserBase::StaticStruct(), Chooser->FallbackResult.GetMutableMemory(), Chooser->FallbackResult.GetScriptStruct(), ContextOwner->OutputObjectType,
-				FOnStructPicked::CreateLambda([this](const UScriptStruct* ChosenStruct)
-				{
-					UChooserTable* ContextOwner = Chooser->GetRootChooser();
-					const FScopedTransaction Transaction(LOCTEXT("Change Row Result Type", "Change Row Result Type"));
-					Chooser->Modify(true);
-					Chooser->FallbackResult.InitializeAs(ChosenStruct);
-					FObjectChooserWidgetFactories::CreateWidget(false, Chooser, FObjectChooserBase::StaticStruct(), Chooser->FallbackResult.GetMutableMemory(), ChosenStruct, ContextOwner->OutputObjectType, FOnStructPicked(), &CacheBorder
-						,FChooserWidgetValueChanged(), LOCTEXT("Fallback Result", "Fallback Result: (None)"));
-				}),
-				&CacheBorder
-				,FChooserWidgetValueChanged(), LOCTEXT("Fallback Result", "Fallback Result: (None)")
+				TSharedPtr<SWidget> ResultWidget = FObjectChooserWidgetFactories::CreateWidget(
+					false,
+					Chooser,
+					FObjectChooserBase::StaticStruct(),
+					Chooser->FallbackResult.GetMutableMemory(),
+					Chooser->FallbackResult.GetScriptStruct(), ContextOwner->OutputObjectType,
+					FOnStructPicked::CreateLambda([this, InRowIndex = RowIndex->RowIndex, ChooserWidgetValueChanged](const UScriptStruct* ChosenStruct)
+						{
+							UChooserTable* ContextOwner = Chooser->GetRootChooser();
+							const FScopedTransaction Transaction(LOCTEXT("Change Row Result Type", "Change Row Result Type"));
+							Chooser->Modify(true);
+							Chooser->FallbackResult.InitializeAs(ChosenStruct);
+							FObjectChooserWidgetFactories::CreateWidget(
+								false,
+								Chooser,
+								FObjectChooserBase::StaticStruct(),
+								Chooser->FallbackResult.GetMutableMemory(),
+								ChosenStruct,
+								ContextOwner->OutputObjectType,
+								FOnStructPicked(),
+								&CacheBorder,
+								ChooserWidgetValueChanged,
+								LOCTEXT("Fallback Result", "Fallback Result: (None)")
+							);
+						}),
+					&CacheBorder,
+					ChooserWidgetValueChanged,
+					LOCTEXT("Fallback Result", "Fallback Result: (None)")
 				);
 				
 				return ResultWidget.ToSharedRef();
