PoseSearch
- adding MM interactions support to FAnimNode_MotionMatching (to mimic FMotionMatchingTrait)
- making UPoseSearchInteractionSubsystem::AddAvailabilities lockless

diff --git a/Engine/Plugins/Animation/PoseSearch/Source/Runtime/Private/AnimNode_MotionMatching.cpp b/Engine/Plugins/Animation/PoseSearch/Source/Runtime/Private/AnimNode_MotionMatching.cpp
--- a/Engine/Plugins/Animation/PoseSearch/Source/Runtime/Private/AnimNode_MotionMatching.cpp
+++ b/Engine/Plugins/Animation/PoseSearch/Source/Runtime/Private/AnimNode_MotionMatching.cpp
@@ -1,7 +1,6 @@
 // Copyright Epic Games, Inc. All Rights Reserved.
 
 #include "PoseSearch/AnimNode_MotionMatching.h"
-
 #include "Animation/AnimInertializationSyncScope.h"
 #include "Animation/AnimInstanceProxy.h"
 #include "Animation/AnimNode_Inertialization.h"
@@ -9,12 +8,18 @@
 #include "Animation/AnimSequence.h"
 #include "Animation/AnimStats.h"
 #include "Animation/BlendSpace.h"
+#include "BoneControllers/AnimNode_OffsetRootBone.h"
 #include "Components/SkeletalMeshComponent.h"
 #include "PoseSearch/AnimNode_PoseSearchHistoryCollector.h"
+#include "PoseSearch/MultiAnimAsset.h"
 #include "PoseSearch/PoseSearchDatabase.h"
 #include "PoseSearch/PoseSearchDerivedData.h"
+#include "PoseSearch/PoseSearchInteractionLibrary.h"
+#include "PoseSearch/PoseSearchInteractionUtils.h"
+#include "PoseSearch/PoseSearchInteractionValidator.h"
 #include "PoseSearch/PoseSearchSchema.h"
 #include "PoseSearch/Trace/PoseSearchTraceLogger.h"
+#include "VisualLogger/VisualLogger.h"
 
 #include UE_INLINE_GENERATED_CPP_BY_NAME(AnimNode_MotionMatching)
 
@@ -157,76 +162,83 @@ void FAnimNode_MotionMatching::UpdateAssetPlayer(const FAnimationUpdateContext&
 	}
 	else
 	{
-		FFloatInterval PoseSearchPlayRate = PlayRate;
+		const UE::AnimationWarping::FRootOffsetProvider* RootOffsetProvider = Context.GetMessage<UE::AnimationWarping::FRootOffsetProvider>();
+		MeshWithoutOffset = Context.AnimInstanceProxy->GetComponentTransform();
+		MeshWithOffset = RootOffsetProvider ? RootOffsetProvider->GetRootTransform() : MeshWithoutOffset;
+
+		const UObject* AnimContext = Context.AnimInstanceProxy->GetAnimInstanceObject();
+		check(AnimContext);
+		const bool bIsInteraction = UPoseSearchInteractionLibrary::MotionMatchInteraction(MotionMatchingState.SearchResult, Availabilities, AnimContext, FName(), PoseHistory, bValidateResultAgainstAvailabilities, bKeepInteractionAlive, BlendTime);
+
+		// performing the regular single character motion matching search in case there's no MM interaction
+		if (!bIsInteraction)
+		{
+			FFloatInterval PoseSearchPlayRate = PlayRate;
 #if ENABLE_ANIM_DEBUG
-		bool bPoseSearchPlayRateEnabled = GVarAnimNodeMotionMatchingPlayRateEnabled == UE::Private::EPlayRateState::Enabled || GVarAnimNodeMotionMatchingPlayRateEnabled == UE::Private::EPlayRateState::PoseSearchOnly;
-		PoseSearchPlayRate = bPoseSearchPlayRateEnabled ? PoseSearchPlayRate : FFloatInterval(1.f, 1.f);
+			bool bPoseSearchPlayRateEnabled = GVarAnimNodeMotionMatchingPlayRateEnabled == UE::Private::EPlayRateState::Enabled || GVarAnimNodeMotionMatchingPlayRateEnabled == UE::Private::EPlayRateState::PoseSearchOnly;
+			PoseSearchPlayRate = bPoseSearchPlayRateEnabled ? PoseSearchPlayRate : FFloatInterval(1.f, 1.f);
 #endif // ENABLE_ANIM_DEBUG
 
-		const float DeltaTime = Context.GetDeltaTime();
+			const float DeltaTime = Context.GetDeltaTime();
 
-		FSearchContext SearchContext(0.f, PoseJumpThresholdTime, EventToSearch.GetPlayRateOverriddenEvent(PoseSearchPlayRate));
-		FChooserEvaluationContext ChooserEvaluationContext(const_cast<UObject*>(Context.AnimInstanceProxy->GetAnimInstanceObject()));
-		SearchContext.AddRole(GetCommonDefaultRole(DatabasesToSearch), &ChooserEvaluationContext, PoseHistory);
+			FSearchContext SearchContext(0.f, PoseJumpThresholdTime, EventToSearch.GetPlayRateOverriddenEvent(PoseSearchPlayRate));
+			FChooserEvaluationContext ChooserEvaluationContext(const_cast<UObject*>(AnimContext));
+			SearchContext.AddRole(GetCommonDefaultRole(DatabasesToSearch), &ChooserEvaluationContext, PoseHistory);
 
-		// If we can't advance or enough time has elapsed since the last pose jump then search
-		const bool bCanAdvance = SearchContext.GetContinuingPoseSearchResult().PoseIdx != INDEX_NONE;
-		const bool bSearch = !bCanAdvance || (MotionMatchingState.ElapsedPoseSearchTime >= SearchThrottleTime);
-		FSearchResult SearchResult;
-		if (bSearch)
-		{
-			MotionMatchingState.ElapsedPoseSearchTime = 0.f;
-
-			TArrayView<const UObject*> AssetsToSearch((const UObject**)FMemory_Alloca(DatabasesToSearch.Num() * sizeof(UObject*)), (int32)DatabasesToSearch.Num());
-			for (int32 DatabaseToSearchIndex = 0; DatabaseToSearchIndex < DatabasesToSearch.Num(); ++DatabaseToSearchIndex)
+			// If we can't advance or enough time has elapsed since the last pose jump then search
+			const bool bCanAdvance = SearchContext.GetContinuingPoseSearchResult().PoseIdx != INDEX_NONE;
+			const bool bSearch = !bCanAdvance || (MotionMatchingState.ElapsedPoseSearchTime >= SearchThrottleTime);
+			FSearchResult SearchResult;
+			if (bSearch)
 			{
-				AssetsToSearch[DatabaseToSearchIndex] = DatabasesToSearch[DatabaseToSearchIndex].Get();
-			}
+				MotionMatchingState.ElapsedPoseSearchTime = 0.f;
 
-			FSearchResults_Single SearchResults;
-			FPoseSearchContinuingProperties ContinuingProperties;
-			ContinuingProperties.InitFrom(MotionMatchingState.SearchResult, NextUpdateInterruptMode);
-			UPoseSearchLibrary::MotionMatch(SearchContext, AssetsToSearch, ContinuingProperties, SearchResults);
-			SearchResult = SearchResults.GetBestResult();
-		}
-		else
-		{
-			MotionMatchingState.ElapsedPoseSearchTime += DeltaTime;
-
-			// if we didn't perform a search we carry on with the continuing pose search result
-			SearchResult = SearchContext.GetContinuingPoseSearchResult();
-			SearchResult.bIsContinuingPoseSearch = true;
+				TArrayView<const UObject*> AssetsToSearch((const UObject**)FMemory_Alloca(DatabasesToSearch.Num() * sizeof(UObject*)), (int32)DatabasesToSearch.Num());
+				for (int32 DatabaseToSearchIndex = 0; DatabaseToSearchIndex < DatabasesToSearch.Num(); ++DatabaseToSearchIndex)
+				{
+					AssetsToSearch[DatabaseToSearchIndex] = DatabasesToSearch[DatabaseToSearchIndex].Get();
+				}
 
-#if UE_POSE_SEARCH_TRACE_ENABLED
-			// in case we skipped the search, we still have to track we would have requested to evaluate DatabasesToSearch and SearchResult.Database
-			for (const TObjectPtr<const UPoseSearchDatabase>& DatabaseToSearch : DatabasesToSearch)
-			{
-				SearchContext.Track(DatabaseToSearch.Get());
+				FSearchResults_Single SearchResults;
+				FPoseSearchContinuingProperties ContinuingProperties;
+				ContinuingProperties.InitFrom(MotionMatchingState.SearchResult, NextUpdateInterruptMode);
+				UPoseSearchLibrary::MotionMatch(SearchContext, AssetsToSearch, ContinuingProperties, SearchResults);
+				SearchResult = SearchResults.GetBestResult();
 			}
+			else
+			{
+				MotionMatchingState.ElapsedPoseSearchTime += DeltaTime;
 
-			SearchContext.Track(SearchResult.Database.Get());
-#endif // UE_POSE_SEARCH_TRACE_ENABLED
-		}
+				// if we didn't perform a search we carry on with the continuing pose search result
+				SearchResult = SearchContext.GetContinuingPoseSearchResult();
+				SearchResult.bIsContinuingPoseSearch = true;
 
-		const float WantedPlayRate = CalculateWantedPlayRate(SearchResult, SearchContext, PoseSearchPlayRate, PoseHistory ? PoseHistory->GetTrajectorySpeedMultiplier() : 1.f, EventToSearch);
-		MotionMatchingState.SearchResult.InitFrom(SearchResult, WantedPlayRate);
+#if UE_POSE_SEARCH_TRACE_ENABLED
+				// in case we skipped the search, we still have to track we would have requested to evaluate DatabasesToSearch and SearchResult.Database
+				for (const TObjectPtr<const UPoseSearchDatabase>& DatabaseToSearch : DatabasesToSearch)
+				{
+					SearchContext.Track(DatabaseToSearch.Get());
+				}
 
-		if (const FPoseIndicesHistory* PoseIndicesHistory = PoseHistory->GetPoseIndicesHistory())
-		{
-			// const casting here is safe since we're in the thread owning the pose history, and it's the correct place to update the previously selected poses
-			const_cast<FPoseIndicesHistory*>(PoseIndicesHistory)->Update(SearchResult, DeltaTime, PoseReselectHistory);
-		}
+				SearchContext.Track(SearchResult.Database.Get());
+#endif // UE_POSE_SEARCH_TRACE_ENABLED
+			}
 
-		PRAGMA_DISABLE_DEPRECATION_WARNINGS
-		MotionMatchingState.CurrentSearchResult = SearchResult;
-		MotionMatchingState.WantedPlayRate = WantedPlayRate;
-		PRAGMA_ENABLE_DEPRECATION_WARNINGS
+			const float WantedPlayRate = CalculateWantedPlayRate(SearchResult, SearchContext, PoseSearchPlayRate, PoseHistory ? PoseHistory->GetTrajectorySpeedMultiplier() : 1.f, EventToSearch);
+			MotionMatchingState.SearchResult.InitFrom(SearchResult, WantedPlayRate);
 
-		const bool bJumpToPose = MotionMatchingState.SearchResult.SelectedAnim && !MotionMatchingState.SearchResult.bIsContinuingPoseSearch;
+			if (const FPoseIndicesHistory* PoseIndicesHistory = PoseHistory->GetPoseIndicesHistory())
+			{
+				// const casting here is safe since we're in the thread owning the pose history, and it's the correct place to update the previously selected poses
+				const_cast<FPoseIndicesHistory*>(PoseIndicesHistory)->Update(SearchResult, DeltaTime, PoseReselectHistory);
+			}
 
-		PRAGMA_DISABLE_DEPRECATION_WARNINGS
-		MotionMatchingState.bJumpedToPose = bJumpToPose;
-		PRAGMA_ENABLE_DEPRECATION_WARNINGS
+			PRAGMA_DISABLE_DEPRECATION_WARNINGS
+			MotionMatchingState.CurrentSearchResult = SearchResult;
+			MotionMatchingState.WantedPlayRate = WantedPlayRate;
+			MotionMatchingState.bJumpedToPose = MotionMatchingState.SearchResult.SelectedAnim && !MotionMatchingState.SearchResult.bIsContinuingPoseSearch;
+			PRAGMA_ENABLE_DEPRECATION_WARNINGS
+		}
 
 		float DesiredPlayRate = MotionMatchingState.SearchResult.WantedPlayRate * PlayRateMultiplier;
 #if ENABLE_ANIM_DEBUG
@@ -236,6 +248,7 @@ void FAnimNode_MotionMatching::UpdateAssetPlayer(const FAnimationUpdateContext&
 		UE::Anim::FNodeFunctionCaller::CallFunction(GetOnUpdateMotionMatchingStateFunction(), Context, *this);
 
 		// If a new pose is requested, blend into the new asset via BlendStackNode
+		const bool bJumpToPose = MotionMatchingState.SearchResult.SelectedAnim && !MotionMatchingState.SearchResult.bIsContinuingPoseSearch;
 		if (bJumpToPose)
 		{
 			const UPoseSearchDatabase* CurrentResultDatabase = MotionMatchingState.SearchResult.SelectedDatabase.Get();
@@ -298,6 +311,91 @@ void FAnimNode_MotionMatching::UpdateAssetPlayer(const FAnimationUpdateContext&
 	NextUpdateInterruptMode = EPoseSearchInterruptMode::DoNotInterrupt;
 }
 
+void FAnimNode_MotionMatching::Evaluate_AnyThread(FPoseContext& Output)
+{
+	using namespace UE::PoseSearch;
+
+	check(Output.AnimInstanceProxy);
+	CheckInteractionThreadSafety(Output.AnimInstanceProxy->GetAnimInstanceObject());
+
+	Super::Evaluate_AnyThread(Output);
+
+	// calculating the animation time we want the full alignment transforms from
+	UAnimInstance* AnimInstance = Cast<UAnimInstance>(Output.AnimInstanceProxy->GetAnimInstanceObject());
+	if (!AnimInstance)
+	{
+		UE_LOG(LogPoseSearch, Error, TEXT("FAnimNode_MotionMatching::Evaluate_AnyThread couldn't find the AnimInstance!?"));
+		return;
+	}
+
+	const UE::Anim::IAnimRootMotionProvider* RootMotionProvider = UE::Anim::IAnimRootMotionProvider::Get();
+
+	if (!RootMotionProvider)
+	{
+		UE_LOG(LogPoseSearch, Error, TEXT("FAnimNode_MotionMatching::Evaluate_AnyThread couldn't find the IAnimRootMotionProvider"));
+		return;
+	}
+
+	const UMultiAnimAsset* MultiAnimAsset = Cast<UMultiAnimAsset>(MotionMatchingState.SearchResult.SelectedAnim);
+	if (!MultiAnimAsset)
+	{
+		return;
+	}
+
+	const int32 NumRoles = MultiAnimAsset->GetNumRoles();
+	if (MotionMatchingState.SearchResult.ActorRootTransforms.Num() != NumRoles)
+	{
+		// warping is supported only for UMultiAnimAsset(s)
+		return;
+	}
+
+	const int32 CurrentResultRoleIndex = GetRoleIndex(MultiAnimAsset, MotionMatchingState.SearchResult.Role);
+	if (CurrentResultRoleIndex == INDEX_NONE)
+	{
+		return;
+	}
+
+	const UAnimationAsset* AnimationAsset = MultiAnimAsset->GetAnimationAsset(MotionMatchingState.SearchResult.Role);
+	if (!AnimationAsset)
+	{
+		UE_LOG(LogPoseSearch, Error, TEXT("FAnimNode_MotionMatching::Evaluate_AnyThread MultiAnimAsset %s for Role %s is invalid!"), *GetNameSafe(MultiAnimAsset), *MotionMatchingState.SearchResult.Role.ToString());
+		return;
+	}
+
+	TArray<FTransform, TInlineAllocator<PreallocatedRolesNum>> FullAlignedTransforms;
+	FullAlignedTransforms.SetNum(NumRoles);
+	CalculateFullAlignedTransforms(MotionMatchingState.SearchResult, bWarpUsingRootBone, FullAlignedTransforms);
+
+	// NoTe: keep in mind DeltaAlignment is relative to the previous execution frame so we still need to extract and and apply the current animation root motion transform to get to the current frame full aligned transform.
+	const FTransform DeltaAlignment = CalculateDeltaAlignment(MeshWithoutOffset, MeshWithOffset, FullAlignedTransforms[CurrentResultRoleIndex], WarpingRotationRatio, WarpingTranslationRatio);
+	
+	FTransform RootMotionDelta;
+	RootMotionProvider->ExtractRootMotion(Output.CustomAttributes, RootMotionDelta);
+
+	const FTransform DeltaAlignmentWithRootMotion = DeltaAlignment * RootMotionDelta;
+	RootMotionProvider->OverrideRootMotion(DeltaAlignmentWithRootMotion, Output.CustomAttributes);
+
+#if ENABLE_VISUAL_LOG
+	if (FVisualLogger::IsRecording())
+	{
+		static const TCHAR* LogName = TEXT("MotionMatching");
+
+		for (int32 Index = 0; Index < NumRoles; ++Index)
+		{
+			const FTransform& ActorRootTransform = MotionMatchingState.SearchResult.ActorRootTransforms[Index];
+			const FTransform& FullAlignedTransform = FullAlignedTransforms[Index];
+
+			UE_VLOG_SEGMENT_THICK(AnimInstance, LogName, Display, FullAlignedTransform.GetLocation(), ActorRootTransform.GetLocation(), FColorList::Orange, 1.f, TEXT(""));
+			UE_VLOG_SEGMENT_THICK(AnimInstance, LogName, Display, ActorRootTransform.GetLocation(), ActorRootTransform.GetLocation() + ActorRootTransform.GetRotation().GetForwardVector() * 35, FColorList::LightGrey, 3.f, TEXT(""));
+			UE_VLOG_SEGMENT_THICK(AnimInstance, LogName, Display, FullAlignedTransform.GetLocation(), FullAlignedTransform.GetLocation() + FullAlignedTransform.GetRotation().GetForwardVector() * 30, FColorList::Orange, 2.f, TEXT(""));
+		}
+
+		UE_VLOG_SEGMENT_THICK(AnimInstance, LogName, Display, MeshWithOffset.GetLocation(), MeshWithOffset.GetLocation() + MeshWithOffset.GetRotation().GetForwardVector() * 35, FColorList::Blue, 3.f, TEXT(""));
+		UE_VLOG_SEGMENT_THICK(AnimInstance, LogName, Display, MeshWithoutOffset.GetLocation(), MeshWithoutOffset.GetLocation() + MeshWithoutOffset.GetRotation().GetForwardVector() * 40, FColorList::Cyan, 4.f, TEXT(""));
+	}
+#endif
+}
+
 const FAnimNodeFunctionRef& FAnimNode_MotionMatching::GetOnUpdateMotionMatchingStateFunction() const
 {
 	return GET_ANIM_NODE_DATA(FAnimNodeFunctionRef, OnMotionMatchingStateUpdated);
diff --git a/Engine/Plugins/Animation/PoseSearch/Source/Runtime/Private/PoseSearchContext.cpp b/Engine/Plugins/Animation/PoseSearch/Source/Runtime/Private/PoseSearchContext.cpp
--- a/Engine/Plugins/Animation/PoseSearch/Source/Runtime/Private/PoseSearchContext.cpp
+++ b/Engine/Plugins/Animation/PoseSearch/Source/Runtime/Private/PoseSearchContext.cpp
@@ -17,10 +17,9 @@
 namespace UE::PoseSearch
 {
 	
-const FTransform& GetContextTransform(const UObject* AnimContext)
+const FTransform& GetContextTransform(const UObject* AnimContext, bool bCheckIsInGameThread)
 {
-	// thread unsafe code. making sure we're on game thread!
-	check(IsInGameThread());
+	check(!bCheckIsInGameThread || IsInGameThread());
 
 	if (const UAnimInstance* AnimInstance = Cast<UAnimInstance>(AnimContext))
 	{
@@ -44,10 +43,9 @@ const FTransform& GetContextTransform(const UObject* AnimContext)
 	return FTransform::Identity;
 }
 
-const USkeleton* GetContextSkeleton(const UObject* AnimContext)
+const USkeleton* GetContextSkeleton(const UObject* AnimContext, bool bCheckIsInGameThread)
 {
-	// @todo: make sure this is not called from worker threads!
-	// check(IsInGameThread());
+	check(!bCheckIsInGameThread || IsInGameThread());
 
 	const USkeleton* Skeleton = nullptr;
 	if (const UAnimInstance* AnimInstance = Cast<UAnimInstance>(AnimContext))
@@ -72,19 +70,17 @@ const USkeleton* GetContextSkeleton(const UObject* AnimContext)
 	return Skeleton;
 }
 	
-const USkeleton* GetContextSkeleton(FChooserEvaluationContext& Context)
+const USkeleton* GetContextSkeleton(FChooserEvaluationContext& Context, bool bCheckIsInGameThread)
 {
-	// @todo: make sure this is not called from worker threads!
-	// check(IsInGameThread());
-
+	check(!bCheckIsInGameThread || IsInGameThread());
 	check(!Context.ObjectParams.IsEmpty());
-	return GetContextSkeleton(Context.ObjectParams[0].Object);
+
+	return GetContextSkeleton(Context.ObjectParams[0].Object, bCheckIsInGameThread);
 }
 
-const USkeletalMeshComponent* GetContextSkeletalMeshComponent(const UObject* AnimContext)
+const USkeletalMeshComponent* GetContextSkeletalMeshComponent(const UObject* AnimContext, bool bCheckIsInGameThread)
 {
-	// thread unsafe code. making sure we're on game thread!
-	check(IsInGameThread());
+	check(!bCheckIsInGameThread || IsInGameThread());
 
 	const USkeletalMeshComponent* SkeletalMeshComponent = nullptr;
 	if (const UAnimInstance* AnimInstance = Cast<UAnimInstance>(AnimContext))
@@ -103,10 +99,9 @@ const USkeletalMeshComponent* GetContextSkeletalMeshComponent(const UObject* Ani
 	return SkeletalMeshComponent;
 }
 
-const FBoneContainer GetBoneContainer(const UObject* AnimContext)
+const FBoneContainer GetBoneContainer(const UObject* AnimContext, bool bCheckIsInGameThread)
 {
-	// thread unsafe code. making sure we're on game thread!
-	check(IsInGameThread());
+	check(!bCheckIsInGameThread || IsInGameThread());
 
 	if (const UAnimInstance* AnimInstance = Cast<UAnimInstance>(AnimContext))
 	{
@@ -132,8 +127,10 @@ const FBoneContainer GetBoneContainer(const UObject* AnimContext)
 	return FBoneContainer();
 }
 
-const AActor* GetContextOwningActor(const UObject* AnimContext)
+const AActor* GetContextOwningActor(const UObject* AnimContext, bool bCheckIsInGameThread)
 {
+	check(!bCheckIsInGameThread || IsInGameThread());
+
 	if (const UAnimInstance* AnimInstance = Cast<UAnimInstance>(AnimContext))
 	{
 		return AnimInstance->GetOwningActor();
@@ -147,9 +144,9 @@ const AActor* GetContextOwningActor(const UObject* AnimContext)
 	return nullptr;
 }
 
-FVector GetContextLocation(const UObject* AnimContext)
+FVector GetContextLocation(const UObject* AnimContext, bool bCheckIsInGameThread)
 {
-	return GetContextTransform(AnimContext).GetLocation();
+	return GetContextTransform(AnimContext, bCheckIsInGameThread).GetLocation();
 }
 
 #if ENABLE_DRAW_DEBUG
diff --git a/Engine/Plugins/Animation/PoseSearch/Source/Runtime/Private/PoseSearchInteractionIsland.cpp b/Engine/Plugins/Animation/PoseSearch/Source/Runtime/Private/PoseSearchInteractionIsland.cpp
--- a/Engine/Plugins/Animation/PoseSearch/Source/Runtime/Private/PoseSearchInteractionIsland.cpp
+++ b/Engine/Plugins/Animation/PoseSearch/Source/Runtime/Private/PoseSearchInteractionIsland.cpp
@@ -132,7 +132,7 @@ static FInteractionSearchResult InitSearchResult(const FSearchResult& SearchResu
 	{
 		const UObject* AnimContext = SearchContext.GetAnimContext(AnimContextIndex);
 		check(AnimContext);
-		const USkeleton* Skeleton = GetContextSkeleton(AnimContext);
+		const USkeleton* Skeleton = GetContextSkeleton(AnimContext, false);
 		const IPoseHistory* PoseHistory = SearchContext.GetPoseHistory(AnimContextIndex);
 		check(PoseHistory && Skeleton);
 		
@@ -904,7 +904,7 @@ const AActor* FInteractionIsland::GetMainActor() const
 {
 	if (IsInitialized())
 	{
-		return GetContextOwningActor(GetMainAnimContext());
+		return GetContextOwningActor(GetMainAnimContext(), false);
 	}
 	return nullptr;
 }
@@ -990,7 +990,7 @@ bool FInteractionIsland::DoSearch_AnyThread(const UObject* AnimContext, const TC
 	TArray<FSearchResult, TMemStackAllocator<>> PoseSearchResults;
 	TArray<FChooserEvaluationContext, TInlineAllocator<UE::PoseSearch::PreallocatedRolesNum>> Contexts;
 
-	// SearchContexts are modified only by UPoseSearchInteractionSubsystem::Tick and constant otherwise, so it's safe to access them in a threaded enviroment without locks
+	// SearchContexts are modified only by UPoseSearchInteractionSubsystem::Tick and constant otherwise, so it's safe to access them in a threaded environment without locks
 	PoseSearchResults.SetNum(SearchContexts.Num());
 
 	for (int32 SearchIndex = 0; SearchIndex < SearchContexts.Num(); ++SearchIndex)
@@ -1047,7 +1047,7 @@ bool FInteractionIsland::GetResult_AnyThread(const UObject* AnimContext, FPoseSe
 {
 	check(AnimContext);
 
-	const AActor* Actor = bCompareOwningActors ? GetContextOwningActor(AnimContext) : nullptr;
+	const AActor* Actor = bCompareOwningActors ? GetContextOwningActor(AnimContext, false) : nullptr;
 
 	// looking for AnimContext in SearchResults to fill up Result
 	for (const FInteractionSearchResult& SearchResult : SearchResults)
@@ -1058,7 +1058,7 @@ bool FInteractionIsland::GetResult_AnyThread(const UObject* AnimContext, FPoseSe
 			bool bAnimContextFound = false;
 			if (bCompareOwningActors)
 			{
-				bAnimContextFound = GetContextOwningActor(SearchContext.GetAnimContext(AnimContextIndex)) == Actor;
+				bAnimContextFound = GetContextOwningActor(SearchContext.GetAnimContext(AnimContextIndex), false) == Actor;
 			}
 			else
 			{
diff --git a/Engine/Plugins/Animation/PoseSearch/Source/Runtime/Private/PoseSearchInteractionLibrary.cpp b/Engine/Plugins/Animation/PoseSearch/Source/Runtime/Private/PoseSearchInteractionLibrary.cpp
--- a/Engine/Plugins/Animation/PoseSearch/Source/Runtime/Private/PoseSearchInteractionLibrary.cpp
+++ b/Engine/Plugins/Animation/PoseSearch/Source/Runtime/Private/PoseSearchInteractionLibrary.cpp
@@ -2,9 +2,10 @@
 
 #include "PoseSearch/PoseSearchInteractionLibrary.h"
 #include "Animation/AnimMontage.h"
+#include "Engine/World.h"
 #include "GameFramework/Pawn.h"
+#include "PoseSearch/PoseSearchDatabase.h"
 #include "PoseSearch/PoseSearchInteractionSubsystem.h"
-#include "Engine/World.h"
 
 #include UE_INLINE_GENERATED_CPP_BY_NAME(PoseSearchInteractionLibrary)
 
@@ -34,6 +35,67 @@ void UPoseSearchInteractionLibrary::MotionMatchInteraction(FPoseSearchBlueprintR
 	}
 }
 
+bool UPoseSearchInteractionLibrary::MotionMatchInteraction(FPoseSearchBlueprintResult& InOutResult, const TArrayView<const FPoseSearchInteractionAvailability> Availabilities, const UObject* AnimContext, FName PoseHistoryName, const UE::PoseSearch::IPoseHistory* PoseHistory, bool bValidateResultAgainstAvailabilities, bool bKeepInteractionAlive, float BlendTime)
+{
+	using namespace UE::PoseSearch;
+	
+	bool bIsInteraction = false;
+	if (!Availabilities.IsEmpty())
+	{
+		CheckInteractionThreadSafety(AnimContext);
+
+		FPoseSearchBlueprintResult SearchResult;
+		UPoseSearchInteractionLibrary::MotionMatchInteraction(SearchResult, Availabilities, AnimContext, PoseHistoryName, PoseHistory, bValidateResultAgainstAvailabilities);
+		check(SearchResult.ActorRootTransforms.Num() == SearchResult.ActorRootBoneTransforms.Num());
+
+		if (SearchResult.SelectedAnim)
+		{
+			bIsInteraction = true;
+			InOutResult = SearchResult;
+		}
+	}
+
+	// performing the regular single character motion matching search in case there's no MM interaction
+	if (!bIsInteraction)
+	{
+		if (!InOutResult.bIsInteraction)
+		{
+			// do nothing
+		}
+		else if (bKeepInteractionAlive)
+		{
+			// checking if the kept alive interaction has reached the end of animation
+			if (!InOutResult.SelectedDatabase || INDEX_NONE == InOutResult.SelectedDatabase->GetPoseIndex(InOutResult.SelectedAnim.Get(), InOutResult.SelectedTime, InOutResult.bIsMirrored, InOutResult.BlendParameters))
+			{
+				InOutResult = FPoseSearchBlueprintResult();
+			}
+			else
+			{
+				// letting the interaction animation run until its length minus blend time
+				// (to avoid having to blend from a frozen animation that reached its end for the entire duration of the blend)
+				const UAnimationAsset* AnimationAssetForRole = InOutResult.GetAnimationAssetForRole();
+				if (!AnimationAssetForRole || InOutResult.SelectedTime >= (AnimationAssetForRole->GetPlayLength() - BlendTime))
+				{
+					InOutResult = FPoseSearchBlueprintResult();
+				}
+				else
+				{
+					// we're keeping alive only the animation part of the search result
+					InOutResult.ActorRootTransforms.Reset();
+					InOutResult.ActorRootBoneTransforms.Reset();
+					InOutResult.AnimContexts.Reset();
+				}
+			}
+		}
+		else
+		{
+			InOutResult = FPoseSearchBlueprintResult();
+		}
+	}
+
+	return bIsInteraction;
+}
+
 FPoseSearchBlueprintResult UPoseSearchInteractionLibrary::GetCachedInteraction(const UObject* AnimContext, bool bCompareOwningActors)
 {
 	FPoseSearchBlueprintResult Result;
diff --git a/Engine/Plugins/Animation/PoseSearch/Source/Runtime/Private/PoseSearchInteractionSubsystem.cpp b/Engine/Plugins/Animation/PoseSearch/Source/Runtime/Private/PoseSearchInteractionSubsystem.cpp
--- a/Engine/Plugins/Animation/PoseSearch/Source/Runtime/Private/PoseSearchInteractionSubsystem.cpp
+++ b/Engine/Plugins/Animation/PoseSearch/Source/Runtime/Private/PoseSearchInteractionSubsystem.cpp
@@ -32,7 +32,14 @@ struct FAnimContextInfo
 	{
 		check(InAnimContextAvailabilities.AnimContext && !InAnimContextAvailabilities.Availabilities.IsEmpty());
 		AnimContextAvailabilities = &InAnimContextAvailabilities;
-		Location = GetContextLocation(InAnimContextAvailabilities.AnimContext);
+		Location = GetContextLocation(InAnimContextAvailabilities.AnimContext, false);
+
+		AvailabilitiesMaxBroadPhaseRadius = 0.f;
+		for (const FPoseSearchInteractionAvailabilityEx& Availability : AnimContextAvailabilities->Availabilities)
+		{
+			// @todo: optimize the AvailabilitiesMaxBroadPhaseRadius, since adding Availability.BroadPhaseRadiusIncrementOnInteraction is required ONLY if AnimContext is already part of an interaction
+			AvailabilitiesMaxBroadPhaseRadius = FMath::Max(AvailabilitiesMaxBroadPhaseRadius, Availability.BroadPhaseRadius + Availability.BroadPhaseRadiusIncrementOnInteraction);
+		}
 	}
 
 	// performs broad phase analysis checking if at least one of the Availabilities associated to AnimContext can interact with OtherAnimContextInfo.
@@ -52,27 +59,16 @@ struct FAnimContextInfo
 
 		const FVector DeltaLocation = Location - OtherAnimContextInfo.Location;
 		const float DistanceSquared = DeltaLocation.SquaredLength();
-		const float MaxDistance = FMath::Min(GetAvailabilitiesMaxBroadPhaseRadius(), OtherAnimContextInfo.GetAvailabilitiesMaxBroadPhaseRadius());
+		const float MaxDistance = FMath::Min(AvailabilitiesMaxBroadPhaseRadius, OtherAnimContextInfo.AvailabilitiesMaxBroadPhaseRadius);
 		const float MaxDistanceSquared = MaxDistance * MaxDistance;
 		return DistanceSquared <= MaxDistanceSquared;
 	}
 
-	float GetAvailabilitiesMaxBroadPhaseRadius() const
-	{
-		float AvailabilitiesMaxBroadPhaseRadius = 0.f;
-		check(AnimContextAvailabilities);
-		for (const FPoseSearchInteractionAvailabilityEx& Availability : AnimContextAvailabilities->Availabilities)
-		{
-			// @todo: optimize the AvailabilitiesMaxBroadPhaseRadius, since adding Availability.BroadPhaseRadiusIncrementOnInteraction is required ONLY if AnimContext is already part of an interaction
-			AvailabilitiesMaxBroadPhaseRadius = FMath::Max(AvailabilitiesMaxBroadPhaseRadius, Availability.BroadPhaseRadius + Availability.BroadPhaseRadiusIncrementOnInteraction);
-		}
-		return AvailabilitiesMaxBroadPhaseRadius;
-	}
-
 	const FPoseSearchInteractionAnimContextAvailabilities* AnimContextAvailabilities = nullptr;
 
 	// cached AnimContext location
 	FVector Location = FVector::ZeroVector;
+	float AvailabilitiesMaxBroadPhaseRadius = 0.f;
 	TArray<const FAnimContextInfo*, TInlineAllocator<8, TMemStackAllocator<>>> NearbyAnimContextInfos;
 };
 struct FAnimContextInfos : TArray<FAnimContextInfo, TMemStackAllocator<>> {};
@@ -443,13 +439,13 @@ UE::PoseSearch::FInteractionIsland* UPoseSearchInteractionSubsystem::FindIsland(
 	{
 		if (bCompareOwningActors)
 		{
-			const AActor* Actor = GetContextOwningActor(AnimContext);
+			const AActor* Actor = GetContextOwningActor(AnimContext, false);
 
 			for (FInteractionIsland* Island : Islands)
 			{
 				for (const TWeakObjectPtr<const UObject>& IslandAnimContext : Island->GetIslandAnimContexts())
 				{ 
-					if (GetContextOwningActor(IslandAnimContext.Get()) == Actor)
+					if (GetContextOwningActor(IslandAnimContext.Get(), false) == Actor)
 					{
 						return Island;
 					}
@@ -508,55 +504,16 @@ void UPoseSearchInteractionSubsystem::AddAvailabilities(const TArrayView<const F
 
 	check(AnimContext && AnimContext->GetWorld() && AnimContext->GetWorld() == GetWorld());
 
-	// collecting valid availabilities indexes here to minimize the time spent on the lock and avoid locking at all 
-	// if there're actually no valid availabilities (as well as avoiding polluting AnimContextsAvailabilities with an empty entry).
-	// This is unfortunately a common setup, and that's the main reason we optimize against it
-	TArray<int32, TInlineAllocator<64>> ValidAvailabilityIndexes;
-	for (int32 AvailabilityIndex = 0; AvailabilityIndex < Availabilities.Num(); ++AvailabilityIndex)
-	{
-		const FPoseSearchInteractionAvailability& Availability = Availabilities[AvailabilityIndex];
-		if ((Availability.Database && Availability.Database->Schema) || Availability.IsTagValid())
-		{
-			ValidAvailabilityIndexes.Add(AvailabilityIndex);
-		}
-	}
+	const int32 ReservedAnimContextsAvailabilitiesIndex = AnimContextsAvailabilitiesIndex.Add(1);
 
-	if (!ValidAvailabilityIndexes.IsEmpty())
+	if (ReservedAnimContextsAvailabilitiesIndex < AnimContextsAvailabilitiesBuffer.Num())
 	{
-		FScopeLock Lock(&AnimContextsAvailabilitiesMutex);
-
-		const int32 UpperBoundIndex = Algo::UpperBoundBy(AnimContextsAvailabilities, AnimContext,
-			[](const FPoseSearchInteractionAnimContextAvailabilities& AnimContextAvailabilities)
-			{
-				return AnimContextAvailabilities.AnimContext.Get();
-			},
-			[](const UObject* AnimContextA, const UObject* AnimContextB)
-			{
-				return AnimContextA < AnimContextB;
-			});
-
-		int32 AnimContextAvailabilityIndex;
-		if (UpperBoundIndex > 0 && AnimContextsAvailabilities[UpperBoundIndex - 1].AnimContext == AnimContext)
+		FPoseSearchInteractionAnimContextAvailabilities& AnimContextAvailabilities = AnimContextsAvailabilitiesBuffer[ReservedAnimContextsAvailabilitiesIndex];
+		AnimContextAvailabilities.AnimContext = AnimContext;
+		AnimContextAvailabilities.Availabilities.Reset();
+		for (const FPoseSearchInteractionAvailability& Availability : Availabilities)
 		{
-			// we found the AnimContextsAvailabilities for AnimContext
-			AnimContextAvailabilityIndex = UpperBoundIndex - 1;
-		}
-		else
-		{
-			// we create a new AnimContextsAvailabilities for AnimContext, preserving AnimContextsAvailabilities sorting
-			AnimContextsAvailabilities.InsertDefaulted_GetRef(UpperBoundIndex).AnimContext = AnimContext;
-			AnimContextAvailabilityIndex = UpperBoundIndex;
-		}
-
-		for (const int32 AvailabilityIndex : ValidAvailabilityIndexes)
-		{
-			// Avoiding adding trivial duplicates. FPoseSearchInteractionAvailabilityEx could not be fully specified to understand if it's an actual duplicate in case the pose
-			// history is passed by name or the Availability.Database is null and supposed to be resolved using other availabilities Database(s) with the same Availability.Tag.
-			// The duplicated availabilities are excluded when creating the combinations of possible interactions during FAnimContextInfoVisitor when 
-			// FRoledAnimContextInfos.AddRoledAnimContextInfos calls FRoledAnimContextInfos.AddUnique
-			FPoseSearchInteractionAvailabilityEx AvailabilityEx;
-			AvailabilityEx.Init(Availabilities[AvailabilityIndex], PoseHistoryName, PoseHistory);
-			AnimContextsAvailabilities[AnimContextAvailabilityIndex].Availabilities.AddUnique(AvailabilityEx);
+			AnimContextAvailabilities.Availabilities.AddDefaulted_GetRef().Init(Availability, PoseHistoryName, PoseHistory);
 		}
 	}
 }
@@ -571,6 +528,7 @@ void UPoseSearchInteractionSubsystem::GenerateAnimContextInfosAndTagToDatabases(
 	check(SubsystemWorld);
 	check(AnimContextInfos.IsEmpty() && TagToDatabases.IsEmpty());
 
+	TConstArrayView<FPoseSearchInteractionAnimContextAvailabilities> AnimContextsAvailabilities = MakeArrayView(AnimContextsAvailabilitiesBuffer.GetData(), AnimContextsAvailabilitiesNum);
 	check(Algo::IsSorted(AnimContextsAvailabilities, [](const FPoseSearchInteractionAnimContextAvailabilities& A, const FPoseSearchInteractionAnimContextAvailabilities& B)
 		{
 			return A.AnimContext < B.AnimContext;
@@ -971,7 +929,7 @@ void UPoseSearchInteractionSubsystem::OnInteractionStart(UE::PoseSearch::FValidI
 		{
 			if (const UObject* AnimContext = ValidInteractionSearch.SearchContext.GetAnimContext(AnimContextIndex))
 			{
-				AActor* Actor = const_cast<AActor*>(GetContextOwningActor(AnimContext));
+				AActor* Actor = const_cast<AActor*>(GetContextOwningActor(AnimContext, false));
 				check(Actor);
 				Actors.Add(Actor);
 				PrimitiveComponents.Add(Cast<UPrimitiveComponent>(Actor->GetRootComponent()));
@@ -1092,45 +1050,96 @@ void UPoseSearchInteractionSubsystem::UpdateValidInteractionSearches()
 	ValidInteractionSearches = NewValidInteractionSearches;
 }
 
-void UPoseSearchInteractionSubsystem::Tick(float DeltaSeconds)
+bool UPoseSearchInteractionSubsystem::ConsolidateAnimContextsAvailabilities()
 {
-	QUICK_SCOPE_CYCLE_COUNTER(STAT_UPoseSearchInteractionSubsystem_Tick);
-
 	using namespace UE::PoseSearch;
 
-	Super::Tick(DeltaSeconds);
+	const int32 AnimContextsAvailabilitiesIndexValue = AnimContextsAvailabilitiesIndex.GetValue();
+	AnimContextsAvailabilitiesNum = AnimContextsAvailabilitiesBuffer.Num();
 
-	FMemMark Mark(FMemStack::Get());
-	UpdateValidInteractionSearches();
+	// @todo: add some setttings to initialize AnimContextsAvailabilitiesBuffer to a big enough value
+	if (AnimContextsAvailabilitiesIndexValue > AnimContextsAvailabilitiesNum)
+	{
+		UE_LOG(LogPoseSearch, Error, TEXT("UPoseSearchInteractionSubsystem::ConsolidateAnimContextsAvailabilities not enough space to add more availabilities locklessly. It'll be adjusted automatically, but some availability requests has been lost this frame [capacity %d / requests %d]"), AnimContextsAvailabilitiesNum, AnimContextsAvailabilitiesIndexValue);
+		AnimContextsAvailabilitiesBuffer.SetNum(AnimContextsAvailabilitiesIndexValue);
+	}
+	else
+	{
+		AnimContextsAvailabilitiesNum = AnimContextsAvailabilitiesIndexValue;
+	}
 
-	if (AnimContextsAvailabilities.IsEmpty())
+	if (AnimContextsAvailabilitiesNum <= 0)
 	{
-		bool bAllUninjected = true;
 		for (FInteractionIsland* Island : Islands)
 		{
 			if (Island->IsInitialized())
 			{
-				bAllUninjected = false;
+				return true;
 			}
 		}
 
-		if (bAllUninjected)
+		// nothing to do. returning false for a subsequent early out in UPoseSearchInteractionSubsystem::Tick
+		return false;
+	}
+	
+	// consolidating AnimContextsAvailabilities sharing the same AnimInstance
+	TArrayView<FPoseSearchInteractionAnimContextAvailabilities> AnimContextsAvailabilities = MakeArrayView(AnimContextsAvailabilitiesBuffer.GetData(), AnimContextsAvailabilitiesNum);
+	AnimContextsAvailabilities.Sort([](const FPoseSearchInteractionAnimContextAvailabilities& AnimContextAvailabilitiesA, const FPoseSearchInteractionAnimContextAvailabilities& AnimContextAvailabilitiesB)
 		{
-			// nothing to do. early out
-			return;
+			return AnimContextAvailabilitiesA.AnimContext < AnimContextAvailabilitiesB.AnimContext;
+		});
+
+	int32 WriteIndex = 0;
+	for (int32 ReadIndex = 1; ReadIndex < AnimContextsAvailabilitiesNum; ++ReadIndex)
+	{
+		// Avoiding adding trivial duplicates. FPoseSearchInteractionAvailabilityEx could not be fully specified to understand if it's an actual duplicate in case the pose
+		// history is passed by name or the Availability.Database is null and supposed to be resolved using other availabilities Database(s) with the same Availability.Tag.
+		// The duplicated availabilities are excluded when creating the combinations of possible interactions during FAnimContextInfoVisitor when 
+		// FRoledAnimContextInfos.AddRoledAnimContextInfos calls FRoledAnimContextInfos.AddUnique
+		if (AnimContextsAvailabilities[WriteIndex].AnimContext == AnimContextsAvailabilities[ReadIndex].AnimContext)
+		{
+			for (FPoseSearchInteractionAvailabilityEx& AvailabilityEx : AnimContextsAvailabilities[ReadIndex].Availabilities)
+			{
+				AnimContextsAvailabilities[WriteIndex].Availabilities.AddUnique(AvailabilityEx);
+			}
+		}
+		else
+		{
+			++WriteIndex;
+			if (WriteIndex != ReadIndex)
+			{
+				AnimContextsAvailabilities[WriteIndex] = AnimContextsAvailabilities[ReadIndex];
+			}
 		}
 	}
+		
+	AnimContextsAvailabilitiesNum = WriteIndex + 1;	
+	return true;
+}
+
+void UPoseSearchInteractionSubsystem::Tick(float DeltaSeconds)
+{
+	using namespace UE::PoseSearch;
+
+	QUICK_SCOPE_CYCLE_COUNTER(STAT_UPoseSearchInteractionSubsystem_Tick);
 
 	check(IsInGameThread());
 
-	RegenerateAllIslands(DeltaSeconds);
+	Super::Tick(DeltaSeconds);
+
+	FMemMark Mark(FMemStack::Get());
+	UpdateValidInteractionSearches();
 
-	// clearing up AnimContextsAvailabilities for the next frame
-	AnimContextsAvailabilities.Reset();
+	if (ConsolidateAnimContextsAvailabilities())
+	{
+		RegenerateAllIslands(DeltaSeconds);
 
 #if DO_CHECK
-	check(ValidateAllIslands());
+		check(ValidateAllIslands());
 #endif
+	}
+
+	AnimContextsAvailabilitiesIndex.Reset();
 }
 
 TStatId UPoseSearchInteractionSubsystem::GetStatId() const
diff --git a/Engine/Plugins/Animation/PoseSearch/Source/Runtime/Private/PoseSearchInteractionValidator.cpp b/Engine/Plugins/Animation/PoseSearch/Source/Runtime/Private/PoseSearchInteractionValidator.cpp
--- a/Engine/Plugins/Animation/PoseSearch/Source/Runtime/Private/PoseSearchInteractionValidator.cpp
+++ b/Engine/Plugins/Animation/PoseSearch/Source/Runtime/Private/PoseSearchInteractionValidator.cpp
@@ -26,7 +26,7 @@ FInteractionValidator::FInteractionValidator(const UObject* AnimContext)
 		return;
 	}
 
-	const AActor* Actor = GetContextOwningActor(AnimContext);
+	const AActor* Actor = GetContextOwningActor(AnimContext, false);
 	check(Actor);
 
 	const AActor* MainActor = ValidatingIsland->GetMainActor();
@@ -101,7 +101,7 @@ FInteractionValidator::~FInteractionValidator()
 		return;
 	}
 
-	const AActor* Actor = GetContextOwningActor(ValidatingAnimContext);
+	const AActor* Actor = GetContextOwningActor(ValidatingAnimContext, false);
 	check(Actor);
 
 	const AActor* MainActor = Island->GetMainActor();
diff --git a/Engine/Plugins/Animation/PoseSearch/Source/Runtime/Private/PoseSearchLibrary.cpp b/Engine/Plugins/Animation/PoseSearch/Source/Runtime/Private/PoseSearchLibrary.cpp
--- a/Engine/Plugins/Animation/PoseSearch/Source/Runtime/Private/PoseSearchLibrary.cpp
+++ b/Engine/Plugins/Animation/PoseSearch/Source/Runtime/Private/PoseSearchLibrary.cpp
@@ -1379,7 +1379,7 @@ void UPoseSearchLibrary::MotionMatch(
 		{
 			if (const IPoseHistory* PoseHistory = InternalPoseHistories[RoleIndex])
 			{
-				const USkeleton* Skeleton = GetContextSkeleton(Contexts[RoleIndex]);
+				const USkeleton* Skeleton = GetContextSkeleton(Contexts[RoleIndex], false);
 				check(Skeleton);
 				// @todo: add input BlendParameters to support sampling FutureAnimation blendspaces and support for multi character
 				const UAnimationAsset* AnimationAsset = Cast<UAnimationAsset>(Future.Animation);
diff --git a/Engine/Plugins/Animation/PoseSearch/Source/Runtime/Public/PoseSearch/AnimNode_MotionMatching.h b/Engine/Plugins/Animation/PoseSearch/Source/Runtime/Public/PoseSearch/AnimNode_MotionMatching.h
--- a/Engine/Plugins/Animation/PoseSearch/Source/Runtime/Public/PoseSearch/AnimNode_MotionMatching.h
+++ b/Engine/Plugins/Animation/PoseSearch/Source/Runtime/Public/PoseSearch/AnimNode_MotionMatching.h
@@ -3,8 +3,9 @@
 #pragma once
 
 #include "Animation/AnimNode_AssetPlayerBase.h"
-#include "GameplayTagContainer.h"
 #include "BlendStack/AnimNode_BlendStack.h"
+#include "GameplayTagContainer.h"
+#include "PoseSearch/PoseSearchInteractionAvailability.h"
 #include "PoseSearch/PoseSearchLibrary.h"
 #include "AnimNode_MotionMatching.generated.h"
 
@@ -59,6 +60,7 @@ public:
 
 	// FAnimNode_AssetPlayerBase interface
 	UE_API virtual void UpdateAssetPlayer(const FAnimationUpdateContext& Context) override;
+	UE_API virtual void Evaluate_AnyThread(FPoseContext& Output) override;
 	// End of FAnimNode_AssetPlayerBase interface
 
 private:
@@ -137,7 +139,7 @@ private:
 	FFloatInterval PlayRate = FFloatInterval(1.f, 1.f);
 	
 	// Experimental: Multiplier applied to the play rate of the selected animation after Motion Matching State has been updated.
-	UPROPERTY(EditAnywhere, Category = "Settings|Experimental", meta = (PinHiddenByDefault))
+	UPROPERTY(EditAnywhere, Category = "Settings | Experimental", meta = (PinHiddenByDefault))
 	float PlayRateMultiplier = 1.0f;
 	
 	UPROPERTY(EditAnywhere, Category = Settings)
@@ -158,7 +160,7 @@ private:
 	bool bShouldUseCachedChannelData = false;
 	
 	// Experimental, this feature might be removed without warning, not for production use
-	UPROPERTY(EditAnywhere, Category = "Settings|Experimental", meta = (PinHiddenByDefault))
+	UPROPERTY(EditAnywhere, Category = "Settings | Experimental", meta = (PinHiddenByDefault))
 	FPoseSearchEvent EventToSearch;
 
 	// Encapsulated motion matching algorithm and internal state
@@ -178,6 +180,34 @@ private:
 	// True if the Database property on this node has been overridden by SetDatabaseToSearch/SetDatabasesToSearch.
 	bool bOverrideDatabaseInput = false;
 
+	// Experimental, this feature might be removed without warning, not for production use
+	UPROPERTY(EditAnywhere, Category = "Interaction | Experimental", meta = (PinHiddenByDefault))
+	TArray<FPoseSearchInteractionAvailability> Availabilities;
+
+	// Experimental, this feature might be removed without warning, not for production use
+	UPROPERTY(EditAnywhere, Category = "Interaction | Experimental", meta = (PinHiddenByDefault))
+	bool bValidateResultAgainstAvailabilities = true;
+
+	// Experimental, this feature might be removed without warning, not for production use
+	UPROPERTY(EditAnywhere, Category = "Interaction | Experimental", meta = (PinHiddenByDefault))
+	bool bKeepInteractionAlive = false;
+
+	// Experimental, this feature might be removed without warning, not for production use
+	// if bWarpUsingRootBone is true, warping will be calculated using the interacting actors previous frame root bone transforms (effective for setups with OffsetRootBone node allowing root bone drifting from capsule)
+	// if bWarpUsingRootBone is true, warping will be calculated using the previous frame root transforms (effective root motion driven for setups)
+	UPROPERTY(EditAnywhere, Category = "Interaction | Experimental", meta = (PinHiddenByDefault))
+	bool bWarpUsingRootBone = true;
+
+	// Experimental, this feature might be removed without warning, not for production use
+	// amount or translation warping to apply
+	UPROPERTY(EditAnywhere, Category = "Interaction | Experimental", meta = (PinHiddenByDefault, ClampMin = "0", ClampMax = "1"))
+	float WarpingTranslationRatio = 1.f;
+
+	// Experimental, this feature might be removed without warning, not for production use
+	// amount or rotation warping to apply
+	UPROPERTY(EditAnywhere, Category = "Interaction | Experimental", meta = (PinHiddenByDefault, ClampMin = "0", ClampMax = "1"))
+	float WarpingRotationRatio = 1.f;
+
 #if WITH_EDITORONLY_DATA
 	
 	UPROPERTY(meta=(FoldProperty))
@@ -187,6 +217,12 @@ private:
 
 	friend class UAnimGraphNode_MotionMatching;
 	friend class UMotionMatchingAnimNodeLibrary;
+
+private:
+	// Experimental, this feature might be removed without warning, not for production use
+	FTransform MeshWithOffset = FTransform::Identity;
+	// Experimental, this feature might be removed without warning, not for production use
+	FTransform MeshWithoutOffset = FTransform::Identity;
 };
 
 #undef UE_API
diff --git a/Engine/Plugins/Animation/PoseSearch/Source/Runtime/Public/PoseSearch/PoseSearchContext.h b/Engine/Plugins/Animation/PoseSearch/Source/Runtime/Public/PoseSearch/PoseSearchContext.h
--- a/Engine/Plugins/Animation/PoseSearch/Source/Runtime/Public/PoseSearch/PoseSearchContext.h
+++ b/Engine/Plugins/Animation/PoseSearch/Source/Runtime/Public/PoseSearch/PoseSearchContext.h
@@ -23,16 +23,16 @@ namespace UE::PoseSearch
 typedef TSet<const UObject*, DefaultKeyFuncs<const UObject*>, TInlineSetAllocator<8, TMemStackSetAllocator<>>> FStackAssetSet;
 
 // Utils functions
-const FTransform& GetContextTransform(const UObject* AnimContext);
-const USkeleton* GetContextSkeleton(const UObject* AnimContext);
-const USkeleton* GetContextSkeleton(FChooserEvaluationContext& AnimContext);
-const AActor* GetContextOwningActor(const UObject* AnimContext);
-FVector GetContextLocation(const UObject* AnimContext);
+const FTransform& GetContextTransform(const UObject* AnimContext, bool bCheckIsInGameThread = true);
+const USkeleton* GetContextSkeleton(const UObject* AnimContext, bool bCheckIsInGameThread = true);
+const USkeleton* GetContextSkeleton(FChooserEvaluationContext& AnimContext, bool bCheckIsInGameThread = true);
+const AActor* GetContextOwningActor(const UObject* AnimContext, bool bCheckIsInGameThread = true);
+FVector GetContextLocation(const UObject* AnimContext, bool bCheckIsInGameThread = true);
 
 // Experimental, this feature might be removed without warning, not for production use
-const USkeletalMeshComponent* GetContextSkeletalMeshComponent(const UObject* AnimContext);
+const USkeletalMeshComponent* GetContextSkeletalMeshComponent(const UObject* AnimContext, bool bCheckIsInGameThread = true);
 // Experimental, this feature might be removed without warning, not for production use
-const FBoneContainer GetBoneContainer(const UObject* AnimContext);
+const FBoneContainer GetBoneContainer(const UObject* AnimContext, bool bCheckIsInGameThread = true);
 
 struct FPoseIndicesHistory;
 struct IPoseHistory;
diff --git a/Engine/Plugins/Animation/PoseSearch/Source/Runtime/Public/PoseSearch/PoseSearchDatabase.h b/Engine/Plugins/Animation/PoseSearch/Source/Runtime/Public/PoseSearch/PoseSearchDatabase.h
--- a/Engine/Plugins/Animation/PoseSearch/Source/Runtime/Public/PoseSearch/PoseSearchDatabase.h
+++ b/Engine/Plugins/Animation/PoseSearch/Source/Runtime/Public/PoseSearch/PoseSearchDatabase.h
@@ -191,9 +191,9 @@ struct FPoseSearchDatabaseAnimationAsset : public FPoseSearchDatabaseAnimationAs
 	UE_API virtual void IterateOverSamplingParameter(const TFunction<void(const FVector& BlendParameters)>& ProcessSamplingParameter) const override;
 	virtual void SetSamplingRange(const FFloatInterval& NewRange) override { SamplingRange = NewRange; }
 
-	UE_DEPRECATED(5.6, "Use IterateOverSamplingParameter to iterate over all the sampling paramters for the blend space")
+	UE_DEPRECATED(5.6, "Use IterateOverSamplingParameter to iterate over all the sampling parameters for the blend space")
 	UE_API void GetBlendSpaceParameterSampleRanges(int32& HorizontalBlendNum, int32& VerticalBlendNum) const;
-	UE_DEPRECATED(5.6, "Use IterateOverSamplingParameter to iterate over all the sampling paramters for the blend space")
+	UE_DEPRECATED(5.6, "Use IterateOverSamplingParameter to iterate over all the sampling parameters for the blend space")
 	UE_API FVector BlendParameterForSampleRanges(int32 HorizontalBlendIndex, int32 VerticalBlendIndex) const;
 #endif // WITH_EDITORONLY_DATA
 
@@ -326,9 +326,9 @@ struct UE_DEPRECATED(5.7, "Use FPoseSearchDatabaseAnimationAsset instead. To con
 	UE_API virtual void IterateOverSamplingParameter(const TFunction<void(const FVector& BlendParameters)>& ProcessSamplingParameter) const override;
 	virtual void SetSamplingRange(const FFloatInterval& NewRange) override { SamplingRange = NewRange; }
 
-	UE_DEPRECATED(5.6, "Use IterateOverSamplingParameter to iterate over all the sampling paramters for the blend space")
+	UE_DEPRECATED(5.6, "Use IterateOverSamplingParameter to iterate over all the sampling parameters for the blend space")
 	UE_API void GetBlendSpaceParameterSampleRanges(int32& HorizontalBlendNum, int32& VerticalBlendNum) const;
-	UE_DEPRECATED(5.6, "Use IterateOverSamplingParameter to iterate over all the sampling paramters for the blend space")
+	UE_DEPRECATED(5.6, "Use IterateOverSamplingParameter to iterate over all the sampling parameters for the blend space")
 	UE_API FVector BlendParameterForSampleRanges(int32 HorizontalBlendIndex, int32 VerticalBlendIndex) const;
 #endif // WITH_EDITORONLY_DATA
 
diff --git a/Engine/Plugins/Animation/PoseSearch/Source/Runtime/Public/PoseSearch/PoseSearchInteractionLibrary.h b/Engine/Plugins/Animation/PoseSearch/Source/Runtime/Public/PoseSearch/PoseSearchInteractionLibrary.h
--- a/Engine/Plugins/Animation/PoseSearch/Source/Runtime/Public/PoseSearch/PoseSearchInteractionLibrary.h
+++ b/Engine/Plugins/Animation/PoseSearch/Source/Runtime/Public/PoseSearch/PoseSearchInteractionLibrary.h
@@ -34,6 +34,7 @@ public:
 	static UE_API FPoseSearchBlueprintResult MotionMatchInteraction(TArray<FPoseSearchInteractionAvailability> Availabilities, const UObject* AnimContext, FName PoseHistoryName, bool bValidateResultAgainstAvailabilities = true);
 
 	static UE_API void MotionMatchInteraction(FPoseSearchBlueprintResult& Result, const TArrayView<const FPoseSearchInteractionAvailability> Availabilities, const UObject* AnimContext, FName PoseHistoryName, const UE::PoseSearch::IPoseHistory* PoseHistory, bool bValidateResultAgainstAvailabilities = true);
+	static UE_API bool MotionMatchInteraction(FPoseSearchBlueprintResult& InOutResult, const TArrayView<const FPoseSearchInteractionAvailability> Availabilities, const UObject* AnimContext, FName PoseHistoryName, const UE::PoseSearch::IPoseHistory* PoseHistory, bool bValidateResultAgainstAvailabilities, bool bKeepInteractionAlive, float BlendTime);
 
 	static UE_API FPoseSearchBlueprintResult GetCachedInteraction(const UObject* AnimContext, bool bCompareOwningActors = false);
 	
diff --git a/Engine/Plugins/Animation/PoseSearch/Source/Runtime/Public/PoseSearch/PoseSearchInteractionSubsystem.h b/Engine/Plugins/Animation/PoseSearch/Source/Runtime/Public/PoseSearch/PoseSearchInteractionSubsystem.h
--- a/Engine/Plugins/Animation/PoseSearch/Source/Runtime/Public/PoseSearch/PoseSearchInteractionSubsystem.h
+++ b/Engine/Plugins/Animation/PoseSearch/Source/Runtime/Public/PoseSearch/PoseSearchInteractionSubsystem.h
@@ -95,40 +95,40 @@ private:
 	friend struct UE::PoseSearch::FInteractionValidator;
 #endif // ENABLE_ANIM_DEBUG
 
-	UE_API UE::PoseSearch::FInteractionIsland& CreateIsland();
-	UE_API UE::PoseSearch::FInteractionIsland& GetAvailableIsland();
+	UE::PoseSearch::FInteractionIsland& CreateIsland();
+	UE::PoseSearch::FInteractionIsland& GetAvailableIsland();
 
-	UE_API void DestroyIsland(int32 Index);
-	UE_API void DestroyAllIslands();
-	UE_API void RegenerateAllIslands(float DeltaSeconds);
+	void DestroyIsland(int32 Index);
+	void DestroyAllIslands();
+	void RegenerateAllIslands(float DeltaSeconds);
+	bool ConsolidateAnimContextsAvailabilities();
 
 #if DO_CHECK
-	UE_API bool ValidateAllIslands() const;
+	bool ValidateAllIslands() const;
 #endif // DO_CHECK
 
-	UE_API void PopulateContinuingProperties(float DeltaSeconds, TArrayView<UE::PoseSearch::FInteractionSearchContext> SearchContexts) const;
-	UE_API UE::PoseSearch::FInteractionIsland* FindIsland(const UObject* AnimContext, bool bCompareOwningActors = false);
+	void PopulateContinuingProperties(float DeltaSeconds, TArrayView<UE::PoseSearch::FInteractionSearchContext> SearchContexts) const;
+	UE::PoseSearch::FInteractionIsland* FindIsland(const UObject* AnimContext, bool bCompareOwningActors = false);
 	
 #if ENABLE_ANIM_DEBUG
-	UE_API void DebugDrawIslands() const;
-	UE_API void DebugLogTickDependencies() const;
+	void DebugDrawIslands() const;
+	void DebugLogTickDependencies() const;
 #endif // ENABLE_ANIM_DEBUG
 
-	UE_API void AddAvailabilities(const TArrayView<const FPoseSearchInteractionAvailability> Availabilities, const UObject* AnimContext, FName PoseHistoryName, const UE::PoseSearch::IPoseHistory* PoseHistory);
-	UE_API void GenerateAnimContextInfosAndTagToDatabases(UE::PoseSearch::FAnimContextInfos& AnimContextInfos, UE::PoseSearch::FTagToDatabases& TagToDatabases) const;
-	UE_API void GenerateSearchContexts(float DeltaSeconds, UE::PoseSearch::FInteractionSearchContexts& SearchContexts) const;
+	void AddAvailabilities(const TArrayView<const FPoseSearchInteractionAvailability> Availabilities, const UObject* AnimContext, FName PoseHistoryName, const UE::PoseSearch::IPoseHistory* PoseHistory);
+	void GenerateAnimContextInfosAndTagToDatabases(UE::PoseSearch::FAnimContextInfos& AnimContextInfos, UE::PoseSearch::FTagToDatabases& TagToDatabases) const;
+	void GenerateSearchContexts(float DeltaSeconds, UE::PoseSearch::FInteractionSearchContexts& SearchContexts) const;
 
-	UE_API void UpdateValidInteractionSearches();
-	static UE_API void OnInteractionStart(UE::PoseSearch::FValidInteractionSearch& ValidInteractionSearch);
-	static UE_API void OnInteractionContinuing(UE::PoseSearch::FValidInteractionSearch& ValidInteractionSearch);
-	static UE_API void OnInteractionEnd(UE::PoseSearch::FValidInteractionSearch& ValidInteractionSearch);
+	void UpdateValidInteractionSearches();
+	static void OnInteractionStart(UE::PoseSearch::FValidInteractionSearch& ValidInteractionSearch);
+	static void OnInteractionContinuing(UE::PoseSearch::FValidInteractionSearch& ValidInteractionSearch);
+	static void OnInteractionEnd(UE::PoseSearch::FValidInteractionSearch& ValidInteractionSearch);
 
-	// @todo: use a thread local storage data structure to make publishing availabilities lockless to avoid the need of AnimContextsAvailabilitiesMutex
 	// it could contain duplicated availabilities
 	UPROPERTY(Transient)
-	TArray<FPoseSearchInteractionAnimContextAvailabilities> AnimContextsAvailabilities;
-
-	FCriticalSection AnimContextsAvailabilitiesMutex;
+	TArray<FPoseSearchInteractionAnimContextAvailabilities> AnimContextsAvailabilitiesBuffer;
+	FThreadSafeCounter AnimContextsAvailabilitiesIndex = 0;
+	int32 AnimContextsAvailabilitiesNum = 0;
 
 	// array of groups of characters that needs to be anaylzed together for possible interactions
 	TArray<UE::PoseSearch::FInteractionIsland*> Islands;
diff --git a/Engine/Plugins/Experimental/UAF/UAFPoseSearch/Source/UAFPoseSearch/Private/MotionMatchingTrait.cpp b/Engine/Plugins/Experimental/UAF/UAFPoseSearch/Source/UAFPoseSearch/Private/MotionMatchingTrait.cpp
--- a/Engine/Plugins/Experimental/UAF/UAFPoseSearch/Source/UAFPoseSearch/Private/MotionMatchingTrait.cpp
+++ b/Engine/Plugins/Experimental/UAF/UAFPoseSearch/Source/UAFPoseSearch/Private/MotionMatchingTrait.cpp
@@ -189,72 +189,12 @@ PRAGMA_ENABLE_DEPRECATION_WARNINGS
 		}
 #endif
 
-		bool bIsInteraction = false;
 		const TArrayView<const FPoseSearchInteractionAvailability> Availabilities = SharedData->GetAvailabilities(Binding);
-		if (!Availabilities.IsEmpty())
-		{
-			// WIP motion matched interactions
-			CheckInteractionThreadSafety(AnimContext);
-
-			FPoseSearchBlueprintResult SearchResult;
-			UPoseSearchInteractionLibrary::MotionMatchInteraction(SearchResult, Availabilities, AnimContext, FName(), PoseHistory, SharedData->GetbValidateResultAgainstAvailabilities(Binding));
-			check(SearchResult.ActorRootTransforms.Num() == SearchResult.ActorRootBoneTransforms.Num());
-
-			if (SearchResult.SelectedAnim)
-			{
-				bIsInteraction = true;
-				MotionMatchingState.SearchResult = SearchResult;
-			}
-		}
+		const bool bIsInteraction = UPoseSearchInteractionLibrary::MotionMatchInteraction(MotionMatchingState.SearchResult, Availabilities, AnimContext, FName(), PoseHistory, SharedData->GetbValidateResultAgainstAvailabilities(Binding), SharedData->GetbKeepInteractionAlive(Binding), SharedData->BlendArguments.BlendTime);
 
 		// performing the regular single character motion matching search in case there's no MM interaction
 		if (!bIsInteraction)
 		{
-			const bool bKeepInteractionAlive = SharedData->GetbKeepInteractionAlive(Binding);
-			if (!MotionMatchingState.SearchResult.bIsInteraction)
-			{
-				// @todo: ask blendstack what it's playing instead of relying on MotionMatchingInteractionState.SelectedAnim
-				//MotionMatchingState.SearchResult.SelectedAnim = GetAnimAsset();
-				MotionMatchingState.SearchResult.SelectedTime = TimelineNormalizedTime;
-				// @todo: ask blendstack what the blend parameters are
-				//MotionMatchingState.SearchResult.BlendParameters = GetBlendParameters();
-			}
-			else if (bKeepInteractionAlive)
-			{
-				// @todo: ask blendstack what it's playing instead of relying on MotionMatchingInteractionState.SelectedAnim
-				//MotionMatchingState.SearchResult.SelectedAnim = GetAnimAsset();
-				MotionMatchingState.SearchResult.SelectedTime = TimelineNormalizedTime;
-				// @todo: ask blendstack what the blend parameters are
-				//MotionMatchingState.SearchResult.BlendParameters = GetBlendParameters();
-
-				// checking if the kept alive interaction has reached the end of animation 
-				if (!MotionMatchingState.SearchResult.SelectedDatabase || INDEX_NONE == MotionMatchingState.SearchResult.SelectedDatabase->GetPoseIndex(MotionMatchingState.SearchResult.SelectedAnim.Get(), MotionMatchingState.SearchResult.SelectedTime, MotionMatchingState.SearchResult.bIsMirrored, MotionMatchingState.SearchResult.BlendParameters))
-				{
-					MotionMatchingState.SearchResult = FPoseSearchBlueprintResult();
-				}
-				else
-				{
-					// letting the interaction animation run until its length minus blend time
-					// (to avoid having to blend from a frozen animation that reached its end for the entire duration of the blend)
-					const UAnimationAsset* AnimationAssetForRole = MotionMatchingState.SearchResult.GetAnimationAssetForRole();
-					if (!AnimationAssetForRole || MotionMatchingState.SearchResult.SelectedTime >= (AnimationAssetForRole->GetPlayLength() - SharedData->BlendArguments.BlendTime))
-					{
-						MotionMatchingState.SearchResult = FPoseSearchBlueprintResult();
-					}
-					else
-					{
-						// we're keeping alive only the animation part of the search result
-						MotionMatchingState.SearchResult.ActorRootTransforms.Reset();
-						MotionMatchingState.SearchResult.ActorRootBoneTransforms.Reset();
-						MotionMatchingState.SearchResult.AnimContexts.Reset();
-					}
-				}
-			}
-			else
-			{
-				MotionMatchingState.SearchResult = FPoseSearchBlueprintResult();
-			}
-
 			const TConstArrayView<TObjectPtr<const UPoseSearchDatabase>> DatabasesToSearch = SharedData->GetDatabases(Binding);
 			if (DatabasesToSearch.IsEmpty())
 			{
@@ -334,12 +274,7 @@ PRAGMA_ENABLE_DEPRECATION_WARNINGS
 				PRAGMA_DISABLE_DEPRECATION_WARNINGS
 				MotionMatchingState.CurrentSearchResult = SearchResult;
 				MotionMatchingState.WantedPlayRate = WantedPlayRate;
-				PRAGMA_ENABLE_DEPRECATION_WARNINGS
-	
-				const bool bJumpToPose = MotionMatchingState.SearchResult.SelectedAnim && !MotionMatchingState.SearchResult.bIsContinuingPoseSearch;
-
-				PRAGMA_DISABLE_DEPRECATION_WARNINGS
-				MotionMatchingState.bJumpedToPose = bJumpToPose;
+				MotionMatchingState.bJumpedToPose = MotionMatchingState.SearchResult.SelectedAnim && !MotionMatchingState.SearchResult.bIsContinuingPoseSearch;
 				PRAGMA_ENABLE_DEPRECATION_WARNINGS
 			}
 		}
@@ -448,7 +383,7 @@ PRAGMA_ENABLE_DEPRECATION_WARNINGS
 			ActiveGraphInstance->AccessVariablesStruct<FBlendSpacePlayerData>([&SharedData, &Binding, DesiredPlayRate](FBlendSpacePlayerData& InBlendSpacePlayer)
 			{
 				InBlendSpacePlayer.PlayRate = DesiredPlayRate;
-				// always updating the blend paramters if bUpdateBlendSpaceInputs
+				// always updating the blend parameters if bUpdateBlendSpaceInputs
 				if (SharedData->GetbUpdateBlendSpaceInputs(Binding))
 				{
 					InBlendSpacePlayer.XAxisSamplePoint = SharedData->GetXAxisSamplePoint(Binding);
