Chooser - FChooserTableEditor::AutoPopulateColumn support for ColumnWidget_SpecialIndex_Fallback

diff --git a/Engine/Plugins/Animation/PoseSearch/Source/Runtime/Private/Chooser/PoseSearchChooserColumn.cpp b/Engine/Plugins/Animation/PoseSearch/Source/Runtime/Private/Chooser/PoseSearchChooserColumn.cpp
--- a/Engine/Plugins/Animation/PoseSearch/Source/Runtime/Private/Chooser/PoseSearchChooserColumn.cpp
+++ b/Engine/Plugins/Animation/PoseSearch/Source/Runtime/Private/Chooser/PoseSearchChooserColumn.cpp
@@ -110,11 +110,15 @@ void FPoseSearchColumn::CheckConsistency() const
 		check(DatabaseAnimationAsset);
 		
 		const UObject* AnimationAsset = DatabaseAnimationAsset->GetAnimationAsset();
-		const UObject* ReferencedObject = GetReferencedObject(GetRowIndex(AnimationAssetIndex));
-		check(AnimationAsset == ReferencedObject);
+		const int32 RowIndex = GetRowIndex(AnimationAssetIndex);
+		const UObject* ReferencedObject = GetReferencedObject(RowIndex);
+		if (AnimationAsset != ReferencedObject)
+		{
+			UE_LOG(LogPoseSearch, Error, TEXT("FPoseSearchColumn::CheckConsistency failed for row %d: chooser table asset '%s' differs from internal database asset '%s'. Click 'AutoPopulate All' button in the Chooser Table Editor to try resynchronize the data"), RowIndex, *GetNameSafe(ReferencedObject), *GetNameSafe(AnimationAsset));
+		}
 	}
 
-#endif // DO_CHECK && WITH_EDITORONLY_DATA
+#endif // DO_CHECK && WITH_EDITORONLY_DATA 
 }
 
 #if WITH_EDITOR
diff --git a/Engine/Plugins/Chooser/Source/ChooserEditor/Private/ChooserTableEditor.cpp b/Engine/Plugins/Chooser/Source/ChooserEditor/Private/ChooserTableEditor.cpp
--- a/Engine/Plugins/Chooser/Source/ChooserEditor/Private/ChooserTableEditor.cpp
+++ b/Engine/Plugins/Chooser/Source/ChooserEditor/Private/ChooserTableEditor.cpp
@@ -1924,14 +1924,12 @@ void FChooserTableEditor::AutoPopulateColumn(FChooserColumnBase& Column)
 		{
 			for (int i = 0; i < RowCount; ++i)
 			{
-				if (Chooser->ResultsStructs[i].IsValid())
-				{
-					if (UObject* ReferencedObject = Chooser->ResultsStructs[i].Get<FObjectChooserBase>().GetReferencedObject()) 
-					{
-						Column.AutoPopulate(i, ReferencedObject);
-					}
-				}
+				UObject* ReferencedObject = Chooser->ResultsStructs[i].IsValid() ? Chooser->ResultsStructs[i].Get<FObjectChooserBase>().GetReferencedObject() : nullptr;
+				Column.AutoPopulate(i, ReferencedObject);
 			}
+
+			UObject* ReferencedObject = Chooser->FallbackResult.IsValid() ? Chooser->FallbackResult.Get<FObjectChooserBase>().GetReferencedObject() : nullptr;
+			Column.AutoPopulate(ColumnWidget_SpecialIndex_Fallback, ReferencedObject);
 		}
 	}
 }
