PoseSearch - fix for FPoseSearchColumn::CheckConsistency reporting mismatch between Chooser Table rows and InternalDatabase number of assets without triggering a check

diff --git a/Engine/Plugins/Animation/PoseSearch/Source/Runtime/Private/Chooser/PoseSearchChooserColumn.cpp b/Engine/Plugins/Animation/PoseSearch/Source/Runtime/Private/Chooser/PoseSearchChooserColumn.cpp
--- a/Engine/Plugins/Animation/PoseSearch/Source/Runtime/Private/Chooser/PoseSearchChooserColumn.cpp
+++ b/Engine/Plugins/Animation/PoseSearch/Source/Runtime/Private/Chooser/PoseSearchChooserColumn.cpp
@@ -102,7 +102,10 @@ void FPoseSearchColumn::CheckConsistency() const
 
 	const int32 NumAnimationAssets = InternalDatabase->GetNumAnimationAssets();
 	const int32 NumRows = GetRowIndex(NumAnimationAssets);
-	check(NumRows == OuterChooser->ResultsStructs.Num());
+	if (NumRows != OuterChooser->ResultsStructs.Num())
+	{
+		UE_LOG(LogPoseSearch, Error, TEXT("FPoseSearchColumn::CheckConsistency for Chooser Table '%s' failed because the number of rows (%d) differs from the internal database number of assets %d. Click 'AutoPopulate All' button in the Chooser Table Editor to try resynchronize the data"), *OuterChooser->GetName(), OuterChooser->ResultsStructs.Num(), NumRows);
+	}
 
 	for (int32 AnimationAssetIndex = 0; AnimationAssetIndex < NumAnimationAssets; ++AnimationAssetIndex)
 	{
@@ -114,7 +117,7 @@ void FPoseSearchColumn::CheckConsistency() const
 		const UObject* ReferencedObject = GetReferencedObject(RowIndex);
 		if (AnimationAsset != ReferencedObject)
 		{
-			UE_LOG(LogPoseSearch, Error, TEXT("FPoseSearchColumn::CheckConsistency failed for row %d: chooser table asset '%s' differs from internal database asset '%s'. Click 'AutoPopulate All' button in the Chooser Table Editor to try resynchronize the data"), RowIndex, *GetNameSafe(ReferencedObject), *GetNameSafe(AnimationAsset));
+			UE_LOG(LogPoseSearch, Error, TEXT("FPoseSearchColumn::CheckConsistency for Chooser Table '%s' failed for row %d: chooser table asset '%s' differs from internal database asset ' % s'. Click 'AutoPopulate All' button in the Chooser Table Editor to try resynchronize the data"), *OuterChooser->GetName(), RowIndex, *GetNameSafe(ReferencedObject), *GetNameSafe(AnimationAsset));
 		}
 	}
 
