Fix for pose history and trajectory history desync by delta time, causing incorrect results especially noticable when the component transform makes large translation and rotation, e.g. a 180 pivot. This fixes a bug with bad pivot selection when the capsule rotation flips
Fix for incorrect timings on trajectory history when delta time is not constant
Fix for duplicate trajectory entries around t = deltatime

Since this code change causes some variables to become unused in a variant of UpdateHistory_TransformHistory, it has been marked as deprecated.
The signature of experimental function UpdateHistory_WorldSpace has been changed

#tests animsandbox

#virtualized

diff --git a/Engine/Plugins/Animation/PoseSearch/Source/Runtime/Private/PoseSearchTrajectoryLibrary.cpp b/Engine/Plugins/Animation/PoseSearch/Source/Runtime/Private/PoseSearchTrajectoryLibrary.cpp
--- a/Engine/Plugins/Animation/PoseSearch/Source/Runtime/Private/PoseSearchTrajectoryLibrary.cpp
+++ b/Engine/Plugins/Animation/PoseSearch/Source/Runtime/Private/PoseSearchTrajectoryLibrary.cpp
@@ -277,8 +277,7 @@ void UPoseSearchTrajectoryLibrary::UpdateHistory_TransformHistory(FTransformTraj
                                                                   FVector CurrentPosition,
                                                                   FVector CurrentVelocity,
                                                                   const FPoseSearchTrajectoryData::FSampling& TrajectoryDataSampling,
-                                                                  float DeltaTime,
-                                                                  float CurrentTime)
+                                                                  float DeltaTime)
 {
 	const int32 NumHistorySamples = TrajectoryDataSampling.NumHistorySamples;
 	if (NumHistorySamples > 0)
@@ -303,9 +302,9 @@ void UPoseSearchTrajectoryLibrary::UpdateHistory_TransformHistory(FTransformTraj
 			}
 
 			// Adding a new history record
-			// Note we add current frame at delta time in the future since t = 0 should be the previous frame
-			Trajectory.Samples[NumHistorySamples - 1].TimeInSeconds = CurrentTime;
-			Trajectory.Samples[NumHistorySamples - 1].Position = CurrentPosition;
+			// Copy over the last frame's current transform (stored at i == NumHistorySamples) into a sample at t = 0
+			Trajectory.Samples[NumHistorySamples - 1].TimeInSeconds = 0.0f;
+			Trajectory.Samples[NumHistorySamples - 1].Position = Trajectory.Samples[NumHistorySamples].Position;
 			Trajectory.Samples[NumHistorySamples - 1].Facing = Trajectory.Samples[NumHistorySamples].Facing;
 		}
 		else
@@ -322,11 +321,8 @@ void UPoseSearchTrajectoryLibrary::UpdateHistory_TransformHistory(FTransformTraj
 }
 
 void UPoseSearchTrajectoryLibrary::UpdateHistory_WorldSpace(FTransformTrajectory& Trajectory,
-                                                            FVector CurrentPosition,
-                                                            FQuat CurrentRotation,
                                                             const FPoseSearchTrajectoryData::FSampling& TrajectoryDataSampling,
-                                                            float DeltaTime,
-                                                            float CurrentTime)
+                                                            float DeltaTime)
 {
 	const int32 NumHistorySamples = TrajectoryDataSampling.NumHistorySamples;
 	if (NumHistorySamples > 0)
@@ -347,9 +343,9 @@ void UPoseSearchTrajectoryLibrary::UpdateHistory_WorldSpace(FTransformTrajectory
 
 			// Adding a new history record
 			// Note we add current frame at delta time in the future since t = 0 should be the previous frame
-			Trajectory.Samples[NumHistorySamples - 1].TimeInSeconds = CurrentTime;
-			Trajectory.Samples[NumHistorySamples - 1].Position = CurrentPosition;
-			Trajectory.Samples[NumHistorySamples - 1].Facing = CurrentRotation;
+			Trajectory.Samples[NumHistorySamples - 1].TimeInSeconds = 0.0f;
+			Trajectory.Samples[NumHistorySamples - 1].Position = Trajectory.Samples[NumHistorySamples].Position;
+			Trajectory.Samples[NumHistorySamples - 1].Facing = Trajectory.Samples[NumHistorySamples].Facing;
 		}
 		else
 		{
diff --git a/Engine/Plugins/Animation/PoseSearch/Source/Runtime/Public/PoseSearch/PoseSearchTrajectoryLibrary.h b/Engine/Plugins/Animation/PoseSearch/Source/Runtime/Public/PoseSearch/PoseSearchTrajectoryLibrary.h
--- a/Engine/Plugins/Animation/PoseSearch/Source/Runtime/Public/PoseSearch/PoseSearchTrajectoryLibrary.h
+++ b/Engine/Plugins/Animation/PoseSearch/Source/Runtime/Public/PoseSearch/PoseSearchTrajectoryLibrary.h
@@ -188,49 +188,31 @@ public:
 	// that to the current world transform. This works well on moving platforms as it only stores a history of movement
 	// that results from character intent, not movement from platforms.
 	// Important: CurrentVelocity should be the velocity relative to the ground as reported by the character movement component or character mover etc
-	static UE_API void UpdateHistory_TransformHistory(FTransformTrajectory& Trajectory,
+	UE_DEPRECATED(5.7, "CurrentTime is no longer needed or used. Use UpdateHistory_TransformHistory with no CurrentTime")
+	static void UpdateHistory_TransformHistory(FTransformTrajectory& Trajectory,
 	                                           FVector CurrentPosition,
 	                                           FVector CurrentVelocity,
 	                                           const FPoseSearchTrajectoryData::FSampling& TrajectoryDataSampling,
 	                                           float DeltaTime,
-	                                           float CurrentTime);
+	                                           float CurrentTime)
+	{
+		UpdateHistory_TransformHistory(Trajectory, CurrentPosition, CurrentVelocity, TrajectoryDataSampling, DeltaTime);
+	}
 
-	static inline void UpdateHistory_TransformHistory(FTransformTrajectory& Trajectory,
+	static UE_API void UpdateHistory_TransformHistory(FTransformTrajectory& Trajectory,
 	                                                       FVector CurrentPosition,
 	                                                       FVector CurrentVelocity,
 	                                                       const FPoseSearchTrajectoryData::FSampling& TrajectoryDataSampling,
-	                                                       float DeltaTime)
-	{
-		// Old behaviour: assume this is called after position has been updated (e.g. post Character Movement), so use DeltaTime for current time
-		UpdateHistory_TransformHistory(Trajectory, CurrentPosition, CurrentVelocity, TrajectoryDataSampling, DeltaTime, DeltaTime);
-	}
+	                                                       float DeltaTime);
 
 	// Update prediction by simulating the movement math for ground locomotion from UCharacterMovementComponent.
 	static UE_API void UpdatePrediction_SimulateCharacterMovement(FTransformTrajectory& Trajectory, const FPoseSearchTrajectoryData& TrajectoryData, const FPoseSearchTrajectoryData::FDerived& TrajectoryDataDerived, const FPoseSearchTrajectoryData::FSampling& TrajectoryDataSampling, float DeltaTime);
 
 	// Experimental: Update the history purely based on current position, without taking into account ground velocities
-	// DeltaTime is the time since the last call to this method
-	UE_EXPERIMENTAL(5.6, "Experimental pose search function, API may change in the future")
+	UE_EXPERIMENTAL(5.7, "Experimental pose search function, API may change in the future")
 	static UE_API void UpdateHistory_WorldSpace(FTransformTrajectory& Trajectory,
-	                                     FVector CurrentPosition,
-	                                     FQuat CurrentRotation,
 	                                     const FPoseSearchTrajectoryData::FSampling& TrajectoryDataSampling,
-	                                     float DeltaTime,
-	                                     float CurrentTime);
-
-	// Experimental: Update the history purely based on current position, without taking into account ground velocities
-	UE_EXPERIMENTAL(5.6, "Experimental pose search function, API may change in the future")
-	static void UpdateHistory_WorldSpace(FTransformTrajectory& Trajectory,
-	                                     FVector CurrentPosition,
-	                                     FQuat CurrentRotation,
-	                                     const FPoseSearchTrajectoryData::FSampling& TrajectoryDataSampling,
-	                                     float DeltaTime)
-	{
-		// Old behaviour: assume this is called after position has been updated (e.g. post CMC)
-		PRAGMA_DISABLE_EXPERIMENTAL_WARNINGS
-		UpdateHistory_WorldSpace(Trajectory, CurrentPosition, CurrentRotation, TrajectoryDataSampling, DeltaTime, DeltaTime);
-		PRAGMA_ENABLE_EXPERIMENTAL_WARNINGS
-	}
+	                                     float DeltaTime);
 
 	// Generates a prediction trajectory based of the current character intent. For use with Character actors.
 	UFUNCTION(BlueprintCallable, Category = "Animation|PoseSearch", meta = (BlueprintThreadSafe, DisplayName = "Pose Search Generate Trajectory (for Character)"))
