Correct comments based on post submit analysis

diff --git a/Engine/Plugins/Animation/PoseSearch/Source/Runtime/Private/PoseSearchTrajectoryLibrary.cpp b/Engine/Plugins/Animation/PoseSearch/Source/Runtime/Private/PoseSearchTrajectoryLibrary.cpp
--- a/Engine/Plugins/Animation/PoseSearch/Source/Runtime/Private/PoseSearchTrajectoryLibrary.cpp
+++ b/Engine/Plugins/Animation/PoseSearch/Source/Runtime/Private/PoseSearchTrajectoryLibrary.cpp
@@ -284,7 +284,9 @@ void UPoseSearchTrajectoryLibrary::UpdateHistory_TransformHistory(FTransformTraj
 	{
 		const float SecondsPerHistorySample = TrajectoryDataSampling.SecondsPerHistorySample;
 
-		check(NumHistorySamples <= Trajectory.Samples.Num());
+		// Trajectory should include room for history + current + future
+		// So num history samples needs to be less than the total number
+		check(NumHistorySamples < Trajectory.Samples.Num());
 
 		// Trajectory.Samples[NumHistorySamples] is last frame position! (assuming this is called every frame)
 		const FVector CurrentTranslationFromMover = CurrentVelocity * DeltaTime;
@@ -329,7 +331,9 @@ void UPoseSearchTrajectoryLibrary::UpdateHistory_WorldSpace(FTransformTrajectory
 	{
 		const float SecondsPerHistorySample = TrajectoryDataSampling.SecondsPerHistorySample;
 
-		check(NumHistorySamples <= Trajectory.Samples.Num());
+		// Trajectory should include room for history + current + future
+		// So num history samples needs to be less than the total number
+		check(NumHistorySamples < Trajectory.Samples.Num());
 
 		// Shift history Samples when it's time to record a new one.
 		if (SecondsPerHistorySample <= 0.f || FMath::Abs(Trajectory.Samples[NumHistorySamples - 1].TimeInSeconds) >= SecondsPerHistorySample)
@@ -342,7 +346,7 @@ void UPoseSearchTrajectoryLibrary::UpdateHistory_WorldSpace(FTransformTrajectory
 			}
 
 			// Adding a new history record
-			// Note we add current frame at delta time in the future since t = 0 should be the previous frame
+			// Copy over the last frame's current transform (stored at i == NumHistorySamples) into a sample at t = 0
 			Trajectory.Samples[NumHistorySamples - 1].TimeInSeconds = 0.0f;
 			Trajectory.Samples[NumHistorySamples - 1].Position = Trajectory.Samples[NumHistorySamples].Position;
 			Trajectory.Samples[NumHistorySamples - 1].Facing = Trajectory.Samples[NumHistorySamples].Facing;
