PoseSearch - PoseSearchColumn AddToDetails / SetFromDetails indexing fix

diff --git a/Engine/Plugins/Animation/PoseSearch/Source/Runtime/Private/Chooser/PoseSearchChooserColumn.cpp b/Engine/Plugins/Animation/PoseSearch/Source/Runtime/Private/Chooser/PoseSearchChooserColumn.cpp
--- a/Engine/Plugins/Animation/PoseSearch/Source/Runtime/Private/Chooser/PoseSearchChooserColumn.cpp
+++ b/Engine/Plugins/Animation/PoseSearch/Source/Runtime/Private/Chooser/PoseSearchChooserColumn.cpp
@@ -740,6 +740,8 @@ void FPoseSearchColumn::Initialize(UChooserTable* OuterChooser)
 
 void FPoseSearchColumn::AddToDetails(FInstancedPropertyBag& PropertyBag, int32 ColumnIndex, int32 RowIndex)
 {
+	using namespace UE::PoseSearch;
+
 	check(InternalDatabase);
 
 	FText DisplayName = NSLOCTEXT("PoseSearchColumn","Pose Search", "Pose Search");
@@ -748,7 +750,8 @@ void FPoseSearchColumn::AddToDetails(FInstancedPropertyBag& PropertyBag, int32 C
 	PropertyDesc.MetaData.Add(FPropertyBagPropertyDescMetaData("DisplayName", DisplayName.ToString()));
 	PropertyBag.AddProperties({PropertyDesc});
 		
-	if (const FPoseSearchDatabaseAnimationAsset* Asset = InternalDatabase->GetDatabaseAnimationAsset(RowIndex))
+	const int32 DatabaseAssetIndex = GetDatabaseAssetIndex(RowIndex);
+	if (const FPoseSearchDatabaseAnimationAsset* Asset = InternalDatabase->GetDatabaseAnimationAsset(DatabaseAssetIndex))
 	{
 		PropertyBag.SetValueStruct(PropertyName, *Asset);
 	}
@@ -756,6 +759,8 @@ void FPoseSearchColumn::AddToDetails(FInstancedPropertyBag& PropertyBag, int32 C
 
 void FPoseSearchColumn::SetFromDetails(FInstancedPropertyBag& PropertyBag, int32 ColumnIndex, int32 RowIndex)
 {
+	using namespace UE::PoseSearch;
+
 	check(InternalDatabase);
 
 	FName PropertyName("RowData", ColumnIndex);
@@ -763,7 +768,8 @@ void FPoseSearchColumn::SetFromDetails(FInstancedPropertyBag& PropertyBag, int32
    	TValueOrError<FStructView, EPropertyBagResult> Result = PropertyBag.GetValueStruct(PropertyName, FPoseSearchDatabaseAnimationAsset::StaticStruct());
 	if (FStructView* StructView = Result.TryGetValue())
 	{
-		if (FPoseSearchDatabaseAnimationAsset* Asset = InternalDatabase->GetMutableDatabaseAnimationAsset(RowIndex))
+		const int32 DatabaseAssetIndex = GetDatabaseAssetIndex(RowIndex);
+		if (FPoseSearchDatabaseAnimationAsset* Asset = InternalDatabase->GetMutableDatabaseAnimationAsset(DatabaseAssetIndex))
 		{
 			InternalDatabase->Modify();
 			*Asset = StructView->Get<FPoseSearchDatabaseAnimationAsset>();
