[CrashReport][Assert] Engine!UAnimSequence::GetBonePose() [AnimSequence.cpp]
#fix invalidate PoseSearchDatabase building in case an Animation(Asset) without a valid Skeleton is found

diff --git a/Engine/Plugins/Animation/PoseSearch/Source/Runtime/Private/PoseSearchDerivedData.cpp b/Engine/Plugins/Animation/PoseSearch/Source/Runtime/Private/PoseSearchDerivedData.cpp
--- a/Engine/Plugins/Animation/PoseSearch/Source/Runtime/Private/PoseSearchDerivedData.cpp
+++ b/Engine/Plugins/Animation/PoseSearch/Source/Runtime/Private/PoseSearchDerivedData.cpp
@@ -1893,6 +1893,7 @@ void FPoseSearchDatabaseAsyncCacheTask::OnGetComplete(UE::DerivedData::FCacheGet
 					{
 						if (const FPoseSearchDatabaseAnimationAssetBase* DatabaseAsset = DependentDatabase->GetDatabaseAnimationAsset(AnimationAssetIndex))
 						{
+							const UAnimationAsset* AnimationAsset = Cast<const UAnimationAsset>(DatabaseAsset->GetAnimationAsset());
 							if (DatabaseAsset->GetAnimationAsset() == nullptr)
 							{
 								UE_LOG(LogPoseSearch, Warning, TEXT("OnGetComplete - - No asset has been selected."));
@@ -1903,6 +1904,12 @@ void FPoseSearchDatabaseAsyncCacheTask::OnGetComplete(UE::DerivedData::FCacheGet
 								ResetSearchIndex();
 								return;
 							}
+							else if (AnimationAsset && AnimationAsset->GetSkeleton() == nullptr)
+							{
+								UE_LOG(LogPoseSearch, Error, TEXT("%s - %s BuildIndex Failed because '%s' referenced asset '%s' which has an invalid (null) Skeleton"), *LexToString(FullIndexKey.Hash), *MainDatabaseName, *DependentDatabaseName, *DatabaseAsset->GetAnimationAsset()->GetName());
+								ResetSearchIndex();
+								return;
+							}
 							// @todo: Find a way to prevent accessing UObjects while GC since this func is ran async. Commenting out for now to prevent crashes.
 							/*else if (!DatabaseAsset->IsSkeletonCompatible(Database->Schema))
 							{
