Fixed missing symbol export for InitTrajectorySamples and UpdateHistory_TransformHistory.

Other changes:
- Updated deprecation macros to only affect deprecated types used in the declaration of non-deprecated functions.
- Added a log warning to prompt licensees to update code or use the FTransformTrajectory instead to avoid implicing conversion constructor.

diff --git a/Engine/Plugins/Animation/PoseSearch/Source/Runtime/Private/PoseSearchTrajectoryLibrary.cpp b/Engine/Plugins/Animation/PoseSearch/Source/Runtime/Private/PoseSearchTrajectoryLibrary.cpp
--- a/Engine/Plugins/Animation/PoseSearch/Source/Runtime/Private/PoseSearchTrajectoryLibrary.cpp
+++ b/Engine/Plugins/Animation/PoseSearch/Source/Runtime/Private/PoseSearchTrajectoryLibrary.cpp
@@ -241,6 +241,18 @@ void UPoseSearchTrajectoryLibrary::InitTrajectorySamples(FTransformTrajectory& T
 
 PRAGMA_DISABLE_DEPRECATION_WARNINGS
 
+// Deprecated
+void UPoseSearchTrajectoryLibrary::UpdateHistory_TransformHistory(
+	FTransformTrajectory& Trajectory,
+	FVector CurrentPosition,
+	FVector CurrentVelocity,
+	const FPoseSearchTrajectoryData::FSampling& TrajectoryDataSampling,
+	float DeltaTime,
+	float CurrentTime)
+{
+	UpdateHistory_TransformHistory(Trajectory, CurrentPosition, CurrentVelocity, TrajectoryDataSampling, DeltaTime);
+}
+
 // Deprecated
 void UPoseSearchTrajectoryLibrary::UpdateHistory_TransformHistory(
 	FPoseSearchQueryTrajectory& Trajectory,
@@ -271,6 +283,18 @@ void UPoseSearchTrajectoryLibrary::UpdateHistory_TransformHistory(
 	Trajectory = TransformTrajectory;
 }
 
+// Deprecated
+void UPoseSearchTrajectoryLibrary::UpdatePrediction_SimulateCharacterMovement(
+	FPoseSearchQueryTrajectory& Trajectory,
+	const FPoseSearchTrajectoryData& TrajectoryData,
+	const FPoseSearchTrajectoryData::FDerived& TrajectoryDataDerived,
+	const FPoseSearchTrajectoryData::FSampling& TrajectoryDataSampling,
+	float DeltaTime)
+{
+	FTransformTrajectory TransformTrajectory = Trajectory;
+	UpdatePrediction_SimulateCharacterMovement(TransformTrajectory, TrajectoryData, TrajectoryDataDerived, TrajectoryDataSampling, DeltaTime);
+}
+
 PRAGMA_ENABLE_DEPRECATION_WARNINGS
 
 void UPoseSearchTrajectoryLibrary::UpdateHistory_TransformHistory(FTransformTrajectory& Trajectory,
diff --git a/Engine/Plugins/Animation/PoseSearch/Source/Runtime/Private/PoseSearchTrajectoryTypes.cpp b/Engine/Plugins/Animation/PoseSearch/Source/Runtime/Private/PoseSearchTrajectoryTypes.cpp
--- a/Engine/Plugins/Animation/PoseSearch/Source/Runtime/Private/PoseSearchTrajectoryTypes.cpp
+++ b/Engine/Plugins/Animation/PoseSearch/Source/Runtime/Private/PoseSearchTrajectoryTypes.cpp
@@ -220,6 +220,10 @@ FArchive& operator<<(FArchive& Ar, FPoseSearchQueryTrajectorySample& TrajectoryS
 
 FPoseSearchQueryTrajectory::FPoseSearchQueryTrajectory(const FTransformTrajectory& InTrajectory)
 {
+#if UE_EDITOR
+	UE_LOG(LogPoseSearch, Warning, TEXT("Converting FTransformTrajectory to FPoseSearchQueryTrajectory (deprecated type) implicitly. Consider updating code to avoid this conversion."))
+#endif
+	
 	Samples.SetNum(InTrajectory.Samples.Num(), EAllowShrinking::No);
 		
 	for (int i = 0; i < InTrajectory.Samples.Num(); ++i)
diff --git a/Engine/Plugins/Animation/PoseSearch/Source/Runtime/Public/PoseSearch/PoseSearchTrajectoryLibrary.h b/Engine/Plugins/Animation/PoseSearch/Source/Runtime/Public/PoseSearch/PoseSearchTrajectoryLibrary.h
--- a/Engine/Plugins/Animation/PoseSearch/Source/Runtime/Public/PoseSearch/PoseSearchTrajectoryLibrary.h
+++ b/Engine/Plugins/Animation/PoseSearch/Source/Runtime/Public/PoseSearch/PoseSearchTrajectoryLibrary.h
@@ -113,97 +113,184 @@ public:
 	// DEPRECATED /////////////////////////////////////////////////////////////
 	///////////////////////////////////////////////////////////////////////////
 	
-	PRAGMA_DISABLE_DEPRECATION_WARNINGS
 	// Initialize history and predicted samples based on sampling settings and a default state
 	UE_DEPRECATED(5.6, "Use the InitTrajectorySamples function that doesn't take FPoseSearchTrajectoryData")
-	static UE_API void InitTrajectorySamples(FPoseSearchQueryTrajectory& Trajectory, const FPoseSearchTrajectoryData& TrajectoryData, FVector DefaultPosition, FQuat DefaultFacing, const FPoseSearchTrajectoryData::FSampling& TrajectoryDataSampling, float DeltaTime);
+	static UE_API void InitTrajectorySamples(
+		PRAGMA_DISABLE_DEPRECATION_WARNINGS FPoseSearchQueryTrajectory & Trajectory PRAGMA_ENABLE_DEPRECATION_WARNINGS,
+		const FPoseSearchTrajectoryData& TrajectoryData,
+		FVector DefaultPosition, FQuat DefaultFacing,
+		const FPoseSearchTrajectoryData::FSampling& TrajectoryDataSampling,
+		float DeltaTime);
 
 	// Initialize history and predicted samples based on sampling settings and a default state
 	UE_DEPRECATED(5.6, "Use the IniTrajctorySamples function that takes in FTransformTrajectory instead of FPoseSearchQueryTrajectory.")
-	static UE_API void InitTrajectorySamples(FPoseSearchQueryTrajectory& Trajectory, FVector DefaultPosition, FQuat DefaultFacing, const FPoseSearchTrajectoryData::FSampling& TrajectoryDataSampling, float DeltaTime);
+	static UE_API void InitTrajectorySamples(
+		PRAGMA_DISABLE_DEPRECATION_WARNINGS FPoseSearchQueryTrajectory & Trajectory PRAGMA_ENABLE_DEPRECATION_WARNINGS,
+		FVector DefaultPosition,
+		FQuat DefaultFacing,
+		const FPoseSearchTrajectoryData::FSampling& TrajectoryDataSampling,
+		float DeltaTime);
 
 	UE_DEPRECATED(5.6, "Use the UpdateHistory_TransformHistory function that doesn't take FPoseSearchTrajectoryData")
-	static UE_API void UpdateHistory_TransformHistory(FPoseSearchQueryTrajectory& Trajectory, const FPoseSearchTrajectoryData& TrajectoryData, FVector CurrentPosition, FVector CurrentVelocity, const FPoseSearchTrajectoryData::FSampling& TrajectoryDataSampling, float DeltaTime);
+	static UE_API void UpdateHistory_TransformHistory(
+		PRAGMA_DISABLE_DEPRECATION_WARNINGS FPoseSearchQueryTrajectory & Trajectory PRAGMA_ENABLE_DEPRECATION_WARNINGS,
+		const FPoseSearchTrajectoryData& TrajectoryData,
+		FVector CurrentPosition,
+		FVector CurrentVelocity,
+		const FPoseSearchTrajectoryData::FSampling& TrajectoryDataSampling,
+		float DeltaTime);
 	
 	UE_DEPRECATED(5.6, "Use the UpdateHistory_TransformHistory function thats in FTransformTrajectory instead of FPoseSearchQueryTrajectory.")
-	static UE_API void UpdateHistory_TransformHistory(FPoseSearchQueryTrajectory& Trajectory, FVector CurrentPosition, FVector CurrentVelocity, const FPoseSearchTrajectoryData::FSampling& TrajectoryDataSampling, float DeltaTime);
+	static UE_API void UpdateHistory_TransformHistory(
+		PRAGMA_DISABLE_DEPRECATION_WARNINGS FPoseSearchQueryTrajectory & Trajectory PRAGMA_ENABLE_DEPRECATION_WARNINGS,
+		FVector CurrentPosition,
+		FVector CurrentVelocity,
+		const FPoseSearchTrajectoryData::FSampling& TrajectoryDataSampling,
+		float DeltaTime);
 
 	UE_DEPRECATED(5.6, "Use the UpdatePrediction_SimulateCharacterMovement function thats in FTransformTrajectory instead of FPoseSearchQueryTrajectory.")
-	static UE_API void UpdatePrediction_SimulateCharacterMovement(FPoseSearchQueryTrajectory& Trajectory, const FPoseSearchTrajectoryData& TrajectoryData, const FPoseSearchTrajectoryData::FDerived& TrajectoryDataDerived, const FPoseSearchTrajectoryData::FSampling& TrajectoryDataSampling, float DeltaTime);
+	static UE_API void UpdatePrediction_SimulateCharacterMovement(
+		PRAGMA_DISABLE_DEPRECATION_WARNINGS FPoseSearchQueryTrajectory & Trajectory PRAGMA_ENABLE_DEPRECATION_WARNINGS,
+		const FPoseSearchTrajectoryData& TrajectoryData,
+		const FPoseSearchTrajectoryData::FDerived& TrajectoryDataDerived,
+		const FPoseSearchTrajectoryData::FSampling& TrajectoryDataSampling,
+		float DeltaTime);
 
 	// Generates a prediction trajectory based of the current character intent. For use with Character actors.
 	UE_DEPRECATED(5.6, "Use PoseSearchGenerateTrajectory function that takes in FTransformTrajectory.")
-	static UE_API void PoseSearchGenerateTrajectory(const UObject* InAnimInstance, 
-		UPARAM(ref) const FPoseSearchTrajectoryData& InTrajectoryData, float InDeltaTime,
-		UPARAM(ref) FPoseSearchQueryTrajectory& InOutTrajectory, UPARAM(ref) float& InOutDesiredControllerYawLastUpdate, FPoseSearchQueryTrajectory& OutTrajectory,
-		float InHistorySamplingInterval = 0.04f, int32 InTrajectoryHistoryCount = 10, float InPredictionSamplingInterval = 0.2f, int32 InTrajectoryPredictionCount = 8);
+	static UE_API void PoseSearchGenerateTrajectory(
+		const UObject* InAnimInstance, 
+		UPARAM(ref) const FPoseSearchTrajectoryData& InTrajectoryData,
+		float InDeltaTime,
+		PRAGMA_DISABLE_DEPRECATION_WARNINGS UPARAM(ref) FPoseSearchQueryTrajectory& InOutTrajectory PRAGMA_ENABLE_DEPRECATION_WARNINGS,
+		UPARAM(ref) float& InOutDesiredControllerYawLastUpdate,
+		PRAGMA_DISABLE_DEPRECATION_WARNINGS FPoseSearchQueryTrajectory& OutTrajectory PRAGMA_ENABLE_DEPRECATION_WARNINGS,
+		float InHistorySamplingInterval = 0.04f,
+		int32 InTrajectoryHistoryCount = 10,
+		float InPredictionSamplingInterval = 0.2f,
+		int32 InTrajectoryPredictionCount = 8);
 
 	// Generates a prediction trajectory based of the current movement intent. For use with predictors. InPredictor must implement IPoseSearchTrajectoryPredictorInterface
 	UE_DEPRECATED(5.6, "Use PoseSearchGenerateTrajectoryWithPredictor function that doesn't take FPoseSearchTrajectoryData and also takes in FTransformTrajectory")
-	static UE_API void PoseSearchGeneratePredictorTrajectory(UObject* InPredictor,
-		UPARAM(ref) const FPoseSearchTrajectoryData& InTrajectoryData, float InDeltaTime,
-		UPARAM(ref) FPoseSearchQueryTrajectory& InOutTrajectory, UPARAM(ref) float& InOutDesiredControllerYawLastUpdate, FPoseSearchQueryTrajectory& OutTrajectory,
-		float InHistorySamplingInterval = 0.04f, int32 InTrajectoryHistoryCount = 10, float InPredictionSamplingInterval = 0.2f, int32 InTrajectoryPredictionCount = 8);
+	static UE_API void PoseSearchGeneratePredictorTrajectory(
+		UObject* InPredictor,
+		UPARAM(ref) const FPoseSearchTrajectoryData& InTrajectoryData,
+		float InDeltaTime,
+		PRAGMA_DISABLE_DEPRECATION_WARNINGS UPARAM(ref) FPoseSearchQueryTrajectory& InOutTrajectory PRAGMA_ENABLE_DEPRECATION_WARNINGS,
+		UPARAM(ref) float& InOutDesiredControllerYawLastUpdate,
+		PRAGMA_DISABLE_DEPRECATION_WARNINGS FPoseSearchQueryTrajectory& OutTrajectory PRAGMA_ENABLE_DEPRECATION_WARNINGS,
+		float InHistorySamplingInterval = 0.04f,
+		int32 InTrajectoryHistoryCount = 10,
+		float InPredictionSamplingInterval = 0.2f,
+		int32 InTrajectoryPredictionCount = 8);
 
 	// Generates a prediction trajectory based of the current movement intent. For use with predictors. InPredictor must implement IPoseSearchTrajectoryPredictorInterface
 	UE_DEPRECATED(5.6, "Use PoseSearchGenerateTrajectoryWithPredictor function that takes in FTransformTrajectory")
 	static UE_API void PoseSearchGenerateTrajectoryWithPredictor(TScriptInterface<IPoseSearchTrajectoryPredictorInterface> InPredictor,
 		float InDeltaTime,
-		UPARAM(ref) FPoseSearchQueryTrajectory& InOutTrajectory, UPARAM(ref) float& InOutDesiredControllerYawLastUpdate, FPoseSearchQueryTrajectory& OutTrajectory,
-		float InHistorySamplingInterval = 0.04f, int32 InTrajectoryHistoryCount = 10, float InPredictionSamplingInterval = 0.2f, int32 InTrajectoryPredictionCount = 8);
+		PRAGMA_DISABLE_DEPRECATION_WARNINGS UPARAM(ref) FPoseSearchQueryTrajectory& InOutTrajectory PRAGMA_ENABLE_DEPRECATION_WARNINGS,
+		UPARAM(ref) float& InOutDesiredControllerYawLastUpdate,
+		PRAGMA_DISABLE_DEPRECATION_WARNINGS FPoseSearchQueryTrajectory& OutTrajectory PRAGMA_ENABLE_DEPRECATION_WARNINGS,
+		float InHistorySamplingInterval = 0.04f,
+		int32 InTrajectoryHistoryCount = 10,
+		float InPredictionSamplingInterval = 0.2f,
+		int32 InTrajectoryPredictionCount = 8);
 
 	UE_DEPRECATED(5.6, "Use HandleTrajectoryWorldCollisions function that takes in FTransformTrajectory")
-	static UE_API void HandleTrajectoryWorldCollisions(const UObject* WorldContextObject, const UAnimInstance* AnimInstance, UPARAM(ref) const FPoseSearchQueryTrajectory& InTrajectory, bool bApplyGravity, float FloorCollisionsOffset, FPoseSearchQueryTrajectory& OutTrajectory, FPoseSearchTrajectory_WorldCollisionResults& CollisionResult,
-		ETraceTypeQuery TraceChannel, bool bTraceComplex, const TArray<AActor*>& ActorsToIgnore, EDrawDebugTrace::Type DrawDebugType, bool bIgnoreSelf = true, float MaxObstacleHeight = 10000.f, FLinearColor TraceColor = FLinearColor::Red, FLinearColor TraceHitColor = FLinearColor::Green, float DrawTime = 5.0f);
+	static UE_API void HandleTrajectoryWorldCollisions(
+		const UObject* WorldContextObject,
+		const UAnimInstance* AnimInstance,
+		PRAGMA_DISABLE_DEPRECATION_WARNINGS UPARAM(ref) const FPoseSearchQueryTrajectory& InTrajectory PRAGMA_ENABLE_DEPRECATION_WARNINGS,
+		bool bApplyGravity,
+		float FloorCollisionsOffset,
+		PRAGMA_DISABLE_DEPRECATION_WARNINGS FPoseSearchQueryTrajectory& OutTrajectory PRAGMA_ENABLE_DEPRECATION_WARNINGS,
+		FPoseSearchTrajectory_WorldCollisionResults& CollisionResult,
+		ETraceTypeQuery TraceChannel,
+		bool bTraceComplex,
+		const TArray<AActor*>& ActorsToIgnore,
+		EDrawDebugTrace::Type DrawDebugType,
+		bool bIgnoreSelf = true,
+		float MaxObstacleHeight = 10000.f,
+		FLinearColor TraceColor = FLinearColor::Red,
+		FLinearColor TraceHitColor = FLinearColor::Green,
+		float DrawTime = 5.0f);
 
 	UE_DEPRECATED(5.6, "Use HandleTrajectoryWorldCollisionsWithGravity that takes in FTransformTrajectory.")
 	static UE_API void HandleTrajectoryWorldCollisionsWithGravity(const UObject* WorldContextObject, 
-		UPARAM(ref) const FPoseSearchQueryTrajectory& InTrajectory, FVector StartingVelocity, bool bApplyGravity, FVector GravityAccel, float FloorCollisionsOffset, FPoseSearchQueryTrajectory& OutTrajectory, FPoseSearchTrajectory_WorldCollisionResults& CollisionResult,
-		ETraceTypeQuery TraceChannel, bool bTraceComplex, const TArray<AActor*>& ActorsToIgnore, EDrawDebugTrace::Type DrawDebugType, bool bIgnoreSelf = true, float MaxObstacleHeight = 10000.f, FLinearColor TraceColor = FLinearColor::Red, FLinearColor TraceHitColor = FLinearColor::Green, float DrawTime = 5.0f);
+		PRAGMA_DISABLE_DEPRECATION_WARNINGS UPARAM(ref) const FPoseSearchQueryTrajectory& InTrajectory PRAGMA_ENABLE_DEPRECATION_WARNINGS,
+		FVector StartingVelocity,
+		bool bApplyGravity,
+		FVector GravityAccel,
+		float FloorCollisionsOffset,
+		PRAGMA_DISABLE_DEPRECATION_WARNINGS FPoseSearchQueryTrajectory& OutTrajectory PRAGMA_ENABLE_DEPRECATION_WARNINGS,
+		FPoseSearchTrajectory_WorldCollisionResults& CollisionResult,
+		ETraceTypeQuery TraceChannel,
+		bool bTraceComplex,
+		const TArray<AActor*>& ActorsToIgnore,
+		EDrawDebugTrace::Type DrawDebugType,
+		bool bIgnoreSelf = true,
+		float MaxObstacleHeight = 10000.f,
+		FLinearColor TraceColor = FLinearColor::Red,
+		FLinearColor TraceHitColor = FLinearColor::Green,
+		float DrawTime = 5.0f);
 
 	UE_DEPRECATED(5.6, "Use GetTrajectorySampleAtTime that takes in FTransformTrajectory and FTransformTrajectorySample.")
-	static UE_API void GetTrajectorySampleAtTime(UPARAM(ref) const FPoseSearchQueryTrajectory& InTrajectory, float Time, FPoseSearchQueryTrajectorySample& OutTrajectorySample, bool bExtrapolate = false);
+	static UE_API void GetTrajectorySampleAtTime(
+		PRAGMA_DISABLE_DEPRECATION_WARNINGS UPARAM(ref) const FPoseSearchQueryTrajectory& InTrajectory PRAGMA_ENABLE_DEPRECATION_WARNINGS,
+		PRAGMA_DISABLE_DEPRECATION_WARNINGS float Time, FPoseSearchQueryTrajectorySample& OutTrajectorySample PRAGMA_ENABLE_DEPRECATION_WARNINGS,
+		bool bExtrapolate = false);
 
 	UE_DEPRECATED(5.6, "Use GetTrajectorySampleAtTime that takes in FTransformTrajectory.")
-	static UE_API void GetTrajectoryVelocity(UPARAM(ref) const FPoseSearchQueryTrajectory& InTrajectory, float Time1, float Time2, FVector& OutVelocity, bool bExtrapolate = false);
+	static UE_API void GetTrajectoryVelocity(
+		PRAGMA_DISABLE_DEPRECATION_WARNINGS UPARAM(ref) const FPoseSearchQueryTrajectory& InTrajectory PRAGMA_ENABLE_DEPRECATION_WARNINGS,
+		float Time1,
+		float Time2,
+		FVector& OutVelocity,
+		bool bExtrapolate = false);
 
 	UE_DEPRECATED(5.6, "Use GetTrajectorySampleAtTime that takes in FTransformTrajectory.")
-	static UE_API void GetTrajectoryAngularVelocity(UPARAM(ref) const FPoseSearchQueryTrajectory& InTrajectory, float Time1, float Time2, FVector& OutAngularVelocity, bool bExtrapolate = false);
+	static UE_API void GetTrajectoryAngularVelocity(
+		PRAGMA_DISABLE_DEPRECATION_WARNINGS UPARAM(ref) const FPoseSearchQueryTrajectory& InTrajectory PRAGMA_ENABLE_DEPRECATION_WARNINGS,
+		float Time1,
+		float Time2,
+		FVector& OutAngularVelocity,
+		bool bExtrapolate = false);
 
 	UE_DEPRECATED(5.6, "Use GetTrajectorySampleAtTime that takes in FTransformTrajectorySample.")
-	static UE_API FTransform GetTransform(UPARAM(ref) const FPoseSearchQueryTrajectorySample& InTrajectorySample);
+	static UE_API FTransform GetTransform(
+		PRAGMA_DISABLE_DEPRECATION_WARNINGS UPARAM(ref) const FPoseSearchQueryTrajectorySample& InTrajectorySample PRAGMA_ENABLE_DEPRECATION_WARNINGS);
 
 	UE_DEPRECATED(5.6, "Use DrawTrajectory that takes in FTransformTrajectory.")
-	static UE_API void DrawTrajectory(const UObject* WorldContextObject, UPARAM(ref) const FPoseSearchQueryTrajectory& InTrajectory, const float DebugThickness, float HeightOffset);
-	UE_API PRAGMA_ENABLE_DEPRECATION_WARNINGS
+	static UE_API void DrawTrajectory(
+		const UObject* WorldContextObject,
+		PRAGMA_DISABLE_DEPRECATION_WARNINGS UPARAM(ref) const FPoseSearchQueryTrajectory& InTrajectory PRAGMA_ENABLE_DEPRECATION_WARNINGS,
+		const float DebugThickness,
+		float HeightOffset);
 
+	UE_DEPRECATED(5.7, "CurrentTime is no longer needed or used. Use UpdateHistory_TransformHistory with no CurrentTime")
+	static UE_API void UpdateHistory_TransformHistory(FTransformTrajectory& Trajectory,
+										   FVector CurrentPosition,
+										   FVector CurrentVelocity,
+										   const FPoseSearchTrajectoryData::FSampling& TrajectoryDataSampling,
+										   float DeltaTime,
+										   float CurrentTime);
+	
 	///////////////////////////////////////////////////////////////////////////
 	// NON- DEPRECATED ////////////////////////////////////////////////////////
 	///////////////////////////////////////////////////////////////////////////
 	
 	// Initialize history and predicted samples based on sampling settings and a default state
-	static void InitTrajectorySamples(FTransformTrajectory& Trajectory, FVector DefaultPosition, FQuat DefaultFacing, const FPoseSearchTrajectoryData::FSampling& TrajectoryDataSampling, float DeltaTime);
-
+	static UE_API void InitTrajectorySamples(FTransformTrajectory& Trajectory, FVector DefaultPosition, FQuat DefaultFacing, const FPoseSearchTrajectoryData::FSampling& TrajectoryDataSampling, float DeltaTime);
+	
 	// Update history by tracking offsets that result from character intent (e.g. movement component velocity) and applying
 	// that to the current world transform. This works well on moving platforms as it only stores a history of movement
 	// that results from character intent, not movement from platforms.
 	// Important: CurrentVelocity should be the velocity relative to the ground as reported by the character movement component or character mover etc
-	UE_DEPRECATED(5.7, "CurrentTime is no longer needed or used. Use UpdateHistory_TransformHistory with no CurrentTime")
-	static void UpdateHistory_TransformHistory(FTransformTrajectory& Trajectory,
-	                                           FVector CurrentPosition,
-	                                           FVector CurrentVelocity,
-	                                           const FPoseSearchTrajectoryData::FSampling& TrajectoryDataSampling,
-	                                           float DeltaTime,
-	                                           float CurrentTime)
-	{
-		UpdateHistory_TransformHistory(Trajectory, CurrentPosition, CurrentVelocity, TrajectoryDataSampling, DeltaTime);
-	}
-
 	static UE_API void UpdateHistory_TransformHistory(FTransformTrajectory& Trajectory,
-	                                                       FVector CurrentPosition,
-	                                                       FVector CurrentVelocity,
-	                                                       const FPoseSearchTrajectoryData::FSampling& TrajectoryDataSampling,
-	                                                       float DeltaTime);
+													  FVector CurrentPosition,
+													  FVector CurrentVelocity,
+													  const FPoseSearchTrajectoryData::FSampling& TrajectoryDataSampling,
+													  float DeltaTime);
 
 	// Update prediction by simulating the movement math for ground locomotion from UCharacterMovementComponent.
 	static UE_API void UpdatePrediction_SimulateCharacterMovement(FTransformTrajectory& Trajectory, const FPoseSearchTrajectoryData& TrajectoryData, const FPoseSearchTrajectoryData::FDerived& TrajectoryDataDerived, const FPoseSearchTrajectoryData::FSampling& TrajectoryDataSampling, float DeltaTime);
