Small fixes for pose search trajectory library.

Fixed missing copy to ouput trajectory in deprecated SimulateCharacterMovement
Fixed inproper use of UE_EDITOR (used WITH_EDITOR instead).

diff --git a/Engine/Plugins/Animation/PoseSearch/Source/Runtime/Private/PoseSearchTrajectoryLibrary.cpp b/Engine/Plugins/Animation/PoseSearch/Source/Runtime/Private/PoseSearchTrajectoryLibrary.cpp
--- a/Engine/Plugins/Animation/PoseSearch/Source/Runtime/Private/PoseSearchTrajectoryLibrary.cpp
+++ b/Engine/Plugins/Animation/PoseSearch/Source/Runtime/Private/PoseSearchTrajectoryLibrary.cpp
@@ -293,6 +293,7 @@ void UPoseSearchTrajectoryLibrary::UpdatePrediction_SimulateCharacterMovement(
 {
 	FTransformTrajectory TransformTrajectory = Trajectory;
 	UpdatePrediction_SimulateCharacterMovement(TransformTrajectory, TrajectoryData, TrajectoryDataDerived, TrajectoryDataSampling, DeltaTime);
+	Trajectory = TransformTrajectory;
 }
 
 PRAGMA_ENABLE_DEPRECATION_WARNINGS
diff --git a/Engine/Plugins/Animation/PoseSearch/Source/Runtime/Private/PoseSearchTrajectoryTypes.cpp b/Engine/Plugins/Animation/PoseSearch/Source/Runtime/Private/PoseSearchTrajectoryTypes.cpp
--- a/Engine/Plugins/Animation/PoseSearch/Source/Runtime/Private/PoseSearchTrajectoryTypes.cpp
+++ b/Engine/Plugins/Animation/PoseSearch/Source/Runtime/Private/PoseSearchTrajectoryTypes.cpp
@@ -220,7 +220,7 @@ FArchive& operator<<(FArchive& Ar, FPoseSearchQueryTrajectorySample& TrajectoryS
 
 FPoseSearchQueryTrajectory::FPoseSearchQueryTrajectory(const FTransformTrajectory& InTrajectory)
 {
-#if UE_EDITOR
+#if WITH_EDITOR
 	UE_LOG(LogPoseSearch, Warning, TEXT("Converting FTransformTrajectory to FPoseSearchQueryTrajectory (deprecated type) implicitly. Consider updating code to avoid this conversion."))
 #endif
 	
