PoseSearch - fix for crash when you add a Pose Match Column when the chooser has no data type

#[review]-46600461 [REDACTED_HANDLE]
#rnx

diff --git a/Engine/Plugins/Animation/PoseSearch/Source/Runtime/Private/Chooser/PoseSearchChooserColumn.cpp b/Engine/Plugins/Animation/PoseSearch/Source/Runtime/Private/Chooser/PoseSearchChooserColumn.cpp
--- a/Engine/Plugins/Animation/PoseSearch/Source/Runtime/Private/Chooser/PoseSearchChooserColumn.cpp
+++ b/Engine/Plugins/Animation/PoseSearch/Source/Runtime/Private/Chooser/PoseSearchChooserColumn.cpp
@@ -133,20 +133,27 @@ UObject* FPoseSearchColumn::GetReferencedObject(int32 RowIndex) const
 	UChooserTable* OuterChooser = Cast<UChooserTable>(InternalDatabase->GetOuter());
 	check(OuterChooser);
 
-	const UObject* ReferencedObject = nullptr;
+	UObject* ReferencedObject = nullptr;
 
 	if (RowIndex != ChooserColumn_SpecialIndex_Fallback)
 	{
 		check(RowIndex >= 0);
-		return OuterChooser->ResultsStructs[RowIndex].Get<FObjectChooserBase>().GetReferencedObject();
+		const FInstancedStruct& Result = OuterChooser->ResultsStructs[RowIndex];
+		if (Result.IsValid())
+		{
+			ReferencedObject = Result.Get<FObjectChooserBase>().GetReferencedObject();
+		}
 	}
-
-	if (OuterChooser->FallbackResult.IsValid())
+	else
 	{
-		return OuterChooser->FallbackResult.Get<FObjectChooserBase>().GetReferencedObject();
+		const FInstancedStruct& FallbackResult = OuterChooser->FallbackResult;
+		if (FallbackResult.IsValid())
+		{
+			ReferencedObject = FallbackResult.Get<FObjectChooserBase>().GetReferencedObject();
+		}
 	}
 
-	return nullptr;
+	return ReferencedObject;
 }
 
 FName FPoseSearchColumn::RowValuesPropertyName()
