PoseSearch - setting FDatabasePreviewActor, FInteractionAssetPreviewActor Samplers bInEnforceCompressedDataSampling to false

diff --git a/Engine/Plugins/Animation/PoseSearch/Source/Editor/Private/PoseSearchDatabaseViewModel.cpp b/Engine/Plugins/Animation/PoseSearch/Source/Editor/Private/PoseSearchDatabaseViewModel.cpp
--- a/Engine/Plugins/Animation/PoseSearch/Source/Editor/Private/PoseSearchDatabaseViewModel.cpp
+++ b/Engine/Plugins/Animation/PoseSearch/Source/Editor/Private/PoseSearchDatabaseViewModel.cpp
@@ -104,7 +104,7 @@ bool FDatabasePreviewActor::SpawnPreviewActor(UWorld* World, const UPoseSearchDa
 
 	const FMirrorDataCache MirrorDataCache(AnimInstance->GetMirrorDataTable(), AnimInstance->GetRequiredBonesOnAnyThread());
 	
-	Sampler.Init(PreviewAsset, SamplerRootTransformOrigin, IndexAsset.GetBlendParameters());
+	Sampler.Init(PreviewAsset, SamplerRootTransformOrigin, IndexAsset.GetBlendParameters(), FAnimationAssetSampler::DefaultRootTransformSamplingRate, true, false);
 
 	PlayTimeOffset = 0.f;
 	if (PoseIdxForTimeOffset >= 0)
diff --git a/Engine/Plugins/Animation/PoseSearch/Source/Editor/Private/PoseSearchInteractionAssetEditor.cpp b/Engine/Plugins/Animation/PoseSearch/Source/Editor/Private/PoseSearchInteractionAssetEditor.cpp
--- a/Engine/Plugins/Animation/PoseSearch/Source/Editor/Private/PoseSearchInteractionAssetEditor.cpp
+++ b/Engine/Plugins/Animation/PoseSearch/Source/Editor/Private/PoseSearchInteractionAssetEditor.cpp
@@ -115,7 +115,7 @@ void FInteractionAssetPreviewActor::UpdatePreviewActor(const UPoseSearchInteract
 			if (PreviewAsset != Sampler.GetAsset() || !Origin.Equals(Sampler.GetRootTransformOrigin()))
 			{
 				// reinitializing the Sampler if the PreviewAsset or the origin transform changed
-				Sampler.Init(PreviewAsset, Origin, BlendParameters);
+				Sampler.Init(PreviewAsset, Origin, BlendParameters, FAnimationAssetSampler::DefaultRootTransformSamplingRate, true, false);
 				bSamplerReinitialized = true;
 			}
 
