PoseSearch - adding PoseSearch to PoseSearch STATGROUP(s)

diff --git a/Engine/Plugins/Animation/PoseSearch/Source/Runtime/Private/PoseSearchDatabase.cpp b/Engine/Plugins/Animation/PoseSearch/Source/Runtime/Private/PoseSearchDatabase.cpp
--- a/Engine/Plugins/Animation/PoseSearch/Source/Runtime/Private/PoseSearchDatabase.cpp
+++ b/Engine/Plugins/Animation/PoseSearch/Source/Runtime/Private/PoseSearchDatabase.cpp
@@ -36,10 +36,10 @@
 
 struct FPoseSearchDatabaseAnimationAssetBase;
 DECLARE_STATS_GROUP(TEXT("PoseSearch"), STATGROUP_PoseSearch, STATCAT_Advanced);
-DECLARE_CYCLE_STAT_EXTERN(TEXT("Search Brute Force"), STAT_PoseSearch_BruteForce, STATGROUP_PoseSearch, );
-DECLARE_CYCLE_STAT_EXTERN(TEXT("Search PCA/KNN"), STAT_PoseSearch_PCAKNN, STATGROUP_PoseSearch, );
-DECLARE_CYCLE_STAT_EXTERN(TEXT("Search VPTree"), STAT_PoseSearch_VPTree, STATGROUP_PoseSearch, );
-DECLARE_CYCLE_STAT_EXTERN(TEXT("Search Event"), STAT_PoseSearch_Event, STATGROUP_PoseSearch, );
+DECLARE_CYCLE_STAT_EXTERN(TEXT("PoseSearch Brute Force"), STAT_PoseSearch_BruteForce, STATGROUP_PoseSearch, );
+DECLARE_CYCLE_STAT_EXTERN(TEXT("PoseSearch PCA/KNN"), STAT_PoseSearch_PCAKNN, STATGROUP_PoseSearch, );
+DECLARE_CYCLE_STAT_EXTERN(TEXT("PoseSearch VPTree"), STAT_PoseSearch_VPTree, STATGROUP_PoseSearch, );
+DECLARE_CYCLE_STAT_EXTERN(TEXT("PoseSearch Event"), STAT_PoseSearch_Event, STATGROUP_PoseSearch, );
 DEFINE_STAT(STAT_PoseSearch_BruteForce);
 DEFINE_STAT(STAT_PoseSearch_PCAKNN);
 DEFINE_STAT(STAT_PoseSearch_VPTree);
