PoseSearch - removing template objects from the PDSs DDC key

diff --git a/Engine/Plugins/Animation/PoseSearch/Source/Runtime/Private/PoseSearchDerivedDataKey.cpp b/Engine/Plugins/Animation/PoseSearch/Source/Runtime/Private/PoseSearchDerivedDataKey.cpp
--- a/Engine/Plugins/Animation/PoseSearch/Source/Runtime/Private/PoseSearchDerivedDataKey.cpp
+++ b/Engine/Plugins/Animation/PoseSearch/Source/Runtime/Private/PoseSearchDerivedDataKey.cpp
@@ -452,6 +452,11 @@ bool FKeyBuilder::ShouldInclude(const UObject* Object)
 		return false;
 	}
 
+	if (Object->IsTemplate())
+	{
+		return false;
+	}
+
 	if (Object->HasAnyFlags(RF_Transient | RF_MirroredGarbage | RF_NewerVersionExists | RF_BeginDestroyed | RF_FinishDestroyed | RF_DuplicateTransient | RF_NonPIEDuplicateTransient))
 	{
 		return false;
