PoseSearch - MM interaction performance improvements

diff --git a/Engine/Plugins/Animation/PoseSearch/Source/Runtime/Private/PoseSearchInteractionIsland.cpp b/Engine/Plugins/Animation/PoseSearch/Source/Runtime/Private/PoseSearchInteractionIsland.cpp
--- a/Engine/Plugins/Animation/PoseSearch/Source/Runtime/Private/PoseSearchInteractionIsland.cpp
+++ b/Engine/Plugins/Animation/PoseSearch/Source/Runtime/Private/PoseSearchInteractionIsland.cpp
@@ -293,26 +293,96 @@ static void AddPrerequisite(FTickFunction& TickFunction, UObject* TargetObject,
 	}
 }
 
-// FInteractionSearchContextBase
+// FValidInteractionSearch
 ///////////////////////////////////////////////////////////
-const UE::PoseSearch::IPoseHistory* FInteractionSearchContextBase::GetPoseHistory(int32 Index) const
+void FValidInteractionSearch::Init(const FInteractionSearchContext& Other)
 {
-	check(PoseHistories.IsValidIndex(Index));
-	check(PoseHistories[Index].IsValid());
-	return PoseHistories[Index].Pin().Get();
+	Database = Other.Database;
+	bDisableCollisions = Other.bDisableCollisions;
+		
+	AnimContexts.SetNum(Other.AnimContexts.Num());
+	PoseHistories.SetNum(Other.PoseHistories.Num());
+	Roles = Other.Roles;
+
+	for (int32 Index = 0; Index < Other.AnimContexts.Num(); ++Index)
+	{
+		AnimContexts[Index] = Other.AnimContexts[Index];
+	}
+
+	for (int32 Index = 0; Index < Other.PoseHistories.Num(); ++Index)
+	{
+		PoseHistories[Index] = Other.PoseHistories[Index]->AsWeak();
+	}
 }
 
-bool FInteractionSearchContextBase::IsEquivalent(const FInteractionSearchContextBase& Other) const
+void FValidInteractionSearch::DoDisableCollisions()
 {
-	// skipping bDisableCollisions for equality
-	return Database == Other.Database &&
-		AnimContexts == Other.AnimContexts &&
-		PoseHistories == Other.PoseHistories &&
-		Roles == Other.Roles;
+	check(DisabledCollisions.IsEmpty());
+	if (bDisableCollisions)
+	{
+		FMemMark Mark(FMemStack::Get());
+		TArray<AActor*, TInlineAllocator<PreallocatedRolesNum, TMemStackAllocator<>>> Actors;
+		TArray<UPrimitiveComponent*, TInlineAllocator<PreallocatedRolesNum, TMemStackAllocator<>>> PrimitiveComponents;
+
+		for (int32 AnimContextIndex = 0; AnimContextIndex < AnimContexts.Num(); ++AnimContextIndex)
+		{
+			if (const UObject* AnimContext = AnimContexts[AnimContextIndex].Get())
+			{
+				AActor* Actor = const_cast<AActor*>(GetContextOwningActor(AnimContext, false));
+				check(Actor);
+				Actors.Add(Actor);
+				PrimitiveComponents.Add(Cast<UPrimitiveComponent>(Actor->GetRootComponent()));
+			}
+		}
+
+		for (int32 IndexA = 0; IndexA < Actors.Num(); ++IndexA)
+		{
+			UPrimitiveComponent* PrimitiveComponentA = PrimitiveComponents[IndexA];
+			AActor* ActorA = Actors[IndexA];
+			check(ActorA);
+
+			for (int32 IndexB = IndexA + 1; IndexB < Actors.Num(); ++IndexB)
+			{
+				AActor* ActorB = Actors[IndexB];
+				check(ActorB);
+
+				UPrimitiveComponent* PrimitiveComponentB = PrimitiveComponents[IndexB];
+
+				if (PrimitiveComponentA && !PrimitiveComponentA->GetMoveIgnoreActors().Contains(ActorB))
+				{
+					DisabledCollisions.Add({ ActorA, ActorB });
+					PrimitiveComponentA->IgnoreActorWhenMoving(ActorB, true);
+				}
+
+				if (PrimitiveComponentB && !PrimitiveComponentB->GetMoveIgnoreActors().Contains(ActorA))
+				{
+					DisabledCollisions.Add({ ActorB, ActorA });
+					PrimitiveComponentB->IgnoreActorWhenMoving(ActorA, true);
+				}
+			}
+		}
+	}
+}
+
+void FValidInteractionSearch::UndoDisableCollisions()
+{
+	for (const FDisabledCollisions::ElementType& DisabledCollision : DisabledCollisions)
+	{
+		if (AActor* ActorA = DisabledCollision.Key.Get())
+		{
+			if (AActor* ActorB = DisabledCollision.Value.Get())
+			{
+				if (UPrimitiveComponent* PrimitiveComponentA = Cast<UPrimitiveComponent>(ActorA->GetRootComponent()))
+				{
+					PrimitiveComponentA->IgnoreActorWhenMoving(ActorB, false);
+				}
+			}
+		}
+	}
 }
 
 #if ENABLE_VISUAL_LOG
-void FInteractionSearchContextBase::VLogContext(const FColor& Color) const
+void FValidInteractionSearch::VLogContext(const FColor& Color) const
 {
 	using namespace UE::PoseSearch;
 
@@ -348,8 +418,75 @@ void FInteractionSearchContextBase::VLogContext(const FColor& Color) const
 }
 #endif // ENABLE_VISUAL_LOG
 
+// FInteractionSearchContext
+///////////////////////////////////////////////////////////
+const UE::PoseSearch::IPoseHistory* FInteractionSearchContext::GetPoseHistory(int32 Index) const
+{
+	check(PoseHistories.IsValidIndex(Index));
+	check(PoseHistories[Index]);
+	return PoseHistories[Index];
+}
+
+uint32 FInteractionSearchContext::GetHashForEquivalence() const
+{
+	uint32 Hash = GetTypeHash(Database);
+
+	// skipping bDisableCollisions for hashing
+
+	for (int32 Index = 0; Index < AnimContexts.Num(); ++Index)
+	{
+		Hash = HashCombineFast(Hash, GetTypeHash(AnimContexts[Index]));
+	}
+
+	for (int32 Index = 0; Index < PoseHistories.Num(); ++Index)
+	{
+		Hash = HashCombineFast(Hash, GetTypeHash(PoseHistories[Index]));
+	}
+
+	for (int32 Index = 0; Index < Roles.Num(); ++Index)
+	{
+		Hash = HashCombineFast(Hash, GetTypeHash(Roles[Index]));
+	}
+
+	return Hash;
+}
+
+void FInteractionSearchContext::Add(const UObject* AnimContext, const UE::PoseSearch::IPoseHistory* PoseHistory, const UE::PoseSearch::FRole Role)
+{
+	check(AnimContext && PoseHistory);
+		
+	// AnimContexts must be sorted to have deterministic searches across multiple frames
+	// as well as for faster IsEquivalent implementation
+	check(AnimContexts.IsEmpty() || GetAnimContext(AnimContexts.Num() - 1) < AnimContext);
+
+	AnimContexts.Add(AnimContext);
+	PoseHistories.Add(PoseHistory);
+	Roles.Add(Role);
+}
+
+FPoseSearchContinuingProperties FInteractionSearchContext::GetContinuingProperties() const
+{
+	FPoseSearchContinuingProperties ContinuingProperties;
+	ContinuingProperties.PlayingAsset = PlayingAsset;
+	ContinuingProperties.PlayingAssetAccumulatedTime = PlayingAssetAccumulatedTime;
+	ContinuingProperties.bIsPlayingAssetMirrored = bIsPlayingAssetMirrored;
+	ContinuingProperties.PlayingAssetBlendParameters = PlayingAssetBlendParameters;
+	ContinuingProperties.InterruptMode = InterruptMode;
+	ContinuingProperties.bIsContinuingInteraction = bIsContinuingInteraction;
+	return ContinuingProperties;
+}
+
+bool FInteractionSearchContext::IsEquivalent(const FInteractionSearchContext& Other) const
+{
+	// skipping bDisableCollisions for equality
+	return Database == Other.Database &&
+		AnimContexts == Other.AnimContexts &&
+		PoseHistories == Other.PoseHistories &&
+		Roles == Other.Roles;
+}
+
 #if DO_CHECK
-bool FInteractionSearchContextBase::CheckForConsistency() const
+bool FInteractionSearchContext::CheckForConsistency() const
 {
 	if (Database == nullptr)
 	{
@@ -391,7 +528,7 @@ bool FInteractionSearchContextBase::CheckForConsistency() const
 	for (int32 IndexA = 1; IndexA < Num; ++IndexA)
 	{
 		// AnimContexts must be sorted to have deterministic searches across multiple frames
-		if (AnimContexts[IndexA - 1].Get() >= AnimContexts[IndexA].Get())
+		if (AnimContexts[IndexA - 1] >= AnimContexts[IndexA])
 		{
 			return false;
 		}
@@ -427,25 +564,9 @@ bool FInteractionSearchContextBase::CheckForConsistency() const
 	return true;
 }
 
-void FInteractionSearchContextBase::TestHistoryCollectorsThreadingAccess() const
-{
-	for (int32 Index = 0; Index < AnimContexts.Num(); ++Index)
-	{
-		if (const UAnimInstance* AnimInstance = Cast<UAnimInstance>(AnimContexts[Index].Get()))
-		{
-			const IPoseHistory* PoseHistory = GetPoseHistory(Index);
-			check(PoseHistory);
-			
-			const USkeleton* Skeleton = AnimInstance->GetRequiredBonesOnAnyThread().GetSkeletonAsset();
-			check(Skeleton);
-
-			FTransform ActorRootBoneTransform;
-			PoseHistory->GetTransformAtTime(0.f, ActorRootBoneTransform, Skeleton, RootBoneIndexType, ComponentSpaceIndexType);
-		}
-	}
-}
 #endif // DO_CHECK
 
+
 // FInteractionSearchResult
 ///////////////////////////////////////////////////////////
 bool FInteractionSearchResult::operator==(const FInteractionSearchResult& Other) const
@@ -473,7 +594,7 @@ void FInteractionIsland::FPreTickFunction::ExecuteTick(float DeltaTime, enum ELe
 				if (const UObject* AnimContext = SearchContext.GetAnimContext(Index))
 				{
 					const IPoseHistory* PoseHistory = SearchContext.GetPoseHistory(Index);
-					check(PoseHistory);
+					check(PoseHistory && PoseHistory->AsWeak().IsValid());
 					
 					// since FInteractionIsland has a tick dependency with the USkeletalMeshComponent it's safe modify the IPoseHistory
 					const_cast<IPoseHistory*>(PoseHistory)->GenerateTrajectory(AnimContext, DeltaTime);
@@ -486,7 +607,7 @@ void FInteractionIsland::FPreTickFunction::ExecuteTick(float DeltaTime, enum ELe
 		{
 			// @todo: need to figure out why when creating a new island FPreTickFunction gets called twice (it's not a real issue rather than a performance hit)
 			//		  use GVarPoseSearchInteractionCacheIslands = false to debug the issue (it destroys the islands every frame)
-			UE_LOG(LogPoseSearch, Warning, TEXT("FInteractionIsland::FPreTickFunction::ExecuteTick, called twice before UPoseSearchInteractionSubsystem::Tick!"));
+			UE_LOG(LogPoseSearch, Log, TEXT("FInteractionIsland::FPreTickFunction::ExecuteTick, called twice before UPoseSearchInteractionSubsystem::Tick!"));
 		}
 		else
 		{
@@ -963,7 +1084,7 @@ void FInteractionIsland::LogTickDependencies() const
 }
 #endif // ENABLE_ANIM_DEBUG
 
-bool FInteractionIsland::DoSearch_AnyThread(const UObject* AnimContext, const TConstArrayView<FValidInteractionSearch> ValidInteractionSearches, FPoseSearchBlueprintResult& Result)
+bool FInteractionIsland::DoSearch_AnyThread(const UObject* AnimContext, FPoseSearchBlueprintResult& Result)
 {
 	check(AnimContext);
 
@@ -996,7 +1117,7 @@ bool FInteractionIsland::DoSearch_AnyThread(const UObject* AnimContext, const TC
 	for (int32 SearchIndex = 0; SearchIndex < SearchContexts.Num(); ++SearchIndex)
 	{
 		FInteractionSearchContext& InteractionSearchContext = SearchContexts[SearchIndex];
-		const UPoseSearchDatabase* Database = InteractionSearchContext.Database.Get();
+		const UPoseSearchDatabase* Database = InteractionSearchContext.Database;
 		if (!Database)
 		{
 			UE_LOG(LogPoseSearch, Error, TEXT("FInteractionIsland::DoSearch_AnyThread invalid context database"));
diff --git a/Engine/Plugins/Animation/PoseSearch/Source/Runtime/Private/PoseSearchInteractionSubsystem.cpp b/Engine/Plugins/Animation/PoseSearch/Source/Runtime/Private/PoseSearchInteractionSubsystem.cpp
--- a/Engine/Plugins/Animation/PoseSearch/Source/Runtime/Private/PoseSearchInteractionSubsystem.cpp
+++ b/Engine/Plugins/Animation/PoseSearch/Source/Runtime/Private/PoseSearchInteractionSubsystem.cpp
@@ -9,8 +9,13 @@
 #include "PoseSearch/PoseSearchInteractionUtils.h"
 #include "PoseSearch/PoseSearchRole.h"
 #include "PoseSearch/PoseSearchSchema.h"
+#include "PoseSearch/Trace/PoseSearchTraceLogger.h"
 #include "VisualLogger/VisualLogger.h"
 
+#if WITH_EDITOR
+#include "Editor.h"
+#endif
+
 #include UE_INLINE_GENERATED_CPP_BY_NAME(PoseSearchInteractionSubsystem)
 
 namespace UE::PoseSearch
@@ -42,28 +47,6 @@ struct FAnimContextInfo
 		}
 	}
 
-	// performs broad phase analysis checking if at least one of the Availabilities associated to AnimContext can interact with OtherAnimContextInfo.
-	// This is a more relaxed analysis than the one performed in FRoledAnimContextInfo::CanInteractWith
-	bool CanInteractWith(const FAnimContextInfo& OtherAnimContextInfo) const
-	{
-		check(this != &OtherAnimContextInfo);
-		
-		// @todo: enable this code if we ended up requiring preventing interactions between the same actor!
-		//check(AnimContextAvailabilities && OtherAnimContextInfo.AnimContextAvailabilities);
-		//const AActor* AnimContextActor = GetContextOwningActor(AnimContextAvailabilities->AnimContext);
-		//const AActor* OtherAnimContextActor = GetContextOwningActor(OtherAnimContextInfo.AnimContextAvailabilities->AnimContext);
-		//if (AnimContextActor == OtherAnimContextActor)
-		//{
-		//	return false;
-		//}
-
-		const FVector DeltaLocation = Location - OtherAnimContextInfo.Location;
-		const float DistanceSquared = DeltaLocation.SquaredLength();
-		const float MaxDistance = FMath::Min(AvailabilitiesMaxBroadPhaseRadius, OtherAnimContextInfo.AvailabilitiesMaxBroadPhaseRadius);
-		const float MaxDistanceSquared = MaxDistance * MaxDistance;
-		return DistanceSquared <= MaxDistanceSquared;
-	}
-
 	const FPoseSearchInteractionAnimContextAvailabilities* AnimContextAvailabilities = nullptr;
 
 	// cached AnimContext location
@@ -152,7 +135,7 @@ void UPoseSearchInteractionSubsystem::RegenerateAllIslands(float DeltaSeconds)
 
 	check(IsInGameThread());
 
-	// FScopeLock Lock(&AnimContextsAvailabilitiesMutex); is not necessary since UPoseSearchInteractionSubsystem gets ticked outside the parallel animation jobs
+	FMemMark Mark(FMemStack::Get());
 
 	// generating all the possible interaction tuples of AnimContext(s) with roles and pose histories (defined in FInteractionSearchContext)
 	FInteractionSearchContexts SearchContexts;
@@ -483,9 +466,24 @@ UPoseSearchInteractionSubsystem* UPoseSearchInteractionSubsystem::GetSubsystem_A
 	return nullptr;
 }
 
+#if WITH_EDITOR
+// to address check(PoseHistory && PoseHistory->AsWeak().IsValid()) in FInteractionIsland::FPreTickFunction::ExecuteTick caused by
+// run PIE, pause (FInteractionIsland are still caching the SearchContexts calculated with the previous pose history nodes), compile BPs, unpause, crash!
+void UPoseSearchInteractionSubsystem::OnBlueprintPreCompile(UBlueprint* BlueprintToCompile)
+{
+	check(IsInGameThread());
+	DestroyAllIslands();
+}
+#endif // WITH_EDITOR
+
 void UPoseSearchInteractionSubsystem::Initialize(FSubsystemCollectionBase& Collection)
 {
 	check(IsInGameThread());
+	
+#if WITH_EDITOR
+	GEditor->OnBlueprintPreCompile().AddUObject(this, &UPoseSearchInteractionSubsystem::OnBlueprintPreCompile);
+#endif // WITH_EDITOR
+
 	Super::Initialize(Collection);
 }
 
@@ -495,6 +493,10 @@ void UPoseSearchInteractionSubsystem::Deinitialize()
 
 	DestroyAllIslands();
 
+#if WITH_EDITOR
+	GEditor->OnBlueprintPreCompile().RemoveAll(this);
+#endif // WITH_EDITOR
+
 	Super::Deinitialize();
 }
 
@@ -525,19 +527,25 @@ void UPoseSearchInteractionSubsystem::GenerateAnimContextInfosAndTagToDatabases(
 	using namespace UE::PoseSearch;
 
 	const UWorld* SubsystemWorld = GetWorld();
+	TConstArrayView<FPoseSearchInteractionAnimContextAvailabilities> AnimContextsAvailabilities = MakeArrayView(AnimContextsAvailabilitiesBuffer.GetData(), AnimContextsAvailabilitiesNum);
+	
+#if DO_CHECK
 	check(SubsystemWorld);
 	check(AnimContextInfos.IsEmpty() && TagToDatabases.IsEmpty());
-
-	TConstArrayView<FPoseSearchInteractionAnimContextAvailabilities> AnimContextsAvailabilities = MakeArrayView(AnimContextsAvailabilitiesBuffer.GetData(), AnimContextsAvailabilitiesNum);
 	check(Algo::IsSorted(AnimContextsAvailabilities, [](const FPoseSearchInteractionAnimContextAvailabilities& A, const FPoseSearchInteractionAnimContextAvailabilities& B)
 		{
 			return A.AnimContext < B.AnimContext;
 		}));
+#endif // DO_CHECK
 
 	for (const FPoseSearchInteractionAnimContextAvailabilities& AnimContextAvailabilities : AnimContextsAvailabilities)
 	{
-		check(AnimContextAvailabilities.AnimContext && AnimContextAvailabilities.AnimContext->GetWorld() && AnimContextAvailabilities.AnimContext->GetWorld() == GetWorld());
+#if DO_CHECK
+		check(AnimContextAvailabilities.AnimContext);
+		const UWorld* AnimContextWorld = AnimContextAvailabilities.AnimContext->GetWorld();
+		check(AnimContextWorld == SubsystemWorld);
 		check(!AnimContextAvailabilities.Availabilities.IsEmpty());
+#endif // DO_CHECK
 		
 		// adding AnimContext to AnimContextsWithAvailabilities only if at least one availability has a valid database or has a valid tag
 		for (const FPoseSearchInteractionAvailabilityEx& InteractionAvailabilityEx : AnimContextAvailabilities.Availabilities)
@@ -566,7 +574,16 @@ void UPoseSearchInteractionSubsystem::GenerateAnimContextInfosAndTagToDatabases(
 	{
 		for (int32 AnimContextInfoIndexB = AnimContextInfoIndexA + 1; AnimContextInfoIndexB < NumAnimContextInfos; ++AnimContextInfoIndexB)
 		{
-			if (AnimContextInfos[AnimContextInfoIndexA].CanInteractWith(AnimContextInfos[AnimContextInfoIndexB]))
+			// performs broad phase analysis checking if at least one of the Availabilities associated to AnimContext can interact with OtherAnimContextInfo.
+			// This is a more relaxed analysis than the one performed in FRoledAnimContextInfo::CanInteractWith
+		
+			// @todo: filter OtherAnimContextInfo if dit doesn't have the "right role"? Maybe add a map of CanInteractWithRoles in this FAnimContextInfo
+		
+			const FVector DeltaLocation = AnimContextInfos[AnimContextInfoIndexA].Location - AnimContextInfos[AnimContextInfoIndexB].Location;
+			const float DistanceSquared = DeltaLocation.SquaredLength();
+			const float MaxDistance = FMath::Min(AnimContextInfos[AnimContextInfoIndexA].AvailabilitiesMaxBroadPhaseRadius, AnimContextInfos[AnimContextInfoIndexB].AvailabilitiesMaxBroadPhaseRadius);
+			const float MaxDistanceSquared = MaxDistance * MaxDistance;
+			if (DistanceSquared <= MaxDistanceSquared)
 			{
 				// the AnimContext of AnimContextInfos[AnimContextInfoIndexA] can potentially interact with the one from AnimContextInfos[AnimContextInfoIndexB]:
 				// linking AnimContextInfos[AnimContextInfoIndexA] to AnimContextInfos[AnimContextInfoIndexA] and vice versa to keep track of this when evaluating the broad phase.
@@ -594,34 +611,24 @@ void UPoseSearchInteractionSubsystem::GenerateSearchContexts(float DeltaSeconds,
 			, Role(InRole)
 			, PoseHistory(InPoseHistory)
 			, Database(&InDatabase)
+			, BroadPhaseRadiusSquared(InAvailability.BroadPhaseRadius)
+			, IncrementedBroadPhaseRadiusSquared(InAvailability.BroadPhaseRadius + InAvailability.BroadPhaseRadiusIncrementOnInteraction)
 		{
+			BroadPhaseRadiusSquared *= BroadPhaseRadiusSquared;
+			IncrementedBroadPhaseRadiusSquared *= IncrementedBroadPhaseRadiusSquared;
 		}
 
-		// perform narrow phase analysis checking if the AnimContextInfo with the specialized properties from Availability, Role, PoseHistory, Database,
-		// can interact with OtherRoledAnimContextInfo. This is a less relaxed analysis than the one performed in FAnimContextInfo::CanInteractWith
-		bool CanInteractWith(const FRoledAnimContextInfo& OtherRoledAnimContextInfo, bool bWasSearchContextInteracting) const
+		bool operator==(const FRoledAnimContextInfo& Other) const
 		{
-			check(this != &OtherRoledAnimContextInfo);
-
-			const FVector DeltaLocation = AnimContextInfo->Location - OtherRoledAnimContextInfo.AnimContextInfo->Location;
-			const float DistanceSquared = DeltaLocation.SquaredLength();
-
-			float MaxDistance;
-			if (bWasSearchContextInteracting)
-			{
-				MaxDistance = FMath::Min(Availability->BroadPhaseRadius + Availability->BroadPhaseRadiusIncrementOnInteraction, OtherRoledAnimContextInfo.Availability->BroadPhaseRadius + OtherRoledAnimContextInfo.Availability->BroadPhaseRadiusIncrementOnInteraction);
-			}
-			else
-			{
-				MaxDistance = FMath::Min(Availability->BroadPhaseRadius, OtherRoledAnimContextInfo.Availability->BroadPhaseRadius);
-			}
-			
-			const float MaxDistanceSquared = MaxDistance * MaxDistance;
-			return DistanceSquared <= MaxDistanceSquared;
+			// no need to check for BroadPhaseRadiusSquared and IncrementedBroadPhaseRadiusSquared since they are derived values from Availability
+			return
+				Availability == Other.Availability &&
+				AnimContextInfo == Other.AnimContextInfo &&
+				Role == Other.Role &&
+				PoseHistory == Other.PoseHistory &&
+				Database == Other.Database;
 		}
 
-		bool operator==(const FRoledAnimContextInfo& Other) const = default;
-
 		// Availability that spawned this FRoledAnimContextInfo
 		const FPoseSearchInteractionAvailabilityEx* Availability = nullptr;
 		// AnimContextInfo containing all the information regarding the AnimContext that spawned this FRoledAnimContextInfo,
@@ -630,6 +637,9 @@ void UPoseSearchInteractionSubsystem::GenerateSearchContexts(float DeltaSeconds,
 		const FRole Role = DefaultRole;
 		const IPoseHistory* PoseHistory = nullptr;
 		const UPoseSearchDatabase* Database = nullptr;
+
+		float BroadPhaseRadiusSquared = 0.f;
+		float IncrementedBroadPhaseRadiusSquared = 0.f;
 	};
 	
 	struct FRoledAnimContextInfos : TArray<FRoledAnimContextInfo, TInlineAllocator<PreallocatedRolesNum, TMemStackAllocator<>>>
@@ -724,8 +734,6 @@ void UPoseSearchInteractionSubsystem::GenerateSearchContexts(float DeltaSeconds,
 	FTagToDatabases TagToDatabases;
 	GenerateAnimContextInfosAndTagToDatabases(AnimContextInfos, TagToDatabases);
 
-	const TConstArrayView<FValidInteractionSearch> PreviousValidInteractionSearches = ValidInteractionSearches;
-
 	TMap<const UPoseSearchDatabase*, FRoledAnimContextInfos, TInlineSetAllocator<PreallocatedSearchesNum, TMemStackSetAllocator<>>> DatabaseToRoledAnimContextInfos;
 	// visiting ALL the AnimContexts in AnimContextInfos and relying on the cached information to refine potential interactions
 	FAnimContextInfoVisitor AnimContextInfoVisitor(AnimContextInfos,
@@ -774,7 +782,7 @@ void UPoseSearchInteractionSubsystem::GenerateSearchContexts(float DeltaSeconds,
 			}
 		},
 		// OnDoneGroupingAnimContexts: called when the FAnimContextInfoVisitor reaches the end of a group of possibly interacting AnimContext(s)
-		[&DatabaseToRoledAnimContextInfos, &SearchContexts, &PreviousValidInteractionSearches]()
+		[this, &DatabaseToRoledAnimContextInfos, &SearchContexts]()
 		{
 			// for each database now we try to create all the possible combinations of the roled anim instances
 			// for example, given a database set up with assets for 2 characters interactions with roles RoleA and RoleB
@@ -811,23 +819,27 @@ void UPoseSearchInteractionSubsystem::GenerateSearchContexts(float DeltaSeconds,
 
 				GenerateCombinations(RoledAnimContextInfos.Num(), CombinationCardinality,
 					// Combination is an array of indexes in RoledAnimContextInfos: 0 <= Combination[i] < RoledAnimContextInfos.Num()
-					[Database, &RoledSkeletons, &RoledAnimContextInfos, &SearchContexts, &PreviousValidInteractionSearches](const TConstArrayView<int32> Combination)
+					[this, Database, &RoledSkeletons, &RoledAnimContextInfos, &SearchContexts](const TConstArrayView<int32> Combination)
 					{
 						// CombinationCardinality represents the number of roles as well as the number interacting AnimContext(s) (ultimately number of Characters involved in the interaction)
 						const int32 CombinationCardinality = Combination.Num();
-						TSet<const UObject*, DefaultKeyFuncs<const UObject*>, TInlineSetAllocator<PreallocatedRolesNum, TMemStackSetAllocator<>>> UniqueAnimContexts;
 
 						for (int32 CombinationIndex = 0; CombinationIndex < CombinationCardinality; ++CombinationIndex)
 						{
-							const int32 RoledAnimContextIndex = Combination[CombinationIndex];
-							const FRoledAnimContextInfo& RoledAnimContextInfo = RoledAnimContextInfos[RoledAnimContextIndex];
-							check(RoledAnimContextInfo.AnimContextInfo);
-							bool bIsAlreadyInSet;
-							UniqueAnimContexts.Add(RoledAnimContextInfo.AnimContextInfo->AnimContextAvailabilities->AnimContext, &bIsAlreadyInSet);
-							if (bIsAlreadyInSet)
+							const FRoledAnimContextInfo& RoledAnimContextInfo = RoledAnimContextInfos[Combination[CombinationIndex]];
+
+							// since CombinationCardinality is usually small 2-3, it's faster to check linearly for duplicated animcontexts rather than holding them in a set
+							const UObject* CombinationIndexAnimContext = RoledAnimContextInfo.AnimContextInfo->AnimContextAvailabilities->AnimContext;
+							for (int32 NextCombinationIndex = CombinationIndex + 1; NextCombinationIndex < CombinationCardinality; ++NextCombinationIndex)
 							{
-								// we have a duplicate AnimContext. this combination is NOT valid
-								return false;
+								const FRoledAnimContextInfo& NextRoledAnimContextInfo = RoledAnimContextInfos[Combination[NextCombinationIndex]];
+								const UObject* NextCombinationIndexAnimContext = NextRoledAnimContextInfo.AnimContextInfo->AnimContextAvailabilities->AnimContext;
+
+								if (CombinationIndexAnimContext == NextCombinationIndexAnimContext)
+								{
+									// we have a duplicate AnimContext. this combination is NOT valid
+									return false;
+								}
 							}
 						}
 
@@ -847,7 +859,8 @@ void UPoseSearchInteractionSubsystem::GenerateSearchContexts(float DeltaSeconds,
 							SearchContext.TickPriorities.Add(RoledAnimContextInfo.Availability->TickPriority);
 							
 							#if ENABLE_ANIM_DEBUG
-							SearchContext.DebugAvailabilities.Add(*RoledAnimContextInfo.Availability);
+							SearchContext.DebugBroadPhaseRadius.Add(RoledAnimContextInfo.Availability->BroadPhaseRadius);
+							SearchContext.BroadPhaseRadiusIncrementOnInteraction.Add(RoledAnimContextInfo.Availability->BroadPhaseRadiusIncrementOnInteraction);
 							#endif // ENABLE_ANIM_DEBUG
 						}
 
@@ -861,15 +874,9 @@ void UPoseSearchInteractionSubsystem::GenerateSearchContexts(float DeltaSeconds,
 							}
 						}
 
-						// looking for a preexisting valid interaction resmbling SearchContext
-						for (const FValidInteractionSearch& PreviousValidInteractionSearch : PreviousValidInteractionSearches)
-						{
-							if (PreviousValidInteractionSearch.SearchContext.IsEquivalent(SearchContext))
-							{
-								SearchContext.bIsContinuingInteraction = true;
-								break;
-							}
-						}
+						// looking for a preexisting valid interaction resembling SearchContext
+						const uint32 SearchContextHash = SearchContext.GetHashForEquivalence();
+						SearchContext.bIsContinuingInteraction = ValidInteractionSearches.Contains(SearchContextHash);
 
 						// checking if this combination is valid for the Database:
 						for (int32 CombinationIndex = 0; CombinationIndex < CombinationCardinality; ++CombinationIndex)
@@ -877,27 +884,33 @@ void UPoseSearchInteractionSubsystem::GenerateSearchContexts(float DeltaSeconds,
 							const int32 RoledAnimContextIndex = Combination[CombinationIndex];
 							const FRoledAnimContextInfo& RoledAnimContextInfo = RoledAnimContextInfos[RoledAnimContextIndex];
 
-							// checking the narrow phase!
+							// checking the narrow phase! if RoledAnimContextInfo can interact with OtherRoledAnimContextInfo
 							for (int32 OtherCombinationIndex = CombinationIndex + 1; OtherCombinationIndex < CombinationCardinality; ++OtherCombinationIndex)
 							{
 								const int32 OtherRoledAnimContextIndex = Combination[OtherCombinationIndex];
 								const FRoledAnimContextInfo& OtherRoledAnimContextInfo = RoledAnimContextInfos[OtherRoledAnimContextIndex];
 
+								const FVector DeltaLocation = RoledAnimContextInfo.AnimContextInfo->Location - OtherRoledAnimContextInfo.AnimContextInfo->Location;
+								const float DistanceSquared = DeltaLocation.SquaredLength();
+
+								float MaxDistanceSquared;
+								if (SearchContext.bIsContinuingInteraction)
+								{
+									MaxDistanceSquared = FMath::Min(RoledAnimContextInfo.IncrementedBroadPhaseRadiusSquared, OtherRoledAnimContextInfo.IncrementedBroadPhaseRadiusSquared);
+								}
+								else
+								{
+									MaxDistanceSquared = FMath::Min(RoledAnimContextInfo.BroadPhaseRadiusSquared, OtherRoledAnimContextInfo.BroadPhaseRadiusSquared);
+								}
+			
 								// if any of the RoledAnimContextInfo cannot interact with any OtherRoledAnimContextInfo the inteaction cannot happen!
-								if (!RoledAnimContextInfo.CanInteractWith(OtherRoledAnimContextInfo, SearchContext.bIsContinuingInteraction))
+								if (DistanceSquared > MaxDistanceSquared)
 								{
 									return false;
 								}
 							}
 						}
-#if DO_CHECK
-						for (const FInteractionSearchContext& ContainedSearchContext : SearchContexts)
-						{
-							check(!ContainedSearchContext.IsEquivalent(SearchContext));
-						}
 
-						check(SearchContext.CheckForConsistency());
-#endif // DO_CHECK
 						SearchContexts.Emplace(SearchContext);
 						return true;
 					});
@@ -916,57 +929,16 @@ void UPoseSearchInteractionSubsystem::OnInteractionStart(UE::PoseSearch::FValidI
 	using namespace UE::PoseSearch;
 
 #if ENABLE_VISUAL_LOG
-	ValidInteractionSearch.SearchContext.VLogContext(FColor::Blue);
+	ValidInteractionSearch.VLogContext(FColor::Blue);
 #endif
 
-	check(ValidInteractionSearch.DisabledCollisions.IsEmpty());
-	if (ValidInteractionSearch.SearchContext.bDisableCollisions)
-	{
-		TArray<AActor*, TInlineAllocator<PreallocatedRolesNum, TMemStackAllocator<>>> Actors;
-		TArray<UPrimitiveComponent*, TInlineAllocator<PreallocatedRolesNum, TMemStackAllocator<>>> PrimitiveComponents;
-
-		for (int32 AnimContextIndex = 0; AnimContextIndex < ValidInteractionSearch.SearchContext.Num(); ++AnimContextIndex)
-		{
-			if (const UObject* AnimContext = ValidInteractionSearch.SearchContext.GetAnimContext(AnimContextIndex))
-			{
-				AActor* Actor = const_cast<AActor*>(GetContextOwningActor(AnimContext, false));
-				check(Actor);
-				Actors.Add(Actor);
-				PrimitiveComponents.Add(Cast<UPrimitiveComponent>(Actor->GetRootComponent()));
-			}
-		}
-
-		for (int32 IndexA = 0; IndexA < Actors.Num(); ++IndexA)
-		{
-			for (int32 IndexB = IndexA + 1; IndexB < Actors.Num(); ++IndexB)
-			{
-				AActor* ActorA = Actors[IndexA];
-				AActor* ActorB = Actors[IndexB];
-				check(ActorA && ActorB);
-
-				UPrimitiveComponent* PrimitiveComponentA = PrimitiveComponents[IndexA];
-				UPrimitiveComponent* PrimitiveComponentB = PrimitiveComponents[IndexB];
-
-				if (PrimitiveComponentA && !PrimitiveComponentA->GetMoveIgnoreActors().Contains(ActorB))
-				{
-					ValidInteractionSearch.DisabledCollisions.Add({ ActorA, ActorB });
-					PrimitiveComponentA->IgnoreActorWhenMoving(ActorB, true);
-				}
-
-				if (PrimitiveComponentB && !PrimitiveComponentB->GetMoveIgnoreActors().Contains(ActorA))
-				{
-					ValidInteractionSearch.DisabledCollisions.Add({ ActorB, ActorA });
-					PrimitiveComponentB->IgnoreActorWhenMoving(ActorA, true);
-				}
-			}
-		}
-	}
+	ValidInteractionSearch.DoDisableCollisions();
 }
 
 void UPoseSearchInteractionSubsystem::OnInteractionContinuing(UE::PoseSearch::FValidInteractionSearch& ValidInteractionSearch)
 {
 #if ENABLE_VISUAL_LOG
-	ValidInteractionSearch.SearchContext.VLogContext(FColor::Green);
+	ValidInteractionSearch.VLogContext(FColor::Green);
 #endif
 }
 
@@ -975,33 +947,18 @@ void UPoseSearchInteractionSubsystem::OnInteractionEnd(UE::PoseSearch::FValidInt
 	using namespace UE::PoseSearch;
 
 #if ENABLE_VISUAL_LOG
-	ValidInteractionSearch.SearchContext.VLogContext(FColor::Black);
+	ValidInteractionSearch.VLogContext(FColor::Black);
 #endif
 
-	for (const FDisabledCollisions::ElementType& DisabledCollision : ValidInteractionSearch.DisabledCollisions)
-	{
-		if (AActor* ActorA = DisabledCollision.Key.Get())
-		{
-			if (AActor* ActorB = DisabledCollision.Value.Get())
-			{
-				if (UPrimitiveComponent* PrimitiveComponentA = Cast<UPrimitiveComponent>(ActorA->GetRootComponent()))
-				{
-					PrimitiveComponentA->IgnoreActorWhenMoving(ActorB, false);
-				}
-			}
-		}
-	}
+	ValidInteractionSearch.UndoDisableCollisions();
 }
 
 void UPoseSearchInteractionSubsystem::UpdateValidInteractionSearches()
 {
 	using namespace UE::PoseSearch;
 
-	const int32 ValidInteractionSearchesNum = ValidInteractionSearches.Num();
-	TArray<bool, TInlineAllocator<PreallocatedSearchesNum * 4>> Visited;
-	Visited.SetNum(ValidInteractionSearchesNum);
-
-	TArray<FValidInteractionSearch, TInlineAllocator<PreallocatedSearchesNum * 4>> NewValidInteractionSearches;
+	FMemMark Mark(FMemStack::Get());
+	TMap<uint32, bool, TInlineSetAllocator<32, TMemStackSetAllocator<>>> ContinuingOrNewValidSearchContexts;
 	for (FInteractionIsland* Island : Islands)
 	{
 		if (Island->IsInitialized())
@@ -1010,44 +967,43 @@ void UPoseSearchInteractionSubsystem::UpdateValidInteractionSearches()
 			for (const FInteractionSearchResult& SearchResult : Island->GetSearchResults())
 			{
 				const FInteractionSearchContext& SearchContext = Island->GetSearchContexts()[SearchResult.SearchIndex];
-
-				int32 Index = 0;
-				for (; Index < ValidInteractionSearchesNum; ++Index)
+				const uint32 ValidSearchContextHash = SearchContext.GetHashForEquivalence();
+				if (FValidInteractionSearch* ContinuingValidInteractionSearch = ValidInteractionSearches.Find(ValidSearchContextHash))
 				{
-					if (ValidInteractionSearches[Index].SearchContext.IsEquivalent(SearchContext))
-					{
-						check(!Visited[Index]);
-						Visited[Index] = true;
-						OnInteractionContinuing(ValidInteractionSearches[Index]);
-						break;
-					}
+					ContinuingOrNewValidSearchContexts.Add(ValidSearchContextHash) = true;
 				}
-
-				if (Index == ValidInteractionSearchesNum)
+				else
 				{
-					// we haven't found an equivalent SearchContext in ValidInteractionSearches, so it's a new interaction!
-					FValidInteractionSearch& NewValidInteractionSearch = NewValidInteractionSearches.AddDefaulted_GetRef();
-					NewValidInteractionSearch.SearchContext = SearchContext;
-					OnInteractionStart(NewValidInteractionSearch);
+					ContinuingOrNewValidSearchContexts.Add(ValidSearchContextHash) = false;
+					ValidInteractionSearches.Add(ValidSearchContextHash).Init(SearchContext);
 				}
 			}
 		}
 	}
 
-	// checking for leftover unvisited ValidInteractionSearches. Those are interactions that just ended
-	for (int32 Index = 0; Index < ValidInteractionSearchesNum; ++Index)
+	// calling all the OnInteractionEnd first
+	for (TMap<uint32, UE::PoseSearch::FValidInteractionSearch>::TIterator It = ValidInteractionSearches.CreateIterator(); It; ++It)
 	{
-		if (Visited[Index])
+		if (!ContinuingOrNewValidSearchContexts.Find(It->Key))
 		{
-			NewValidInteractionSearches.Add(ValidInteractionSearches[Index]);
+			OnInteractionEnd(It->Value);
+			It.RemoveCurrent();
+		}
+	}
+
+	// calling all the OnInteractionContinuing and OnInteractionStart after
+	for (TMap<uint32, UE::PoseSearch::FValidInteractionSearch>::TIterator It = ValidInteractionSearches.CreateIterator(); It; ++It)
+	{
+		const bool bIsContinuingValidSearchContexts = ContinuingOrNewValidSearchContexts.FindChecked(It->Key);
+		if (bIsContinuingValidSearchContexts)
+		{
+			OnInteractionContinuing(It->Value);
 		}
 		else
 		{
-			OnInteractionEnd(ValidInteractionSearches[Index]);
+			OnInteractionStart(It->Value);
 		}
 	}
-
-	ValidInteractionSearches = NewValidInteractionSearches;
 }
 
 bool UPoseSearchInteractionSubsystem::ConsolidateAnimContextsAvailabilities()
@@ -1060,7 +1016,7 @@ bool UPoseSearchInteractionSubsystem::ConsolidateAnimContextsAvailabilities()
 	// @todo: add some setttings to initialize AnimContextsAvailabilitiesBuffer to a big enough value
 	if (AnimContextsAvailabilitiesIndexValue > AnimContextsAvailabilitiesNum)
 	{
-		UE_LOG(LogPoseSearch, Error, TEXT("UPoseSearchInteractionSubsystem::ConsolidateAnimContextsAvailabilities not enough space to add more availabilities locklessly. It'll be adjusted automatically, but some availability requests has been lost this frame [capacity %d / requests %d]"), AnimContextsAvailabilitiesNum, AnimContextsAvailabilitiesIndexValue);
+		UE_LOG(LogPoseSearch, Warning, TEXT("UPoseSearchInteractionSubsystem::ConsolidateAnimContextsAvailabilities not enough space to add more availabilities locklessly. It'll be adjusted automatically, but some availability requests has been lost this frame [capacity %d / requests %d]"), AnimContextsAvailabilitiesNum, AnimContextsAvailabilitiesIndexValue);
 		AnimContextsAvailabilitiesBuffer.SetNum(AnimContextsAvailabilitiesIndexValue);
 	}
 	else
@@ -1166,7 +1122,7 @@ void UPoseSearchInteractionSubsystem::Query_AnyThread(const TArrayView<const FPo
 	// if we find AnimContext in an island, we perform ALL the Island motion matching searches.
 	if (FInteractionIsland* Island = FindIsland(AnimContext))
 	{
-		Island->DoSearch_AnyThread(AnimContext, ValidInteractionSearches, Result);
+		Island->DoSearch_AnyThread(AnimContext, Result);
 
 		if (bValidateResultAgainstAvailabilities && Result.SelectedAnim)
 		{
@@ -1216,7 +1172,8 @@ void UPoseSearchInteractionSubsystem::DebugDrawIslands() const
 
 	check(IsInGameThread());
 
-	if (!FVisualLogger::IsRecording())
+	const bool bChannelEnabled = UE_TRACE_CHANNELEXPR_IS_ENABLED(PoseSearchChannel);
+	if (!bChannelEnabled)
 	{
 		return;
 	}
@@ -1264,15 +1221,14 @@ void UPoseSearchInteractionSubsystem::DebugDrawIslands() const
 				{
 					if (const UObject* AnimContext = SearchContext.GetAnimContext(Index))
 					{
-						const FPoseSearchInteractionAvailability& DebugAvailability = SearchContext.DebugAvailabilities[Index];
 						float MaxBroadPhaseRadius;
 						if (SearchContext.bIsContinuingInteraction)
 						{
-							MaxBroadPhaseRadius = DebugAvailability.BroadPhaseRadius + DebugAvailability.BroadPhaseRadiusIncrementOnInteraction;
+							MaxBroadPhaseRadius = SearchContext.DebugBroadPhaseRadius[Index] + SearchContext.BroadPhaseRadiusIncrementOnInteraction[Index];
 						}
 						else
 						{
-							MaxBroadPhaseRadius = DebugAvailability.BroadPhaseRadius;
+							MaxBroadPhaseRadius = SearchContext.DebugBroadPhaseRadius[Index];
 						}
 
 						if (MaxBroadPhaseRadius > UE_SMALL_NUMBER)
diff --git a/Engine/Plugins/Animation/PoseSearch/Source/Runtime/Public/PoseSearch/PoseSearchInteractionIsland.h b/Engine/Plugins/Animation/PoseSearch/Source/Runtime/Public/PoseSearch/PoseSearchInteractionIsland.h
--- a/Engine/Plugins/Animation/PoseSearch/Source/Runtime/Public/PoseSearch/PoseSearchInteractionIsland.h
+++ b/Engine/Plugins/Animation/PoseSearch/Source/Runtime/Public/PoseSearch/PoseSearchInteractionIsland.h
@@ -16,91 +16,73 @@ namespace UE::PoseSearch
 {
 struct FInteractionIsland;
 struct IInteractionIslandDependency;
+struct FInteractionSearchContext;
 
 // Experimental, this feature might be removed without warning, not for production use
-struct FInteractionSearchContextBase
+typedef TArray<TPair<TWeakObjectPtr<AActor>, TWeakObjectPtr<AActor>>> FDisabledCollisions;
+
+// Experimental, this feature might be removed without warning, not for production use
+struct FValidInteractionSearch
 {
-private:
-	TArray<TWeakObjectPtr<const UObject>, TInlineAllocator<PreallocatedRolesNum>> AnimContexts;
-	TArray<TWeakPtr<const UE::PoseSearch::IPoseHistory, ESPMode::ThreadSafe>, TInlineAllocator<PreallocatedRolesNum>> PoseHistories;
-	TArray<UE::PoseSearch::FRole, TInlineAllocator<PreallocatedRolesNum>> Roles;
+	void Init(const FInteractionSearchContext& Other);
+	void DoDisableCollisions();
+	void UndoDisableCollisions();
+
+#if ENABLE_VISUAL_LOG
+	void VLogContext(const FColor& Color) const;
+#endif // ENABLE_VISUAL_LOG
 
-public:
+private:
 	TWeakObjectPtr<const UPoseSearchDatabase> Database;
 	bool bDisableCollisions = false;
 
-	void Add(const UObject* AnimContext, const UE::PoseSearch::IPoseHistory* PoseHistory, const UE::PoseSearch::FRole Role)
-	{
-		check(AnimContext && PoseHistory);
-		
-		// AnimContexts must be sorted to have deterministic searches across multiple frames
-		// as well as for faster IsEquivalent implementation
-		check(AnimContexts.IsEmpty() || GetAnimContext(AnimContexts.Num() - 1) < AnimContext);
-
-		AnimContexts.Add(AnimContext);
-		PoseHistories.Add(PoseHistory->AsWeak());
-		Roles.Add(Role);
-	}
+	TArray<TWeakObjectPtr<const UObject>, TInlineAllocator<PreallocatedRolesNum>> AnimContexts;
+	TArray<TWeakPtr<const UE::PoseSearch::IPoseHistory, ESPMode::ThreadSafe>, TInlineAllocator<PreallocatedRolesNum>> PoseHistories;
+	TArray<UE::PoseSearch::FRole, TInlineAllocator<PreallocatedRolesNum>> Roles;
+	FDisabledCollisions DisabledCollisions;
+};
 
+// Experimental, this feature might be removed without warning, not for production use
+// stack allocated temporary struct used to generate interaction contexts, then stored using weak pointers as FValidInteractionSearch
+struct FInteractionSearchContext
+{
+	uint32 GetHashForEquivalence() const;
+	void Add(const UObject* AnimContext, const UE::PoseSearch::IPoseHistory* PoseHistory, const UE::PoseSearch::FRole Role);
 	int32 Num() const { return AnimContexts.Num(); }
-	const UObject* GetAnimContext(int32 Index) const { return AnimContexts[Index].Get(); }
+	const UObject* GetAnimContext(int32 Index) const { return AnimContexts[Index]; }
 	const UE::PoseSearch::IPoseHistory* GetPoseHistory(int32 Index) const;
 	const UE::PoseSearch::FRole GetRole(int32 Index) const { return Roles[Index]; }
 	const TConstArrayView<UE::PoseSearch::FRole> GetRoles() const { return Roles; }
+	FPoseSearchContinuingProperties GetContinuingProperties() const;
 
 	bool IsValid() const { return !AnimContexts.IsEmpty(); }
-	bool IsEquivalent(const FInteractionSearchContextBase& Other) const;
-
-#if ENABLE_VISUAL_LOG
-	void VLogContext(const FColor& Color) const;
-#endif // ENABLE_VISUAL_LOG
+	bool IsEquivalent(const FInteractionSearchContext& Other) const;
 
 #if DO_CHECK
-	// method that queries the HistoryCollectors. the PoseHistory will complain if accessed in a non thread safe manner!
-	void TestHistoryCollectorsThreadingAccess() const;
 	bool CheckForConsistency() const;
 #endif // DO_CHECK
-};
 
-// Experimental, this feature might be removed without warning, not for production use
-typedef TArray<TPair<TWeakObjectPtr<AActor>, TWeakObjectPtr<AActor>>> FDisabledCollisions;
+	const UPoseSearchDatabase* Database = nullptr;
+	bool bDisableCollisions = false;
 
-// Experimental, this feature might be removed without warning, not for production use
-struct FValidInteractionSearch
-{
-	FInteractionSearchContextBase SearchContext;
-	FDisabledCollisions DisabledCollisions;
-};
+	TArray<const UObject*, TInlineAllocator<PreallocatedRolesNum, TMemStackAllocator<>>> AnimContexts;
+	TArray<const UE::PoseSearch::IPoseHistory*, TInlineAllocator<PreallocatedRolesNum, TMemStackAllocator<>>> PoseHistories;
+	TArray<UE::PoseSearch::FRole, TInlineAllocator<PreallocatedRolesNum, TMemStackAllocator<>>> Roles;
 
-// @todo: make FPoseSearchContinuingProperties.PlayingAsset a TWeakObjectPtr and remove this struct!
-// Experimental, this feature might be removed without warning, not for production use
-struct FInteractionSearchContext : FInteractionSearchContextBase
-{
-	// TWeakObjectPtr version of FPoseSearchContinuingProperties to be GC friendly
-	TWeakObjectPtr<const UObject> PlayingAsset = nullptr;
+	const UObject* PlayingAsset = nullptr;
 	float PlayingAssetAccumulatedTime = 0.f;
 	bool bIsPlayingAssetMirrored = false;
 	FVector PlayingAssetBlendParameters = FVector::ZeroVector;
 	EPoseSearchInterruptMode InterruptMode = EPoseSearchInterruptMode::DoNotInterrupt;
 	bool bIsContinuingInteraction = false;
-	TArray<int32, TInlineAllocator<PreallocatedRolesNum>> TickPriorities;
+	TArray<int32, TInlineAllocator<PreallocatedRolesNum, TMemStackAllocator<>>> TickPriorities;
 
 #if ENABLE_ANIM_DEBUG
-	TArray<FPoseSearchInteractionAvailability, TInlineAllocator<PreallocatedRolesNum>> DebugAvailabilities;
+	TArray<float, TInlineAllocator<PreallocatedRolesNum, TMemStackAllocator<>>> DebugBroadPhaseRadius;
+	TArray<float, TInlineAllocator<PreallocatedRolesNum, TMemStackAllocator<>>> BroadPhaseRadiusIncrementOnInteraction;
 #endif // ENABLE_ANIM_DEBUG
-
-	FPoseSearchContinuingProperties GetContinuingProperties() const
-	{
-		FPoseSearchContinuingProperties ContinuingProperties;
-		ContinuingProperties.PlayingAsset = PlayingAsset.Get();
-		ContinuingProperties.PlayingAssetAccumulatedTime = PlayingAssetAccumulatedTime;
-		ContinuingProperties.bIsPlayingAssetMirrored = bIsPlayingAssetMirrored;
-		ContinuingProperties.PlayingAssetBlendParameters = PlayingAssetBlendParameters;
-		ContinuingProperties.InterruptMode = InterruptMode;
-		ContinuingProperties.bIsContinuingInteraction = bIsContinuingInteraction;
-		return ContinuingProperties;
-	}
 };
+
 // Experimental, this feature might be removed without warning, not for production use
 typedef TArray<FInteractionSearchContext, TInlineAllocator<PreallocatedSearchesNum, TMemStackAllocator<>>> FInteractionSearchContexts;
 
@@ -132,7 +114,7 @@ struct FInteractionIsland
 	FInteractionIsland(ULevel* Level, UPoseSearchInteractionSubsystem* Subsystem);
 	~FInteractionIsland();
 
-	bool DoSearch_AnyThread(const UObject* AnimContext, const TConstArrayView<FValidInteractionSearch> ValidInteractionSearches, FPoseSearchBlueprintResult& Result);
+	bool DoSearch_AnyThread(const UObject* AnimContext, FPoseSearchBlueprintResult& Result);
 	bool GetResult_AnyThread(const UObject* AnimContext, FPoseSearchBlueprintResult& Result, bool bCompareOwningActors = false);
 
 	const TArray<TWeakObjectPtr<UActorComponent>>& GetTickActorComponents() const { return TickActorComponents; }
diff --git a/Engine/Plugins/Animation/PoseSearch/Source/Runtime/Public/PoseSearch/PoseSearchInteractionSubsystem.h b/Engine/Plugins/Animation/PoseSearch/Source/Runtime/Public/PoseSearch/PoseSearchInteractionSubsystem.h
--- a/Engine/Plugins/Animation/PoseSearch/Source/Runtime/Public/PoseSearch/PoseSearchInteractionSubsystem.h
+++ b/Engine/Plugins/Animation/PoseSearch/Source/Runtime/Public/PoseSearch/PoseSearchInteractionSubsystem.h
@@ -124,6 +124,10 @@ private:
 	static void OnInteractionContinuing(UE::PoseSearch::FValidInteractionSearch& ValidInteractionSearch);
 	static void OnInteractionEnd(UE::PoseSearch::FValidInteractionSearch& ValidInteractionSearch);
 
+#if WITH_EDITOR
+	void OnBlueprintPreCompile(UBlueprint* BlueprintToCompile);
+#endif // WITH_EDITOR
+
 	// it could contain duplicated availabilities
 	UPROPERTY(Transient)
 	TArray<FPoseSearchInteractionAnimContextAvailabilities> AnimContextsAvailabilitiesBuffer;
@@ -133,7 +137,7 @@ private:
 	// array of groups of characters that needs to be anaylzed together for possible interactions
 	TArray<UE::PoseSearch::FInteractionIsland*> Islands;
 
-	TArray<UE::PoseSearch::FValidInteractionSearch> ValidInteractionSearches;
+	TMap<uint32, UE::PoseSearch::FValidInteractionSearch> ValidInteractionSearches;
 };
 
 #undef UE_API
diff --git a/Engine/Plugins/Animation/PoseSearch/Source/Runtime/Public/PoseSearch/PoseSearchInteractionUtils.h b/Engine/Plugins/Animation/PoseSearch/Source/Runtime/Public/PoseSearch/PoseSearchInteractionUtils.h
--- a/Engine/Plugins/Animation/PoseSearch/Source/Runtime/Public/PoseSearch/PoseSearchInteractionUtils.h
+++ b/Engine/Plugins/Animation/PoseSearch/Source/Runtime/Public/PoseSearch/PoseSearchInteractionUtils.h
@@ -25,9 +25,9 @@ POSESEARCH_API void CalculateFullAlignedTransforms(const FPoseSearchBlueprintRes
 // Experimental, this feature might be removed without warning, not for production use
 POSESEARCH_API FTransform CalculateDeltaAlignment(const FTransform& MeshWithoutOffset, const FTransform& MeshWithOffset, const FTransform& FullAlignedTransform, float WarpingRotationRatio, float WarpingTranslationRatio);
 
-
+// Experimental, this feature might be removed without warning, not for production use
 template <typename EvaluateCombinationType>
-static void GenerateCombinationsRecursive(int32 DataCardinality, int32 DataIndex, TArrayView<int32> Combination, int32 CombinationIndex, EvaluateCombinationType EvaluateCombination)
+static void GenerateCombinationsRecursive(int32 DataCardinality, int32 DataIndex, TArray<int32, TInlineAllocator<PreallocatedRolesNum>>& Combination, int32 CombinationIndex, EvaluateCombinationType EvaluateCombination)
 {
 	if (CombinationIndex == Combination.Num())
 	{ 
@@ -41,6 +41,7 @@ static void GenerateCombinationsRecursive(int32 DataCardinality, int32 DataIndex
 	}
 } 
 
+// Experimental, this feature might be removed without warning, not for production use
 // generates all the possible (unique) combinations for the indexes of a set of DataCardinality number of elements grouped in tuples of CombinationCardinality size:
 // the number of generated combinations P(DataCardinality, CombinationCardinality) = (DataCardinality! / (DataCardinality - CombinationCardinality)!) / CombinationCardinality!
 // see: https://en.wikipedia.org/wiki/Combination
