Fix nullptr access of GEditor

diff --git a/Engine/Plugins/Animation/PoseSearch/Source/Runtime/Private/PoseSearchInteractionSubsystem.cpp b/Engine/Plugins/Animation/PoseSearch/Source/Runtime/Private/PoseSearchInteractionSubsystem.cpp
--- a/Engine/Plugins/Animation/PoseSearch/Source/Runtime/Private/PoseSearchInteractionSubsystem.cpp
+++ b/Engine/Plugins/Animation/PoseSearch/Source/Runtime/Private/PoseSearchInteractionSubsystem.cpp
@@ -481,7 +481,10 @@ void UPoseSearchInteractionSubsystem::Initialize(FSubsystemCollectionBase& Colle
 	check(IsInGameThread());
 	
 #if WITH_EDITOR
-	GEditor->OnBlueprintPreCompile().AddUObject(this, &UPoseSearchInteractionSubsystem::OnBlueprintPreCompile);
+	if (GEditor != nullptr)
+	{
+		GEditor->OnBlueprintPreCompile().AddUObject(this, &UPoseSearchInteractionSubsystem::OnBlueprintPreCompile);
+	}
 #endif // WITH_EDITOR
 
 	Super::Initialize(Collection);
@@ -494,7 +497,10 @@ void UPoseSearchInteractionSubsystem::Deinitialize()
 	DestroyAllIslands();
 
 #if WITH_EDITOR
-	GEditor->OnBlueprintPreCompile().RemoveAll(this);
+	if (GEditor != nullptr)
+	{
+		GEditor->OnBlueprintPreCompile().RemoveAll(this);
+	}
 #endif // WITH_EDITOR
 
 	Super::Deinitialize();
