PoseSearch - avoiding possible division by 0 in FDebugDrawParams::DrawWedge

diff --git a/Engine/Plugins/Animation/PoseSearch/Source/Runtime/Private/PoseSearchContext.cpp b/Engine/Plugins/Animation/PoseSearch/Source/Runtime/Private/PoseSearchContext.cpp
--- a/Engine/Plugins/Animation/PoseSearch/Source/Runtime/Private/PoseSearchContext.cpp
+++ b/Engine/Plugins/Animation/PoseSearch/Source/Runtime/Private/PoseSearchContext.cpp
@@ -594,7 +594,7 @@ void FDebugDrawParams::DrawCircle(const FMatrix& TransformMatrix, float Radius,
 void FDebugDrawParams::DrawWedge(const FVector& Origin, const FVector& Direction, float InnerRadius, float OuterRadius, float Width, int32 Segments, const FColor& Color, float Thickness) const
 {
 	FVector NormalizedDirection = Direction;
-	if (Color.A > 0 && Segments > 0 && NormalizedDirection.Normalize())
+	if (Color.A > 0 && Segments > 1 && NormalizedDirection.Normalize())
 	{
 		float AngleDeg = -Width * 0.5f;
 		const float AngleDegIncrement = Width / (Segments - 1);
