PoseSearch - MM interaction optimizations: avoiding costly availability sorting and performing data structure sorting only when needed during interaction islands update

diff --git a/Engine/Plugins/Animation/PoseSearch/Source/Runtime/Private/PoseSearchInteractionIsland.cpp b/Engine/Plugins/Animation/PoseSearch/Source/Runtime/Private/PoseSearchInteractionIsland.cpp
--- a/Engine/Plugins/Animation/PoseSearch/Source/Runtime/Private/PoseSearchInteractionIsland.cpp
+++ b/Engine/Plugins/Animation/PoseSearch/Source/Runtime/Private/PoseSearchInteractionIsland.cpp
@@ -397,7 +397,11 @@ void FValidInteractionSearch::VLogContext(const FColor& Color) const
 		{
 			if (const UObject* AnimContext = AnimContexts[Index].Get())
 			{
-				Locations[Index] = GetContextLocation(AnimContext);
+				const AActor* ContextOwningActor = GetContextOwningActor(AnimContext);
+				check(ContextOwningActor);
+				// using ContextOwningActor location instead of skeletal mesh location since it's faster for UAF
+				//Locations[Index] = GetContextLocation(AnimContext);
+				Locations[Index] = ContextOwningActor->GetTransform().GetLocation();
 			}
 		}
 
diff --git a/Engine/Plugins/Animation/PoseSearch/Source/Runtime/Private/PoseSearchInteractionSubsystem.cpp b/Engine/Plugins/Animation/PoseSearch/Source/Runtime/Private/PoseSearchInteractionSubsystem.cpp
--- a/Engine/Plugins/Animation/PoseSearch/Source/Runtime/Private/PoseSearchInteractionSubsystem.cpp
+++ b/Engine/Plugins/Animation/PoseSearch/Source/Runtime/Private/PoseSearchInteractionSubsystem.cpp
@@ -37,7 +37,12 @@ struct FAnimContextInfo
 	{
 		check(InAnimContextAvailabilities.AnimContext && !InAnimContextAvailabilities.Availabilities.IsEmpty());
 		AnimContextAvailabilities = &InAnimContextAvailabilities;
-		Location = GetContextLocation(InAnimContextAvailabilities.AnimContext, false);
+		
+		// using ContextOwningActor location instead of skeletal mesh location since it's faster for UAF
+		//Location = GetContextLocation(InAnimContextAvailabilities.AnimContext, false);
+		const AActor* ContextOwningActor = GetContextOwningActor(InAnimContextAvailabilities.AnimContext);
+		check(ContextOwningActor);
+		Location = ContextOwningActor->GetTransform().GetLocation();
 
 		AvailabilitiesMaxBroadPhaseRadius = 0.f;
 		for (const FPoseSearchInteractionAvailabilityEx& Availability : AnimContextAvailabilities->Availabilities)
@@ -168,11 +173,42 @@ void UPoseSearchInteractionSubsystem::RegenerateAllIslands(float DeltaSeconds)
 
 	struct FInteractionSearchContextGroup
 	{
+		typedef TPair<const UObject*, int32> FAnimContextToTickPriorityPair;
+		typedef TArray<FAnimContextToTickPriorityPair, TInlineAllocator<16, TMemStackAllocator<>>> FAnimContextToTickPriority;
+		typedef TArray<int32, TInlineAllocator<16, TMemStackAllocator<>>> FSearchContextsIndices;
+
+	private:
+		typedef TMap<const UObject*, int32, TInlineSetAllocator<16, TMemStackSetAllocator<>>> FAnimContextToTickPriorityMap;
+
+		// all the AnimContexts in this group with their TickPriority
+		FAnimContextToTickPriorityMap AnimContextToTickPriorityMap;
+
+		// sorted pair contained in AnimContextToTickPriorityMap calcualted during Finalize
+		FAnimContextToTickPriority AnimContextToTickPriority;
+
+		// indexes to the searchcontexts assigned to this group
+		FSearchContextsIndices SearchContextsIndices;
+
+	public:
+		const FAnimContextToTickPriority& GetAnimContextToTickPriority() const
+		{
+			// did we call Finalize?
+			check(AnimContextToTickPriority.Num() == AnimContextToTickPriorityMap.Num());
+			return AnimContextToTickPriority;
+		}
+
+		const FSearchContextsIndices& GetSearchContextsIndices() const
+		{
+			// did we call Finalize?
+			check(AnimContextToTickPriority.Num() == AnimContextToTickPriorityMap.Num());
+			return SearchContextsIndices;
+		}
+
 		bool Contains(const FInteractionSearchContext& SearchContext) const
 		{
 			for (int32 AnimContextIndex = 0; AnimContextIndex < SearchContext.Num(); ++AnimContextIndex)
 			{
-				if (AnimContextToTickPriority.Find(SearchContext.GetAnimContext(AnimContextIndex)))
+				if (AnimContextToTickPriorityMap.Find(SearchContext.GetAnimContext(AnimContextIndex)))
 				{
 					return true;
 				}
@@ -186,13 +222,13 @@ void UPoseSearchInteractionSubsystem::RegenerateAllIslands(float DeltaSeconds)
 			{
 				if (const UObject* AnimContext = SearchContext.GetAnimContext(AnimContextIndex))
 				{
-					if (int32* TickPriority = AnimContextToTickPriority.Find(AnimContext))
+					if (int32* TickPriority = AnimContextToTickPriorityMap.Find(AnimContext))
 					{
 						*TickPriority = FMath::Max(*TickPriority, SearchContext.TickPriorities[AnimContextIndex]);
 					}
 					else
 					{
-						AnimContextToTickPriority.Add(AnimContext) = SearchContext.TickPriorities[AnimContextIndex];
+						AnimContextToTickPriorityMap.Add(AnimContext) = SearchContext.TickPriorities[AnimContextIndex];
 					}
 				}
 			}
@@ -202,15 +238,15 @@ void UPoseSearchInteractionSubsystem::RegenerateAllIslands(float DeltaSeconds)
 
 		void Merge(const FInteractionSearchContextGroup& SearchContextGroup)
 		{
-			for (const FAnimContextToTickPriorityPair& AnimContextToTickPriorityPair : SearchContextGroup.AnimContextToTickPriority)
+			for (const FAnimContextToTickPriorityPair& AnimContextToTickPriorityPair : SearchContextGroup.AnimContextToTickPriorityMap)
 			{
-				if (int32* TickPriority = AnimContextToTickPriority.Find(AnimContextToTickPriorityPair.Key))
+				if (int32* TickPriority = AnimContextToTickPriorityMap.Find(AnimContextToTickPriorityPair.Key))
 				{
 					*TickPriority = FMath::Max(*TickPriority, AnimContextToTickPriorityPair.Value);
 				}
 				else
 				{
-					AnimContextToTickPriority.Add(AnimContextToTickPriorityPair.Key) = AnimContextToTickPriorityPair.Value;
+					AnimContextToTickPriorityMap.Add(AnimContextToTickPriorityPair.Key) = AnimContextToTickPriorityPair.Value;
 				}
 			}
 
@@ -220,58 +256,120 @@ void UPoseSearchInteractionSubsystem::RegenerateAllIslands(float DeltaSeconds)
 			}
 		}
 
-		// all the AnimContexts in this group with their TickPriority
-		typedef TPair<const UObject*, int32> FAnimContextToTickPriorityPair;
-		typedef TMap<const UObject*, int32, TInlineSetAllocator<16, TMemStackSetAllocator<>>> FAnimContextToTickPriority;
-		FAnimContextToTickPriority AnimContextToTickPriority;
+		// sorting by TickPriority first and by AnimContext in case of same TickPriority
+		void Finalize(const FInteractionSearchContexts& SearchContexts)
+		{
+			SearchContextsIndices.Sort([&SearchContexts](int32 A, int32 B)
+				{
+					return SearchContexts[A].AnimContexts[0] < SearchContexts[B].AnimContexts[0];
+				});
 
-		// indexes to the searchcontexts assigned to this group
-		TArray<int32, TInlineAllocator<16, TMemStackAllocator<>>> SearchContextsIndices;
+			AnimContextToTickPriority.Reset();
+			AnimContextToTickPriority.Reserve(AnimContextToTickPriorityMap.Num());
+			for (const FInteractionSearchContextGroup::FAnimContextToTickPriorityPair& AnimContextToTickPriorityPair : AnimContextToTickPriorityMap)
+			{
+				AnimContextToTickPriority.Add(AnimContextToTickPriorityPair);
+			}
+			AnimContextToTickPriority.Sort(
+				[](const FInteractionSearchContextGroup::FAnimContextToTickPriorityPair& A, const FInteractionSearchContextGroup::FAnimContextToTickPriorityPair& B)
+				{
+					if (B.Value != A.Value)
+					{
+						return B.Value < A.Value;
+					}
+
+					return A.Key < B.Key;
+				});
+		}
 	};
 
-	// grouping SearchContexts AnimContext(s) in FInteractionSearchContextGroup(s). We'll create as many interaction islands as many groups
-	TArray<FInteractionSearchContextGroup, TInlineAllocator<PreallocatedSearchesNum, TMemStackAllocator<>>> SearchContextGroups;
-	for (int32 SearchContextIndex = 0; SearchContextIndex < SearchContexts.Num(); ++SearchContextIndex)
+	struct FInteractionSearchContextGroups
 	{
-		// evaluating where to place SearchContext..
-		const FInteractionSearchContext& SearchContext = SearchContexts[SearchContextIndex];
+	private:
+		// grouping SearchContexts AnimContext(s) in FInteractionSearchContextGroup(s). We'll create as many interaction islands as many groups
+		TArray<FInteractionSearchContextGroup, TInlineAllocator<PreallocatedSearchesNum, TMemStackAllocator<>>> Groups;
+		TArray<int32, TInlineAllocator<PreallocatedSearchesNum, TMemStackAllocator<>>> SortedGroupsIndexes;
+
+	public:
+		int32 Num() const
+		{
+			check(SortedGroupsIndexes.Num() == Groups.Num());
+			return SortedGroupsIndexes.Num();
+		}
+
+		const FInteractionSearchContextGroup& operator[](int32 Index) const
+		{
+			return Groups[SortedGroupsIndexes[Index]];
+		}
 
-		int32 MainSearchContextGroupIndex = INDEX_NONE;
-		for (int32 SearchContextGroupIndex = 0; SearchContextGroupIndex < SearchContextGroups.Num();)
+		FInteractionSearchContextGroups(const FInteractionSearchContexts& SearchContexts)
 		{
-			// ..if SearchContextGroups[SearchContextGroupIndex] contains ANY of the AnimContexts from SearchContext..
-			if (SearchContextGroups[SearchContextGroupIndex].Contains(SearchContext))
+			for (int32 SearchContextIndex = 0; SearchContextIndex < SearchContexts.Num(); ++SearchContextIndex)
 			{
-				if (MainSearchContextGroupIndex == INDEX_NONE)
+				// evaluating where to place SearchContext..
+				const FInteractionSearchContext& SearchContext = SearchContexts[SearchContextIndex];
+
+				int32 MainGroupIndex = INDEX_NONE;
+				for (int32 GroupIndex = 0; GroupIndex < Groups.Num();)
 				{
-					// ..we add SearchContext to SearchContextGroups[SearchContextGroupIndex] 
-					// and set MainSearchContextGroupIndex to SearchContextGroupIndex to know what is the group containing SearchContext, so..
-					MainSearchContextGroupIndex = SearchContextGroupIndex;
-					SearchContextGroups[MainSearchContextGroupIndex].Add(SearchContext, SearchContextIndex);
-					++SearchContextGroupIndex;
+					// ..if Groups[GroupIndex] contains ANY of the AnimContexts from SearchContext..
+					if (Groups[GroupIndex].Contains(SearchContext))
+					{
+						if (MainGroupIndex == INDEX_NONE)
+						{
+							// ..we add SearchContext to Groups[GroupIndex] 
+							// and set MainGroupIndex to GroupIndex to know what is the group containing SearchContext, so..
+							MainGroupIndex = GroupIndex;
+							Groups[MainGroupIndex].Add(SearchContext, SearchContextIndex);
+							++GroupIndex;
+						}
+						else
+						{
+							// ..in case SearchContext has already being inserted in MainGroupIndex group 
+							// we merge the newly found Groups[GroupIndex] to Groups[MainGroupIndex]
+							// (containing another of the the AnimContexts)
+							Groups[MainGroupIndex].Merge(Groups[GroupIndex]);
+							Groups.RemoveAt(GroupIndex);
+						}
+					}
+					else
+					{
+						++GroupIndex;
+					}
 				}
-				else
+				if (MainGroupIndex == INDEX_NONE)
 				{
-					// ..in case SearchContext has already being inserted in MainSearchContextGroupIndex group 
-					// we merge the newly found SearchContextGroups[SearchContextGroupIndex] to SearchContextGroups[MainSearchContextGroupIndex]
-					// (containing another of the the AnimContexts)
-					SearchContextGroups[MainSearchContextGroupIndex].Merge(SearchContextGroups[SearchContextGroupIndex]);
-					SearchContextGroups.RemoveAt(SearchContextGroupIndex);
+					Groups.AddDefaulted_GetRef().Add(SearchContext, SearchContextIndex);
 				}
 			}
-			else
+
+			// finalizing the Groups first to sort their internal data structures to be able to then sort the SortedGroupsIndexes
+			for (FInteractionSearchContextGroup& Group : Groups)
 			{
-				++SearchContextGroupIndex;
+				Group.Finalize(SearchContexts);
 			}
+
+			SortedGroupsIndexes.SetNumUninitialized(Groups.Num());
+			for (int32 GroupIndex = 0; GroupIndex < Groups.Num(); ++GroupIndex)
+			{
+				SortedGroupsIndexes[GroupIndex] = GroupIndex;
+			}
+
+			SortedGroupsIndexes.Sort([this, &SearchContexts](int32 A, int32 B)
+				{
+					const int32 SearchContextsIndexA = Groups[A].GetSearchContextsIndices()[0];
+					const int32 SearchContextsIndexB = Groups[B].GetSearchContextsIndices()[0];
+
+					const UObject* AnimContextA = SearchContexts[SearchContextsIndexA].GetAnimContext(0);
+					const UObject* AnimContextB = SearchContexts[SearchContextsIndexB].GetAnimContext(0);
+
+					return AnimContextA < AnimContextB;
+				});
 		}
-		if (MainSearchContextGroupIndex == INDEX_NONE)
-		{
-			SearchContextGroups.AddDefaulted_GetRef().Add(SearchContext, SearchContextIndex);
-		}
-	}
+	};
 
-	TArray<FInteractionSearchContextGroup::FAnimContextToTickPriorityPair, TInlineAllocator<16, TMemStackAllocator<>>> SortedByTickPriorityAnimContexts;
-	for (FInteractionSearchContextGroup& SearchContextGroup : SearchContextGroups)
+	FInteractionSearchContextGroups SearchContextGroups(SearchContexts);
+	for (int32 SearchContextGroupIndex = 0; SearchContextGroupIndex < SearchContextGroups.Num(); ++SearchContextGroupIndex)
 	{
 		// @todo: search for the most suitable island to reuse to avoid having to Uninitialize/RemoveTickDependencies and InjectToActor right away
 		FInteractionIsland& Island = GetAvailableIsland();
@@ -280,7 +378,9 @@ void UPoseSearchInteractionSubsystem::RegenerateAllIslands(float DeltaSeconds)
 		// initializing the island with its assigned SearchContexts
 		bool bAreTickDependenciesRequired = false;
 		check(Island.GetSearchContexts().IsEmpty());
-		for (int32 SearchContextsIndex : SearchContextGroup.SearchContextsIndices)
+		
+		const FInteractionSearchContextGroup& SearchContextGroup = SearchContextGroups[SearchContextGroupIndex];
+		for (int32 SearchContextsIndex : SearchContextGroup.GetSearchContextsIndices())
 		{
 			const FInteractionSearchContext& SearchContext = SearchContexts[SearchContextsIndex];
 			// if there're at least two AnimContext(s) potentially interacting with each other 
@@ -289,25 +389,11 @@ void UPoseSearchInteractionSubsystem::RegenerateAllIslands(float DeltaSeconds)
 			Island.AddSearchContext(SearchContext);
 		}
 
-		// sorting SearchContextGroup.AnimContextToTickPriority by TickPriority
-		// (using SortedByTickPriorityAnimContexts since SearchContextGroup.AnimContextToTickPriority it's a TMap)
-		SortedByTickPriorityAnimContexts.Reset();
-		SortedByTickPriorityAnimContexts.Reserve(SearchContextGroup.AnimContextToTickPriority.Num());
-		for (const FInteractionSearchContextGroup::FAnimContextToTickPriorityPair& AnimContextToTickPriorityPair : SearchContextGroup.AnimContextToTickPriority)
-		{
-			SortedByTickPriorityAnimContexts.Add(AnimContextToTickPriorityPair);
-		}
-		SortedByTickPriorityAnimContexts.Sort(
-			[](const FInteractionSearchContextGroup::FAnimContextToTickPriorityPair& A, const FInteractionSearchContextGroup::FAnimContextToTickPriorityPair& B)
-			{
-				return B.Value < A.Value;
-			});
-
 		// injecting tick dependencies between island AnimContext following their TickPriorities, 
 		// so the AnimContext with the highest TickPriority will be elected as "Main Actor", being evaluated, 
 		// and performing all the island searches, before any other Actor in the same island\
 		// (that'll end up using the cached search results in a multithread manner)
-		for (const FInteractionSearchContextGroup::FAnimContextToTickPriorityPair& SortedByTickPriorityAnimContext : SortedByTickPriorityAnimContexts)
+		for (const FInteractionSearchContextGroup::FAnimContextToTickPriorityPair& SortedByTickPriorityAnimContext : SearchContextGroup.GetAnimContextToTickPriority())
 		{
 			Island.InjectToActor(SortedByTickPriorityAnimContext.Key, bAreTickDependenciesRequired);
 		}
@@ -538,10 +624,6 @@ void UPoseSearchInteractionSubsystem::GenerateAnimContextInfosAndTagToDatabases(
 #if DO_CHECK
 	check(SubsystemWorld);
 	check(AnimContextInfos.IsEmpty() && TagToDatabases.IsEmpty());
-	check(Algo::IsSorted(AnimContextsAvailabilities, [](const FPoseSearchInteractionAnimContextAvailabilities& A, const FPoseSearchInteractionAnimContextAvailabilities& B)
-		{
-			return A.AnimContext < B.AnimContext;
-		}));
 #endif // DO_CHECK
 
 	for (const FPoseSearchInteractionAnimContextAvailabilities& AnimContextAvailabilities : AnimContextsAvailabilities)
@@ -635,12 +717,42 @@ void UPoseSearchInteractionSubsystem::GenerateSearchContexts(float DeltaSeconds,
 				Database == Other.Database;
 		}
 
+		bool operator<(const FRoledAnimContextInfo& Other) const
+		{
+			if (AnimContextInfo != Other.AnimContextInfo)
+			{
+				if (AnimContextInfo->AnimContextAvailabilities->AnimContext != Other.AnimContextInfo->AnimContextAvailabilities->AnimContext)
+				{
+					return AnimContextInfo->AnimContextAvailabilities->AnimContext < Other.AnimContextInfo->AnimContextAvailabilities->AnimContext;
+				}
+
+				return AnimContextInfo < Other.AnimContextInfo;
+			}
+
+			if (Availability != Other.Availability)
+			{
+				return Availability < Other.Availability;
+			}
+
+			if (Role != Other.Role)
+			{
+				return Role.FastLess(Other.Role);
+			}
+
+			if (PoseHistory != Other.PoseHistory)
+			{
+				return PoseHistory < Other.PoseHistory;
+			}
+
+			return Database < Other.Database;
+		}
+
 		// Availability that spawned this FRoledAnimContextInfo
 		const FPoseSearchInteractionAvailabilityEx* Availability = nullptr;
 		// AnimContextInfo containing all the information regarding the AnimContext that spawned this FRoledAnimContextInfo,
 		// including all the availabilities associated to the AnimContext as well as all the other AnimContext(s) it can potentially interact with
 		const FAnimContextInfo* AnimContextInfo = nullptr;
-		const FRole Role = DefaultRole;
+		FRole Role = DefaultRole;
 		const IPoseHistory* PoseHistory = nullptr;
 		const UPoseSearchDatabase* Database = nullptr;
 
@@ -648,34 +760,119 @@ void UPoseSearchInteractionSubsystem::GenerateSearchContexts(float DeltaSeconds,
 		float IncrementedBroadPhaseRadiusSquared = 0.f;
 	};
 	
-	struct FRoledAnimContextInfos : TArray<FRoledAnimContextInfo, TInlineAllocator<PreallocatedRolesNum, TMemStackAllocator<>>>
+	struct FRoledAnimContextInfos
+	{
+	private:
+		typedef TArray<FRoledAnimContextInfo, TInlineAllocator<PreallocatedRolesNum, TMemStackAllocator<>>> FInternalData;
+
+		FInternalData InternalData;
+		TArray<bool, TInlineAllocator<PreallocatedRolesNum, TMemStackAllocator<>>> PresentRoleIndexes;
+
+	public:
+		FInternalData::SizeType Num() const
+		{
+			return InternalData.Num();
+		}
+
+		const FRoledAnimContextInfo& operator[](int32 Index)
+		{
+			return InternalData[Index];
+		}
+
+		bool HasAllRoles() const
+		{
+			for (bool bPresentRoleIndex : PresentRoleIndexes)
+			{
+				if (!bPresentRoleIndex)
+				{
+					return false;
+				}
+			}
+			return true;
+		}
+
+		// sort RoledAnimContextInfos to generate deterministic SearchContext across multiple frames!
+		void SortAndRemoveDuplicates()
+		{
+			if (InternalData.Num() <= 1)
+			{
+				return;
+			}
+
+			InternalData.Sort();
+
+			int32 PrevIndex = 0;
+			for (int32 Index = 1; Index < InternalData.Num(); ++Index)
+			{
+				if (InternalData[PrevIndex] != InternalData[Index])
+				{
+					++PrevIndex;
+
+					if (PrevIndex != Index)
 	{
+						InternalData[PrevIndex] = MoveTemp(InternalData[Index]);
+					}
+				}
+			}
+
+			++PrevIndex;
+			InternalData.RemoveAt(PrevIndex, Num() - PrevIndex, EAllowShrinking::No);
+		}
+
 		void AddRoledAnimContextInfos(const FPoseSearchInteractionAvailabilityEx& Availability, const FAnimContextInfo& AnimContextInfo, const IPoseHistory* PoseHistory, const UPoseSearchDatabase& Database)
 		{
 			const UPoseSearchSchema* Schema = Database.Schema;
 			check(Schema);
+			const TArray<FPoseSearchRoledSkeleton>& RoledSkeletons = Schema->GetRoledSkeletons();
+
+			if (PresentRoleIndexes.IsEmpty())
+			{
+				check(InternalData.IsEmpty());
+				PresentRoleIndexes.SetNumZeroed(RoledSkeletons.Num());
+			}
+			else
+			{
+				check(InternalData.IsEmpty() || InternalData[0].Database == &Database);
+				check(Schema->GetRoledSkeletons().Num() == PresentRoleIndexes.Num());
+			}
+
 			if (Availability.RolesFilter.IsEmpty())
 			{
 				// adding ALL the possible roles from the database:
-				for (const FPoseSearchRoledSkeleton& RoledSkeleton : Schema->GetRoledSkeletons())
+				for (int32 RoledSkeletonIndex = 0; RoledSkeletonIndex < RoledSkeletons.Num(); ++RoledSkeletonIndex)
 				{
-					AddUnique(FRoledAnimContextInfo(Availability, AnimContextInfo, RoledSkeleton.Role, PoseHistory, Database));
+					InternalData.Emplace(Availability, AnimContextInfo, RoledSkeletons[RoledSkeletonIndex].Role, PoseHistory, Database);
+					PresentRoleIndexes[RoledSkeletonIndex] = true;
 				}
 			}
 			else
 			{
 				for (const FRole& Role : Availability.RolesFilter)
 				{
-					if (Schema->GetRoledSkeleton(Role))
+					int32 RoledSkeletonIndex = 0;
+					while (RoledSkeletonIndex < RoledSkeletons.Num())
+					{
+						if (RoledSkeletons[RoledSkeletonIndex].Role == Role)
 					{
-						AddUnique(FRoledAnimContextInfo(Availability, AnimContextInfo, Role, PoseHistory, Database));
+							InternalData.Emplace(Availability, AnimContextInfo, Role, PoseHistory, Database);
+							PresentRoleIndexes[RoledSkeletonIndex] = true;
+							break;
+						}
+						++RoledSkeletonIndex;
 					}
-					else
+#if !NO_LOGGING
+					if (RoledSkeletonIndex == RoledSkeletons.Num())
 					{
 						UE_LOG(LogPoseSearch, Warning, TEXT("UPoseSearchInteractionSubsystem::GenerateSearchContexts unsupported Role %s for Database %s"), *Role.ToString(), *Database.GetName());
 					}
+#endif // !NO_LOGGING
+				}
 				}
 			}
+
+		void Reset()
+		{
+			InternalData.Reset();
 		}
 	};
 
@@ -740,7 +937,38 @@ void UPoseSearchInteractionSubsystem::GenerateSearchContexts(float DeltaSeconds,
 	FTagToDatabases TagToDatabases;
 	GenerateAnimContextInfosAndTagToDatabases(AnimContextInfos, TagToDatabases);
 
-	TMap<const UPoseSearchDatabase*, FRoledAnimContextInfos, TInlineSetAllocator<PreallocatedSearchesNum, TMemStackSetAllocator<>>> DatabaseToRoledAnimContextInfos;
+	struct FDatabaseToRoledAnimContextInfos
+	{
+	private:
+		typedef TMap<const UPoseSearchDatabase*, FRoledAnimContextInfos, TInlineSetAllocator<PreallocatedSearchesNum, TMemStackSetAllocator<>>> FInternalMap;
+		FInternalMap InternalMap;
+
+	public:
+		typedef FInternalMap::TIterator TIterator;
+
+		void AddRoledAnimContextInfos(const FPoseSearchInteractionAvailabilityEx& Availability, const FAnimContextInfo& AnimContextInfo, const IPoseHistory* PoseHistory, const UPoseSearchDatabase* Database)
+		{
+			check(Database && Database->Schema);
+			FRoledAnimContextInfos& RoledAnimContextInfos = InternalMap.FindOrAdd(Database);
+			RoledAnimContextInfos.AddRoledAnimContextInfos(Availability, AnimContextInfo, PoseHistory, *Database);
+		}
+
+		TIterator CreateIterator()
+		{
+			return InternalMap.CreateIterator();
+		}
+
+		void Reset()
+		{
+			// preserving the FRoledAnimContextInfos allocations, since reenumerating the few databases we have around with empty InternalMap is quite fast
+			for (TPair<const UPoseSearchDatabase*, FRoledAnimContextInfos>& Pair : InternalMap)
+			{
+				Pair.Value.Reset();
+			}
+		}
+
+	} DatabaseToRoledAnimContextInfos;
+
 	// visiting ALL the AnimContexts in AnimContextInfos and relying on the cached information to refine potential interactions
 	FAnimContextInfoVisitor AnimContextInfoVisitor(AnimContextInfos,
 		// OnNewAnimContextFound: called when the FAnimContextInfoVisitor finds a new AnimContext that can be grouped in the current group of possibly interacting AnimContext(s)
@@ -757,8 +985,7 @@ void UPoseSearchInteractionSubsystem::GenerateSearchContexts(float DeltaSeconds,
 				if (const UPoseSearchDatabase* Database = Availability.Database.Get())
 				{
 					check(Database->Schema);
-					FRoledAnimContextInfos& RoledAnimContextInfos = DatabaseToRoledAnimContextInfos.FindOrAdd(Database);
-					RoledAnimContextInfos.AddRoledAnimContextInfos(Availability, AnimContextInfo, PoseHistory, *Database);
+					DatabaseToRoledAnimContextInfos.AddRoledAnimContextInfos(Availability, AnimContextInfo, PoseHistory, Database);
 				}
 				else if (Availability.IsTagValid())
 				{
@@ -768,9 +995,7 @@ void UPoseSearchInteractionSubsystem::GenerateSearchContexts(float DeltaSeconds,
 						check(!DatabasesPerTag->IsEmpty());
 						for (const UPoseSearchDatabase* DatabaseFromTag : *DatabasesPerTag)
 						{
-							check(DatabaseFromTag && DatabaseFromTag->Schema);
-							FRoledAnimContextInfos& RoledAnimContextInfos = DatabaseToRoledAnimContextInfos.FindOrAdd(DatabaseFromTag);
-							RoledAnimContextInfos.AddRoledAnimContextInfos(Availability, AnimContextInfo, PoseHistory, *DatabaseFromTag);
+							DatabaseToRoledAnimContextInfos.AddRoledAnimContextInfos(Availability, AnimContextInfo, PoseHistory, DatabaseFromTag);
 						}
 					}
 					else
@@ -809,117 +1034,117 @@ void UPoseSearchInteractionSubsystem::GenerateSearchContexts(float DeltaSeconds,
 			//
 			// CharB/RoleA - CharB/RoleB -> invalid because of same CharB
 
-			for (TPair<const UPoseSearchDatabase*, FRoledAnimContextInfos>& DatabaseToRoledAnimContextInfosPair : DatabaseToRoledAnimContextInfos)
+			for (FDatabaseToRoledAnimContextInfos::TIterator It = DatabaseToRoledAnimContextInfos.CreateIterator(); It; ++It)
 			{
-				const UPoseSearchDatabase* Database = DatabaseToRoledAnimContextInfosPair.Key;
-				check(Database->Schema);
-				const TArray<FPoseSearchRoledSkeleton>& RoledSkeletons = Database->Schema->GetRoledSkeletons();
-				const int32 CombinationCardinality = RoledSkeletons.Num();
-				FRoledAnimContextInfos& RoledAnimContextInfos = DatabaseToRoledAnimContextInfosPair.Value;
-
-				// sort RoledAnimContextInfos to generate deterministic SearchContext across multiple frames!
-				RoledAnimContextInfos.Sort([](const FRoledAnimContextInfo& RoledAnimContextInfoA, const FRoledAnimContextInfo& RoledAnimContextInfoB)
-					{
-						return RoledAnimContextInfoA.AnimContextInfo->AnimContextAvailabilities->AnimContext < RoledAnimContextInfoB.AnimContextInfo->AnimContextAvailabilities->AnimContext;
-					});
+				FRoledAnimContextInfos& RoledAnimContextInfos = It->Value;
+				if (RoledAnimContextInfos.HasAllRoles())
+				{
+					const UPoseSearchDatabase* Database = It->Key;
+					check(Database->Schema);
+					const TArray<FPoseSearchRoledSkeleton>& RoledSkeletons = Database->Schema->GetRoledSkeletons();
+					const int32 CombinationCardinality = RoledSkeletons.Num();
 
-				GenerateCombinations(RoledAnimContextInfos.Num(), CombinationCardinality,
-					// Combination is an array of indexes in RoledAnimContextInfos: 0 <= Combination[i] < RoledAnimContextInfos.Num()
-					[this, Database, &RoledSkeletons, &RoledAnimContextInfos, &SearchContexts](const TConstArrayView<int32> Combination)
-					{
-						// CombinationCardinality represents the number of roles as well as the number interacting AnimContext(s) (ultimately number of Characters involved in the interaction)
-						const int32 CombinationCardinality = Combination.Num();
+					// sort RoledAnimContextInfos to generate deterministic SearchContext across multiple frames!
+					RoledAnimContextInfos.SortAndRemoveDuplicates();
 
-						for (int32 CombinationIndex = 0; CombinationIndex < CombinationCardinality; ++CombinationIndex)
+					GenerateCombinations(RoledAnimContextInfos.Num(), CombinationCardinality,
+						// Combination is an array of indexes in RoledAnimContextInfos: 0 <= Combination[i] < RoledAnimContextInfos.Num()
+						[this, Database, &RoledSkeletons, &RoledAnimContextInfos, &SearchContexts](const TConstArrayView<int32> Combination)
 						{
-							const FRoledAnimContextInfo& RoledAnimContextInfo = RoledAnimContextInfos[Combination[CombinationIndex]];
+							// CombinationCardinality represents the number of roles as well as the number interacting AnimContext(s) (ultimately number of Characters involved in the interaction)
+							const int32 CombinationCardinality = Combination.Num();
 
-							// since CombinationCardinality is usually small 2-3, it's faster to check linearly for duplicated animcontexts rather than holding them in a set
-							const UObject* CombinationIndexAnimContext = RoledAnimContextInfo.AnimContextInfo->AnimContextAvailabilities->AnimContext;
-							for (int32 NextCombinationIndex = CombinationIndex + 1; NextCombinationIndex < CombinationCardinality; ++NextCombinationIndex)
+							for (int32 CombinationIndex = 0; CombinationIndex < CombinationCardinality; ++CombinationIndex)
 							{
-								const FRoledAnimContextInfo& NextRoledAnimContextInfo = RoledAnimContextInfos[Combination[NextCombinationIndex]];
-								const UObject* NextCombinationIndexAnimContext = NextRoledAnimContextInfo.AnimContextInfo->AnimContextAvailabilities->AnimContext;
+								const FRoledAnimContextInfo& RoledAnimContextInfo = RoledAnimContextInfos[Combination[CombinationIndex]];
 
-								if (CombinationIndexAnimContext == NextCombinationIndexAnimContext)
+								// since CombinationCardinality is usually small 2-3, it's faster to check linearly for duplicated animcontexts rather than holding them in a set
+								const UObject* CombinationIndexAnimContext = RoledAnimContextInfo.AnimContextInfo->AnimContextAvailabilities->AnimContext;
+								for (int32 NextCombinationIndex = CombinationIndex + 1; NextCombinationIndex < CombinationCardinality; ++NextCombinationIndex)
 								{
-									// we have a duplicate AnimContext. this combination is NOT valid
-									return false;
+									const FRoledAnimContextInfo& NextRoledAnimContextInfo = RoledAnimContextInfos[Combination[NextCombinationIndex]];
+									const UObject* NextCombinationIndexAnimContext = NextRoledAnimContextInfo.AnimContextInfo->AnimContextAvailabilities->AnimContext;
+
+									if (CombinationIndexAnimContext == NextCombinationIndexAnimContext)
+									{
+										// we have a duplicate AnimContext. this combination is NOT valid
+										return false;
+									}
 								}
 							}
-						}
 
-						FInteractionSearchContext SearchContext;
-						SearchContext.Database = Database;
+							FInteractionSearchContext SearchContext;
+							SearchContext.Database = Database;
 
-						// setting up a FRoledAnimContextInfo in RoledAnimContextInfos describing 
-						// this potential interaction properties about how to perform the search
-						for (int32 CombinationIndex = 0; CombinationIndex < CombinationCardinality; ++CombinationIndex)
-						{
-							const int32 RoledAnimContextIndex = Combination[CombinationIndex];
-							const FRoledAnimContextInfo& RoledAnimContextInfo = RoledAnimContextInfos[RoledAnimContextIndex];
+							// setting up a FRoledAnimContextInfo in RoledAnimContextInfos describing 
+							// this potential interaction properties about how to perform the search
+							for (int32 CombinationIndex = 0; CombinationIndex < CombinationCardinality; ++CombinationIndex)
+							{
+								const int32 RoledAnimContextIndex = Combination[CombinationIndex];
+								const FRoledAnimContextInfo& RoledAnimContextInfo = RoledAnimContextInfos[RoledAnimContextIndex];
 
-							check(RoledAnimContextInfo.PoseHistory && RoledAnimContextInfo.Availability);
-							SearchContext.Add(RoledAnimContextInfo.AnimContextInfo->AnimContextAvailabilities->AnimContext, RoledAnimContextInfo.PoseHistory, RoledAnimContextInfo.Role);
-							SearchContext.bDisableCollisions |= RoledAnimContextInfo.Availability->bDisableCollisions;
-							SearchContext.TickPriorities.Add(RoledAnimContextInfo.Availability->TickPriority);
-							
-							#if ENABLE_ANIM_DEBUG
-							SearchContext.DebugBroadPhaseRadius.Add(RoledAnimContextInfo.Availability->BroadPhaseRadius);
-							SearchContext.BroadPhaseRadiusIncrementOnInteraction.Add(RoledAnimContextInfo.Availability->BroadPhaseRadiusIncrementOnInteraction);
-							#endif // ENABLE_ANIM_DEBUG
-						}
+								check(RoledAnimContextInfo.PoseHistory && RoledAnimContextInfo.Availability);
+								SearchContext.Add(RoledAnimContextInfo.AnimContextInfo->AnimContextAvailabilities->AnimContext, RoledAnimContextInfo.PoseHistory, RoledAnimContextInfo.Role);
+								SearchContext.bDisableCollisions |= RoledAnimContextInfo.Availability->bDisableCollisions;
+								SearchContext.TickPriorities.Add(RoledAnimContextInfo.Availability->TickPriority);
 
-						// does SearchContext cover all the roles required by this interaction?
-						for (const FPoseSearchRoledSkeleton& RoledSkeleton : RoledSkeletons)
-						{
-							// CombinationCardinality is usually 2-3, so we can search the SearchContext.Roles array for duplicates without requiring a faster container like TSet
-							if (!SearchContext.GetRoles().Contains(RoledSkeleton.Role))
-							{
-								return false;
+#if ENABLE_ANIM_DEBUG
+								SearchContext.DebugBroadPhaseRadius.Add(RoledAnimContextInfo.Availability->BroadPhaseRadius);
+								SearchContext.BroadPhaseRadiusIncrementOnInteraction.Add(RoledAnimContextInfo.Availability->BroadPhaseRadiusIncrementOnInteraction);
+#endif // ENABLE_ANIM_DEBUG
 							}
-						}
 
-						// looking for a preexisting valid interaction resembling SearchContext
-						const uint32 SearchContextHash = SearchContext.GetHashForEquivalence();
-						SearchContext.bIsContinuingInteraction = ValidInteractionSearches.Contains(SearchContextHash);
+							// does SearchContext cover all the roles required by this interaction?
+							for (const FPoseSearchRoledSkeleton& RoledSkeleton : RoledSkeletons)
+							{
+								// CombinationCardinality is usually 2-3, so we can search the SearchContext.Roles array for duplicates without requiring a faster container like TSet
+								if (!SearchContext.GetRoles().Contains(RoledSkeleton.Role))
+								{
+									return false;
+								}
+							}
 
-						// checking if this combination is valid for the Database:
-						for (int32 CombinationIndex = 0; CombinationIndex < CombinationCardinality; ++CombinationIndex)
-						{
-							const int32 RoledAnimContextIndex = Combination[CombinationIndex];
-							const FRoledAnimContextInfo& RoledAnimContextInfo = RoledAnimContextInfos[RoledAnimContextIndex];
+							// looking for a preexisting valid interaction resembling SearchContext
+							const uint32 SearchContextHash = SearchContext.GetHashForEquivalence();
+							SearchContext.bIsContinuingInteraction = ValidInteractionSearches.Contains(SearchContextHash);
 
-							// checking the narrow phase! if RoledAnimContextInfo can interact with OtherRoledAnimContextInfo
-							for (int32 OtherCombinationIndex = CombinationIndex + 1; OtherCombinationIndex < CombinationCardinality; ++OtherCombinationIndex)
+							// checking if this combination is valid for the Database:
+							for (int32 CombinationIndex = 0; CombinationIndex < CombinationCardinality; ++CombinationIndex)
 							{
-								const int32 OtherRoledAnimContextIndex = Combination[OtherCombinationIndex];
-								const FRoledAnimContextInfo& OtherRoledAnimContextInfo = RoledAnimContextInfos[OtherRoledAnimContextIndex];
-
-								const FVector DeltaLocation = RoledAnimContextInfo.AnimContextInfo->Location - OtherRoledAnimContextInfo.AnimContextInfo->Location;
-								const float DistanceSquared = DeltaLocation.SquaredLength();
+								const int32 RoledAnimContextIndex = Combination[CombinationIndex];
+								const FRoledAnimContextInfo& RoledAnimContextInfo = RoledAnimContextInfos[RoledAnimContextIndex];
 
-								float MaxDistanceSquared;
-								if (SearchContext.bIsContinuingInteraction)
-								{
-									MaxDistanceSquared = FMath::Min(RoledAnimContextInfo.IncrementedBroadPhaseRadiusSquared, OtherRoledAnimContextInfo.IncrementedBroadPhaseRadiusSquared);
-								}
-								else
-								{
-									MaxDistanceSquared = FMath::Min(RoledAnimContextInfo.BroadPhaseRadiusSquared, OtherRoledAnimContextInfo.BroadPhaseRadiusSquared);
-								}
-			
-								// if any of the RoledAnimContextInfo cannot interact with any OtherRoledAnimContextInfo the inteaction cannot happen!
-								if (DistanceSquared > MaxDistanceSquared)
+								// checking the narrow phase! if RoledAnimContextInfo can interact with OtherRoledAnimContextInfo
+								for (int32 OtherCombinationIndex = CombinationIndex + 1; OtherCombinationIndex < CombinationCardinality; ++OtherCombinationIndex)
 								{
-									return false;
+									const int32 OtherRoledAnimContextIndex = Combination[OtherCombinationIndex];
+									const FRoledAnimContextInfo& OtherRoledAnimContextInfo = RoledAnimContextInfos[OtherRoledAnimContextIndex];
+
+									const FVector DeltaLocation = RoledAnimContextInfo.AnimContextInfo->Location - OtherRoledAnimContextInfo.AnimContextInfo->Location;
+									const float DistanceSquared = DeltaLocation.SquaredLength();
+
+									float MaxDistanceSquared;
+									if (SearchContext.bIsContinuingInteraction)
+									{
+										MaxDistanceSquared = FMath::Min(RoledAnimContextInfo.IncrementedBroadPhaseRadiusSquared, OtherRoledAnimContextInfo.IncrementedBroadPhaseRadiusSquared);
+									}
+									else
+									{
+										MaxDistanceSquared = FMath::Min(RoledAnimContextInfo.BroadPhaseRadiusSquared, OtherRoledAnimContextInfo.BroadPhaseRadiusSquared);
+									}
+
+									// if any of the RoledAnimContextInfo cannot interact with any OtherRoledAnimContextInfo the inteaction cannot happen!
+									if (DistanceSquared > MaxDistanceSquared)
+									{
+										return false;
+									}
 								}
 							}
-						}
 
-						SearchContexts.Emplace(SearchContext);
-						return true;
-					});
+							SearchContexts.Emplace(SearchContext);
+							return true;
+						});
+				}
 			}
 
 			// done using DatabaseToRoledAnimContextInfos. clearing up for the next group of AnimContext(s)
@@ -1012,7 +1237,7 @@ void UPoseSearchInteractionSubsystem::UpdateValidInteractionSearches()
 	}
 }
 
-bool UPoseSearchInteractionSubsystem::ConsolidateAnimContextsAvailabilities()
+bool UPoseSearchInteractionSubsystem::ShouldRegenerateAllIslands()
 {
 	using namespace UE::PoseSearch;
 
@@ -1044,38 +1269,6 @@ bool UPoseSearchInteractionSubsystem::ConsolidateAnimContextsAvailabilities()
 		return false;
 	}
 	
-	// consolidating AnimContextsAvailabilities sharing the same AnimInstance
-	TArrayView<FPoseSearchInteractionAnimContextAvailabilities> AnimContextsAvailabilities = MakeArrayView(AnimContextsAvailabilitiesBuffer.GetData(), AnimContextsAvailabilitiesNum);
-	AnimContextsAvailabilities.Sort([](const FPoseSearchInteractionAnimContextAvailabilities& AnimContextAvailabilitiesA, const FPoseSearchInteractionAnimContextAvailabilities& AnimContextAvailabilitiesB)
-		{
-			return AnimContextAvailabilitiesA.AnimContext < AnimContextAvailabilitiesB.AnimContext;
-		});
-
-	int32 WriteIndex = 0;
-	for (int32 ReadIndex = 1; ReadIndex < AnimContextsAvailabilitiesNum; ++ReadIndex)
-	{
-		// Avoiding adding trivial duplicates. FPoseSearchInteractionAvailabilityEx could not be fully specified to understand if it's an actual duplicate in case the pose
-		// history is passed by name or the Availability.Database is null and supposed to be resolved using other availabilities Database(s) with the same Availability.Tag.
-		// The duplicated availabilities are excluded when creating the combinations of possible interactions during FAnimContextInfoVisitor when 
-		// FRoledAnimContextInfos.AddRoledAnimContextInfos calls FRoledAnimContextInfos.AddUnique
-		if (AnimContextsAvailabilities[WriteIndex].AnimContext == AnimContextsAvailabilities[ReadIndex].AnimContext)
-		{
-			for (FPoseSearchInteractionAvailabilityEx& AvailabilityEx : AnimContextsAvailabilities[ReadIndex].Availabilities)
-			{
-				AnimContextsAvailabilities[WriteIndex].Availabilities.AddUnique(AvailabilityEx);
-			}
-		}
-		else
-		{
-			++WriteIndex;
-			if (WriteIndex != ReadIndex)
-			{
-				AnimContextsAvailabilities[WriteIndex] = AnimContextsAvailabilities[ReadIndex];
-			}
-		}
-	}
-		
-	AnimContextsAvailabilitiesNum = WriteIndex + 1;	
 	return true;
 }
 
@@ -1092,7 +1285,7 @@ void UPoseSearchInteractionSubsystem::Tick(float DeltaSeconds)
 	FMemMark Mark(FMemStack::Get());
 	UpdateValidInteractionSearches();
 
-	if (ConsolidateAnimContextsAvailabilities())
+	if (ShouldRegenerateAllIslands())
 	{
 		RegenerateAllIslands(DeltaSeconds);
 
diff --git a/Engine/Plugins/Animation/PoseSearch/Source/Runtime/Public/PoseSearch/PoseSearchInteractionSubsystem.h b/Engine/Plugins/Animation/PoseSearch/Source/Runtime/Public/PoseSearch/PoseSearchInteractionSubsystem.h
--- a/Engine/Plugins/Animation/PoseSearch/Source/Runtime/Public/PoseSearch/PoseSearchInteractionSubsystem.h
+++ b/Engine/Plugins/Animation/PoseSearch/Source/Runtime/Public/PoseSearch/PoseSearchInteractionSubsystem.h
@@ -101,7 +101,7 @@ private:
 	void DestroyIsland(int32 Index);
 	void DestroyAllIslands();
 	void RegenerateAllIslands(float DeltaSeconds);
-	bool ConsolidateAnimContextsAvailabilities();
+	bool ShouldRegenerateAllIslands();
 
 #if DO_CHECK
 	bool ValidateAllIslands() const;
