PoseSearch - temporarily disabling UPoseSearchInteractionSubsystem::OnBlueprintPreCompile logic, since not called from game thread

diff --git a/Engine/Plugins/Animation/PoseSearch/Source/Runtime/Private/PoseSearchInteractionSubsystem.cpp b/Engine/Plugins/Animation/PoseSearch/Source/Runtime/Private/PoseSearchInteractionSubsystem.cpp
--- a/Engine/Plugins/Animation/PoseSearch/Source/Runtime/Private/PoseSearchInteractionSubsystem.cpp
+++ b/Engine/Plugins/Animation/PoseSearch/Source/Runtime/Private/PoseSearchInteractionSubsystem.cpp
@@ -557,8 +557,10 @@ UPoseSearchInteractionSubsystem* UPoseSearchInteractionSubsystem::GetSubsystem_A
 // run PIE, pause (FInteractionIsland are still caching the SearchContexts calculated with the previous pose history nodes), compile BPs, unpause, crash!
 void UPoseSearchInteractionSubsystem::OnBlueprintPreCompile(UBlueprint* BlueprintToCompile)
 {
-	check(IsInGameThread());
-	DestroyAllIslands();
+	// @todo: reenable after fixing OnBlueprintPreCompile being called from worker thread
+
+	//check(IsInGameThread());
+	//DestroyAllIslands();
 }
 #endif // WITH_EDITOR
 
