PoseSearch - UPoseSearchInteractionSubsystem::GenerateAnimContextInfosAndTagToDatabases optimization

diff --git a/Engine/Plugins/Animation/PoseSearch/Source/Runtime/Private/PoseSearchInteractionSubsystem.cpp b/Engine/Plugins/Animation/PoseSearch/Source/Runtime/Private/PoseSearchInteractionSubsystem.cpp
--- a/Engine/Plugins/Animation/PoseSearch/Source/Runtime/Private/PoseSearchInteractionSubsystem.cpp
+++ b/Engine/Plugins/Animation/PoseSearch/Source/Runtime/Private/PoseSearchInteractionSubsystem.cpp
@@ -20,7 +20,11 @@
 
 namespace UE::PoseSearch
 {
+static bool GVarPoseSearchInteractionSortAnimContextInfosByUniqueCommonRoleFilter = true;
 #if !NO_CVARS
+static FAutoConsoleVariableRef CVarPoseSearchInteractionSortAnimContextInfosByUniqueCommonRoleFilter(TEXT("a.PoseSearchInteraction.SortAnimContextInfosByUniqueCommonRoleFilter"),
+	GVarPoseSearchInteractionSortAnimContextInfosByUniqueCommonRoleFilter, TEXT("Sort AnimContextInfos By UniqueCommonRoleFilter (it improves broadphase performance when actors publish availabilities with a single role filter)"));
+
 static bool GVarPoseSearchInteractionEnabled = true;
 static FAutoConsoleVariableRef CVarPoseSearchInteractionEnabled(TEXT("a.PoseSearchInteraction.Enabled"), GVarPoseSearchInteractionEnabled, TEXT("Enable/Disable Pose Search Interaction"));
 
@@ -52,12 +56,73 @@ struct FAnimContextInfo
 		}
 	}
 
+	// returns the common FRole used as RolesFilter if it's only one and consistent across all the AnimContextAvailabilities->Availabilities
+	// used to categorize FAnimContextInfo(s) in different buckets to speed up the broed phase calculation
+	const FRole* GetUniqueCommonRoleFilter() const
+	{
+		const FRole* UniqueCommonRoleFilter = nullptr;
+		for (const FPoseSearchInteractionAvailabilityEx& Availability : AnimContextAvailabilities->Availabilities)
+		{
+			if (Availability.RolesFilter.Num() == 1)
+			{
+				if (!UniqueCommonRoleFilter)
+				{
+					UniqueCommonRoleFilter = &Availability.RolesFilter[0];
+				}
+				else if (*UniqueCommonRoleFilter != Availability.RolesFilter[0])
+				{
+					UniqueCommonRoleFilter = nullptr;
+					break;
+				}
+			}
+			else
+			{
+				UniqueCommonRoleFilter = nullptr;
+				break;
+			}
+		}
+
+		return UniqueCommonRoleFilter;
+	}
+
+	const FVector& GetLocation() const
+	{
+		return Location;
+	}
+
+	float GetAvailabilitiesMaxBroadPhaseRadius() const
+	{
+		return AvailabilitiesMaxBroadPhaseRadius;
+	}
+	
+	void AddNearbyAnimContextInfo(const FAnimContextInfo& NearbyAnimContextInfo)
+	{
+		NearbyAnimContextInfos.Emplace(&NearbyAnimContextInfo);
+	}
+
+	const UObject* GetAnimContext() const
+	{
+		return AnimContextAvailabilities->AnimContext;
+	}
+
+	const TArray<FPoseSearchInteractionAvailabilityEx>& GetAvailabilities() const
+	{
+		return AnimContextAvailabilities->Availabilities;
+	}
+
+	typedef TArray<const FAnimContextInfo*, TInlineAllocator<8, TMemStackAllocator<>>> FNearbyAnimContextInfos;
+	const FNearbyAnimContextInfos& GetNearbyAnimContextInfos() const
+	{
+		return NearbyAnimContextInfos;
+	}
+
+private:
 	const FPoseSearchInteractionAnimContextAvailabilities* AnimContextAvailabilities = nullptr;
 
 	// cached AnimContext location
 	FVector Location = FVector::ZeroVector;
 	float AvailabilitiesMaxBroadPhaseRadius = 0.f;
-	TArray<const FAnimContextInfo*, TInlineAllocator<8, TMemStackAllocator<>>> NearbyAnimContextInfos;
+	FNearbyAnimContextInfos NearbyAnimContextInfos;
 };
 struct FAnimContextInfos : TArray<FAnimContextInfo, TMemStackAllocator<>> {};
 
@@ -652,34 +717,115 @@ void UPoseSearchInteractionSubsystem::GenerateAnimContextInfosAndTagToDatabases(
 		}
 	}
 
-	const int32 NumAnimContextInfos = AnimContextsAvailabilities.Num();
-	AnimContextInfos.SetNum(NumAnimContextInfos);
-	for (int32 AnimContextInfoIndex = 0; AnimContextInfoIndex < NumAnimContextInfos; ++AnimContextInfoIndex)
-	{
-		AnimContextInfos[AnimContextInfoIndex].Init(AnimContextsAvailabilities[AnimContextInfoIndex]);
-	}
-
-	// solving the broad phase using the AnimContextInfos
-	for (int32 AnimContextInfoIndexA = 0; AnimContextInfoIndexA < NumAnimContextInfos; ++AnimContextInfoIndexA)
-	{
-		for (int32 AnimContextInfoIndexB = AnimContextInfoIndexA + 1; AnimContextInfoIndexB < NumAnimContextInfos; ++AnimContextInfoIndexB)
+	auto EvaluateForNearbyAnimContextInfo = [&AnimContextInfos](int32 AnimContextInfoIndexA, int32 AnimContextInfoIndexB)
 		{
 			// performs broad phase analysis checking if at least one of the Availabilities associated to AnimContext can interact with OtherAnimContextInfo.
-			// This is a more relaxed analysis than the one performed in FRoledAnimContextInfo::CanInteractWith
-		
-			// @todo: filter OtherAnimContextInfo if dit doesn't have the "right role"? Maybe add a map of CanInteractWithRoles in this FAnimContextInfo
-		
-			const FVector DeltaLocation = AnimContextInfos[AnimContextInfoIndexA].Location - AnimContextInfos[AnimContextInfoIndexB].Location;
+			const FVector DeltaLocation = AnimContextInfos[AnimContextInfoIndexA].GetLocation() - AnimContextInfos[AnimContextInfoIndexB].GetLocation();
 			const float DistanceSquared = DeltaLocation.SquaredLength();
-			const float MaxDistance = FMath::Min(AnimContextInfos[AnimContextInfoIndexA].AvailabilitiesMaxBroadPhaseRadius, AnimContextInfos[AnimContextInfoIndexB].AvailabilitiesMaxBroadPhaseRadius);
+			const float MaxDistance = FMath::Min(AnimContextInfos[AnimContextInfoIndexA].GetAvailabilitiesMaxBroadPhaseRadius(), AnimContextInfos[AnimContextInfoIndexB].GetAvailabilitiesMaxBroadPhaseRadius());
 			const float MaxDistanceSquared = MaxDistance * MaxDistance;
 			if (DistanceSquared <= MaxDistanceSquared)
 			{
 				// the AnimContext of AnimContextInfos[AnimContextInfoIndexA] can potentially interact with the one from AnimContextInfos[AnimContextInfoIndexB]:
 				// linking AnimContextInfos[AnimContextInfoIndexA] to AnimContextInfos[AnimContextInfoIndexA] and vice versa to keep track of this when evaluating the broad phase.
 				// Since AnimContextInfos does't reallocate anymore, it's safe to store pointers to internal elements of the array!
-				AnimContextInfos[AnimContextInfoIndexA].NearbyAnimContextInfos.Emplace(&AnimContextInfos[AnimContextInfoIndexB]);
-				AnimContextInfos[AnimContextInfoIndexB].NearbyAnimContextInfos.Emplace(&AnimContextInfos[AnimContextInfoIndexA]);
+				AnimContextInfos[AnimContextInfoIndexA].AddNearbyAnimContextInfo(AnimContextInfos[AnimContextInfoIndexB]);
+				AnimContextInfos[AnimContextInfoIndexB].AddNearbyAnimContextInfo(AnimContextInfos[AnimContextInfoIndexA]);
+			}
+		};
+
+	if (GVarPoseSearchInteractionSortAnimContextInfosByUniqueCommonRoleFilter)
+	{
+		typedef TMap<const FRole, int32, TInlineSetAllocator<16, TMemStackSetAllocator<>>> FCommonRoleFilterToIndex;
+		FCommonRoleFilterToIndex CommonRoleFilterToIndex;
+
+		typedef TArray<TArray<int32, TMemStackAllocator<>>, TInlineAllocator<16, TMemStackAllocator<>>> FRoledAnimContextInfoIndexes;
+		FRoledAnimContextInfoIndexes RoledAnimContextInfoIndexes;
+
+		// adding RoledAnimContextInfoIndexes[0] containing all the indexes for those AnimContextInfos returning nullptr from GetUniqueCommonRoleFilter
+		RoledAnimContextInfoIndexes.AddDefaulted();
+
+		const int32 NumAnimContextInfos = AnimContextsAvailabilities.Num();
+		AnimContextInfos.SetNum(NumAnimContextInfos);
+		for (int32 AnimContextInfoIndex = 0; AnimContextInfoIndex < NumAnimContextInfos; ++AnimContextInfoIndex)
+		{
+			AnimContextInfos[AnimContextInfoIndex].Init(AnimContextsAvailabilities[AnimContextInfoIndex]);
+			if (const FRole* UniqueCommonRoleFilter = AnimContextInfos[AnimContextInfoIndex].GetUniqueCommonRoleFilter())
+			{
+				if (int32* UniqueCommonRoleFilterIndex = CommonRoleFilterToIndex.Find(*UniqueCommonRoleFilter))
+				{
+					RoledAnimContextInfoIndexes[*UniqueCommonRoleFilterIndex].Add(AnimContextInfoIndex);
+				}
+				else
+				{
+					CommonRoleFilterToIndex.Add(*UniqueCommonRoleFilter) = RoledAnimContextInfoIndexes.Num();
+					RoledAnimContextInfoIndexes.AddDefaulted_GetRef().Add(AnimContextInfoIndex);
+				}
+			}
+			else
+			{
+				RoledAnimContextInfoIndexes[0].Add(AnimContextInfoIndex);
+			}
+		}
+
+		// RoledAnimContextInfoIndexes[0] is a special case containing all those AnimContextInfos with more than just one unique common filter 
+		// (GetUniqueCommonRoleFilter returned null), so we need to check all the possible pairs between those AnimContextInfos.
+		// we're calling it first, so we can then sort RoledAnimContextInfoIndexes (losing reference to who's the first item) to speed up execution
+		for (int32 IndexA = 0; IndexA < RoledAnimContextInfoIndexes[0].Num(); ++IndexA)
+		{
+			const int32 AnimContextInfoIndexA = RoledAnimContextInfoIndexes[0][IndexA];
+			for (int32 IndexB = IndexA + 1; IndexB < RoledAnimContextInfoIndexes[0].Num(); ++IndexB)
+			{
+				const int32 AnimContextInfoIndexB = RoledAnimContextInfoIndexes[0][IndexB];
+				EvaluateForNearbyAnimContextInfo(AnimContextInfoIndexA, AnimContextInfoIndexB);
+			}
+		}
+
+		// checking all the Ith RoledAnimContextInfoIndexes against each other to find NearbyAnimContextInfo
+		if (RoledAnimContextInfoIndexes.Num() > 1)
+		{
+			// sorting RoledAnimContextInfoIndexes by the internal arrays cardinality to minimize the number of calls 
+			// of the inner 'for(s)' of the following nested 'for' statements
+			RoledAnimContextInfoIndexes.Sort([](const TArray<int32, TMemStackAllocator<>>& A, const TArray<int32, TMemStackAllocator<>>& B)
+				{
+					return A.Num() < B.Num();
+				});
+
+			// obviously CommonRoleFilterToIndex is now invalid. Resetting it here just for clarity
+			CommonRoleFilterToIndex.Reset();
+
+			for (int32 IndexA = 0; IndexA < RoledAnimContextInfoIndexes.Num() - 1; ++IndexA)
+			{
+				for (int32 AnimContextInfoIndexA : RoledAnimContextInfoIndexes[IndexA])
+				{
+					for (int32 IndexB = IndexA + 1; IndexB < RoledAnimContextInfoIndexes.Num(); ++IndexB)
+					{
+						for (int32 AnimContextInfoIndexB : RoledAnimContextInfoIndexes[IndexB])
+						{
+							EvaluateForNearbyAnimContextInfo(AnimContextInfoIndexA, AnimContextInfoIndexB);
+						}
+					}
+				}
+			}
+		}
+
+		
+	}
+	else
+	{
+		const int32 NumAnimContextInfos = AnimContextsAvailabilities.Num();
+		AnimContextInfos.SetNum(NumAnimContextInfos);
+		for (int32 AnimContextInfoIndex = 0; AnimContextInfoIndex < NumAnimContextInfos; ++AnimContextInfoIndex)
+		{
+			AnimContextInfos[AnimContextInfoIndex].Init(AnimContextsAvailabilities[AnimContextInfoIndex]);
+		}
+
+		// solving the broad phase using the AnimContextInfos
+		for (int32 AnimContextInfoIndexA = 0; AnimContextInfoIndexA < NumAnimContextInfos - 1; ++AnimContextInfoIndexA)
+		{
+			for (int32 AnimContextInfoIndexB = AnimContextInfoIndexA + 1; AnimContextInfoIndexB < NumAnimContextInfos; ++AnimContextInfoIndexB)
+			{
+				EvaluateForNearbyAnimContextInfo(AnimContextInfoIndexA, AnimContextInfoIndexB);
 			}
 		}
 	}
@@ -723,9 +869,11 @@ void UPoseSearchInteractionSubsystem::GenerateSearchContexts(float DeltaSeconds,
 		{
 			if (AnimContextInfo != Other.AnimContextInfo)
 			{
-				if (AnimContextInfo->AnimContextAvailabilities->AnimContext != Other.AnimContextInfo->AnimContextAvailabilities->AnimContext)
+				const UObject* AnimContext = AnimContextInfo->GetAnimContext();
+				const UObject* OtherAnimContext = Other.AnimContextInfo->GetAnimContext();
+				if (AnimContext != OtherAnimContext)
 				{
-					return AnimContextInfo->AnimContextAvailabilities->AnimContext < Other.AnimContextInfo->AnimContextAvailabilities->AnimContext;
+					return AnimContext < OtherAnimContext;
 				}
 
 				return AnimContextInfo < Other.AnimContextInfo;
@@ -890,7 +1038,7 @@ void UPoseSearchInteractionSubsystem::GenerateSearchContexts(float DeltaSeconds,
 		{
 			for (const FAnimContextInfo& AnimContextInfo : AnimContextInfos)
 			{
-				if (AnimContextInfo.NearbyAnimContextInfos.IsEmpty())
+				if (AnimContextInfo.GetNearbyAnimContextInfos().IsEmpty())
 				{
 					check(!VisitedAnimContextInfos.Find(&AnimContextInfo));
 					// no need to add this context to the VisitedAnimContextInfos since it's isolated!
@@ -912,7 +1060,7 @@ void UPoseSearchInteractionSubsystem::GenerateSearchContexts(float DeltaSeconds,
 	private:
 		void VisitRecursively(const FAnimContextInfo& AnimContextInfoToVisit, TFunctionRef<void(const FAnimContextInfo&)> OnNewAnimAnimContextInfoFound)
 		{
-			check(!AnimContextInfoToVisit.NearbyAnimContextInfos.IsEmpty());
+			check(!AnimContextInfoToVisit.GetNearbyAnimContextInfos().IsEmpty());
 
 			bool bIsAlreadyInSet;
 			VisitedAnimContextInfos.FindOrAdd(&AnimContextInfoToVisit, &bIsAlreadyInSet);
@@ -921,7 +1069,7 @@ void UPoseSearchInteractionSubsystem::GenerateSearchContexts(float DeltaSeconds,
 			{
 				OnNewAnimAnimContextInfoFound(AnimContextInfoToVisit);
 
-				for (const FAnimContextInfo* NearbyAnimContextInfo : AnimContextInfoToVisit.NearbyAnimContextInfos)
+				for (const FAnimContextInfo* NearbyAnimContextInfo : AnimContextInfoToVisit.GetNearbyAnimContextInfos())
 				{
 					check(NearbyAnimContextInfo);
 					VisitRecursively(*NearbyAnimContextInfo, OnNewAnimAnimContextInfoFound);
@@ -978,10 +1126,9 @@ void UPoseSearchInteractionSubsystem::GenerateSearchContexts(float DeltaSeconds,
 		{
 			// analyzing all the Availability(s) associated with this AnimContext and eventually generate the associated FRoledAnimContextInfos,
 			// inserted in a per database sorted data structure (DatabaseToRoledAnimContextInfos)
-			check(AnimContextInfo.AnimContextAvailabilities);
-			for (const FPoseSearchInteractionAvailabilityEx& Availability : AnimContextInfo.AnimContextAvailabilities->Availabilities)
+			for (const FPoseSearchInteractionAvailabilityEx& Availability : AnimContextInfo.GetAvailabilities())
 			{
-				const IPoseHistory* PoseHistory = Availability.GetPoseHistory(AnimContextInfo.AnimContextAvailabilities->AnimContext);
+				const IPoseHistory* PoseHistory = Availability.GetPoseHistory(AnimContextInfo.GetAnimContext());
 				check(PoseHistory);
 				
 				if (const UPoseSearchDatabase* Database = Availability.Database.Get())
@@ -1061,11 +1208,11 @@ void UPoseSearchInteractionSubsystem::GenerateSearchContexts(float DeltaSeconds,
 								const FRoledAnimContextInfo& RoledAnimContextInfo = RoledAnimContextInfos[Combination[CombinationIndex]];
 
 								// since CombinationCardinality is usually small 2-3, it's faster to check linearly for duplicated animcontexts rather than holding them in a set
-								const UObject* CombinationIndexAnimContext = RoledAnimContextInfo.AnimContextInfo->AnimContextAvailabilities->AnimContext;
+								const UObject* CombinationIndexAnimContext = RoledAnimContextInfo.AnimContextInfo->GetAnimContext();
 								for (int32 NextCombinationIndex = CombinationIndex + 1; NextCombinationIndex < CombinationCardinality; ++NextCombinationIndex)
 								{
 									const FRoledAnimContextInfo& NextRoledAnimContextInfo = RoledAnimContextInfos[Combination[NextCombinationIndex]];
-									const UObject* NextCombinationIndexAnimContext = NextRoledAnimContextInfo.AnimContextInfo->AnimContextAvailabilities->AnimContext;
+									const UObject* NextCombinationIndexAnimContext = NextRoledAnimContextInfo.AnimContextInfo->GetAnimContext();
 
 									if (CombinationIndexAnimContext == NextCombinationIndexAnimContext)
 									{
@@ -1086,7 +1233,7 @@ void UPoseSearchInteractionSubsystem::GenerateSearchContexts(float DeltaSeconds,
 								const FRoledAnimContextInfo& RoledAnimContextInfo = RoledAnimContextInfos[RoledAnimContextIndex];
 
 								check(RoledAnimContextInfo.PoseHistory && RoledAnimContextInfo.Availability);
-								SearchContext.Add(RoledAnimContextInfo.AnimContextInfo->AnimContextAvailabilities->AnimContext, RoledAnimContextInfo.PoseHistory, RoledAnimContextInfo.Role);
+								SearchContext.Add(RoledAnimContextInfo.AnimContextInfo->GetAnimContext(), RoledAnimContextInfo.PoseHistory, RoledAnimContextInfo.Role);
 								SearchContext.bDisableCollisions |= RoledAnimContextInfo.Availability->bDisableCollisions;
 								SearchContext.TickPriorities.Add(RoledAnimContextInfo.Availability->TickPriority);
 
@@ -1122,7 +1269,7 @@ void UPoseSearchInteractionSubsystem::GenerateSearchContexts(float DeltaSeconds,
 									const int32 OtherRoledAnimContextIndex = Combination[OtherCombinationIndex];
 									const FRoledAnimContextInfo& OtherRoledAnimContextInfo = RoledAnimContextInfos[OtherRoledAnimContextIndex];
 
-									const FVector DeltaLocation = RoledAnimContextInfo.AnimContextInfo->Location - OtherRoledAnimContextInfo.AnimContextInfo->Location;
+									const FVector DeltaLocation = RoledAnimContextInfo.AnimContextInfo->GetLocation() - OtherRoledAnimContextInfo.AnimContextInfo->GetLocation();
 									const float DistanceSquared = DeltaLocation.SquaredLength();
 
 									float MaxDistanceSquared;
@@ -1249,7 +1396,7 @@ bool UPoseSearchInteractionSubsystem::ShouldRegenerateAllIslands()
 	// @todo: add some setttings to initialize AnimContextsAvailabilitiesBuffer to a big enough value
 	if (AnimContextsAvailabilitiesIndexValue > AnimContextsAvailabilitiesNum)
 	{
-		UE_LOG(LogPoseSearch, Warning, TEXT("UPoseSearchInteractionSubsystem::ConsolidateAnimContextsAvailabilities not enough space to add more availabilities locklessly. It'll be adjusted automatically, but some availability requests has been lost this frame [capacity %d / requests %d]"), AnimContextsAvailabilitiesNum, AnimContextsAvailabilitiesIndexValue);
+		UE_LOG(LogPoseSearch, Warning, TEXT("UPoseSearchInteractionSubsystem::ShouldRegenerateAllIslands not enough space to add more availabilities locklessly. It'll be adjusted automatically, but some availability requests has been lost this frame [capacity %d / requests %d]"), AnimContextsAvailabilitiesNum, AnimContextsAvailabilitiesIndexValue);
 		AnimContextsAvailabilitiesBuffer.SetNum(AnimContextsAvailabilitiesIndexValue);
 	}
 	else
