PoseSearch - removing weak references in MM interactions to improve performance

#FailedValidations Preflight

diff --git a/Engine/Plugins/Animation/PoseSearch/Source/Runtime/Private/PoseSearchInteractionIsland.cpp b/Engine/Plugins/Animation/PoseSearch/Source/Runtime/Private/PoseSearchInteractionIsland.cpp
--- a/Engine/Plugins/Animation/PoseSearch/Source/Runtime/Private/PoseSearchInteractionIsland.cpp
+++ b/Engine/Plugins/Animation/PoseSearch/Source/Runtime/Private/PoseSearchInteractionIsland.cpp
@@ -130,14 +130,22 @@ static FInteractionSearchResult InitSearchResult(const FSearchResult& SearchResu
 
 	for (int32 AnimContextIndex = 0; AnimContextIndex < AnimContextsNum; ++AnimContextIndex)
 	{
-		const UObject* AnimContext = SearchContext.GetAnimContext(AnimContextIndex);
-		check(AnimContext);
-		const USkeleton* Skeleton = GetContextSkeleton(AnimContext, false);
-		const IPoseHistory* PoseHistory = SearchContext.GetPoseHistory(AnimContextIndex);
-		check(PoseHistory && Skeleton);
-		
-		PoseHistory->GetTransformAtTime(0.f, InteractionSearchResult.ActorRootTransforms[AnimContextIndex], Skeleton, ComponentSpaceIndexType, WorldSpaceIndexType);
-		PoseHistory->GetTransformAtTime(0.f, InteractionSearchResult.ActorRootBoneTransforms[AnimContextIndex], Skeleton, RootBoneIndexType, ComponentSpaceIndexType);
+		if (const IPoseHistory* PoseHistory = SearchContext.GetPoseHistory(AnimContextIndex))
+		{
+			const UObject* AnimContext = SearchContext.GetAnimContext(AnimContextIndex);
+			check(AnimContext);
+			const USkeleton* Skeleton = GetContextSkeleton(AnimContext, false);
+			check(Skeleton);
+
+			PoseHistory->GetTransformAtTime(0.f, InteractionSearchResult.ActorRootTransforms[AnimContextIndex], Skeleton, ComponentSpaceIndexType, WorldSpaceIndexType);
+			PoseHistory->GetTransformAtTime(0.f, InteractionSearchResult.ActorRootBoneTransforms[AnimContextIndex], Skeleton, RootBoneIndexType, ComponentSpaceIndexType);
+		}
+		else
+		{
+			// a pose history got disposed, this result is invalid
+			InteractionSearchResult = FInteractionSearchResult();
+			break;
+		}
 	}
 
 	return InteractionSearchResult;
@@ -293,40 +301,169 @@ static void AddPrerequisite(FTickFunction& TickFunction, UObject* TargetObject,
 	}
 }
 
-// FValidInteractionSearch
+// FInteractionSearchContextBase
 ///////////////////////////////////////////////////////////
-void FValidInteractionSearch::Init(const FInteractionSearchContext& Other)
+const UE::PoseSearch::IPoseHistory* FInteractionSearchContextBase::GetPoseHistory(int32 Index) const
+{
+	check(PoseHistories.IsValidIndex(Index));
+	if (TSharedPtr<const IPoseHistory, ESPMode::ThreadSafe> Pin = PoseHistories[Index].Pin())
+	{
+		return Pin.Get();
+	}
+	return nullptr;
+}
+
+uint32 FInteractionSearchContextBase::GetHashForEquivalence() const
+{
+	uint32 Hash = GetTypeHash(Database);
+
+	// skipping bDisableCollisions and PoseHistories for hashing
+	for (int32 Index = 0; Index < AnimContexts.Num(); ++Index)
+	{
+		Hash = HashCombineFast(Hash, GetTypeHash(AnimContexts[Index]));
+	}
+
+	for (int32 Index = 0; Index < Roles.Num(); ++Index)
+	{
+		Hash = HashCombineFast(Hash, GetTypeHash(Roles[Index]));
+	}
+
+	return Hash;
+}
+
+bool FInteractionSearchContextBase::IsEquivalent(const FInteractionSearchContextBase& Other) const
+{
+	// skipping bDisableCollisions and PoseHistories for equality
+	return Database == Other.Database &&
+		AnimContexts == Other.AnimContexts &&
+		Roles == Other.Roles;
+}
+
+void FInteractionSearchContextBase::AddReferencedObjects(FReferenceCollector& Collector)
+{
+	Collector.AddReferencedObject(Database);
+
+	for (TObjectPtr<const UObject>& AnimContext : AnimContexts)
+	{
+		Collector.AddReferencedObject(AnimContext);
+	}
+}
+
+void FInteractionSearchContextBase::Add(const UObject* AnimContext, const UE::PoseSearch::IPoseHistory* PoseHistory, const UE::PoseSearch::FRole Role, bool bWantsDisableCollisions)
+{
+	check(AnimContext && PoseHistory);
+
+	// AnimContexts must be sorted to have deterministic searches across multiple frames
+	// as well as for faster IsEquivalent implementation
+	check(AnimContexts.IsEmpty() || GetAnimContext(AnimContexts.Num() - 1) < AnimContext);
+
+	AnimContexts.Add(AnimContext);
+	PoseHistories.Add(PoseHistory->AsWeak());
+	Roles.Add(Role);
+
+	bDisableCollisions |= bWantsDisableCollisions;
+}
+
+#if DO_CHECK
+bool FInteractionSearchContextBase::CheckForConsistency() const
 {
-	Database = Other.Database;
-	bDisableCollisions = Other.bDisableCollisions;
-		
-	AnimContexts.SetNum(Other.AnimContexts.Num());
-	PoseHistories.SetNum(Other.PoseHistories.Num());
-	Roles = Other.Roles;
-
-	for (int32 Index = 0; Index < Other.AnimContexts.Num(); ++Index)
+	if (Database == nullptr)
+	{
+		return false;
+	}
+
+	const int32 Num = AnimContexts.Num();
+	if (Num < 1)
+	{
+		return false;
+	}
+
+	if (Num != PoseHistories.Num())
+	{
+		return false;
+	}
+
+	if (Num != Roles.Num())
+	{
+		return false;
+	}
+
+	for (int32 IndexA = 0; IndexA < Num; ++IndexA)
+	{
+		if (AnimContexts[IndexA] == nullptr)
+		{
+			return false;
+		}
+
+		for (int32 IndexB = IndexA + 1; IndexB < Num; ++IndexB)
+		{
+			if (AnimContexts[IndexA] == AnimContexts[IndexB])
+			{
+				return false;
+			}
+		}
+	}
+
+	for (int32 IndexA = 1; IndexA < Num; ++IndexA)
+	{
+		// AnimContexts must be sorted to have deterministic searches across multiple frames
+		if (AnimContexts[IndexA - 1] >= AnimContexts[IndexA])
+		{
+			return false;
+		}
+	}
+
+	for (int32 IndexA = 0; IndexA < Num; ++IndexA)
 	{
-		AnimContexts[Index] = Other.AnimContexts[Index];
+		for (int32 IndexB = IndexA + 1; IndexB < Num; ++IndexB)
+		{
+			if (Roles[IndexA] == Roles[IndexB])
+			{
+				return false;
+			}
+		}
 	}
 
-	for (int32 Index = 0; Index < Other.PoseHistories.Num(); ++Index)
+	for (int32 IndexA = 0; IndexA < Num; ++IndexA)
 	{
-		PoseHistories[Index] = Other.PoseHistories[Index]->AsWeak();
+		if (!PoseHistories[IndexA].IsValid())
+		{
+			return false;
+		}
+
+		for (int32 IndexB = IndexA + 1; IndexB < Num; ++IndexB)
+		{
+			if (PoseHistories[IndexA] == PoseHistories[IndexB])
+			{
+				return false;
+			}
+		}
 	}
+
+	return true;
+}
+
+#endif // DO_CHECK
+
+// FValidInteractionSearch
+///////////////////////////////////////////////////////////
+void FValidInteractionSearch::Init(const FInteractionSearchContext& Other)
+{
+	static_cast<FInteractionSearchContextBase&>(*this) = static_cast<const FInteractionSearchContextBase&>(Other);
 }
 
 void FValidInteractionSearch::DoDisableCollisions()
 {
 	check(DisabledCollisions.IsEmpty());
-	if (bDisableCollisions)
+	if (IsDisableCollisions())
 	{
 		FMemMark Mark(FMemStack::Get());
 		TArray<AActor*, TInlineAllocator<PreallocatedRolesNum, TMemStackAllocator<>>> Actors;
 		TArray<UPrimitiveComponent*, TInlineAllocator<PreallocatedRolesNum, TMemStackAllocator<>>> PrimitiveComponents;
 
-		for (int32 AnimContextIndex = 0; AnimContextIndex < AnimContexts.Num(); ++AnimContextIndex)
+		for (int32 AnimContextIndex = 0; AnimContextIndex < Num(); ++AnimContextIndex)
 		{
-			if (const UObject* AnimContext = AnimContexts[AnimContextIndex].Get())
+			if (const UObject* AnimContext = GetAnimContext(AnimContextIndex))
 			{
 				AActor* Actor = const_cast<AActor*>(GetContextOwningActor(AnimContext, false));
 				check(Actor);
@@ -368,9 +505,9 @@ void FValidInteractionSearch::UndoDisableCollisions()
 {
 	for (const FDisabledCollisions::ElementType& DisabledCollision : DisabledCollisions)
 	{
-		if (AActor* ActorA = DisabledCollision.Key.Get())
+		if (AActor* ActorA = DisabledCollision.Key)
 		{
-			if (AActor* ActorB = DisabledCollision.Value.Get())
+			if (AActor* ActorB = DisabledCollision.Value)
 			{
 				if (UPrimitiveComponent* PrimitiveComponentA = Cast<UPrimitiveComponent>(ActorA->GetRootComponent()))
 				{
@@ -390,12 +527,12 @@ void FValidInteractionSearch::VLogContext(const FColor& Color) const
 	{
 		static const TCHAR* LogName = TEXT("PoseSearchInteraction");
 
-		const int32 AnimContextsNum = AnimContexts.Num();
+		const int32 AnimContextsNum = Num();
 		TArray<FVector, TInlineAllocator<PreallocatedRolesNum, TMemStackAllocator<>>> Locations;
 		Locations.SetNum(AnimContextsNum);
 		for (int32 Index = 0; Index < AnimContextsNum; ++Index)
 		{
-			if (const UObject* AnimContext = AnimContexts[Index].Get())
+			if (const UObject* AnimContext = GetAnimContext(Index))
 			{
 				const AActor* ContextOwningActor = GetContextOwningActor(AnimContext);
 				check(ContextOwningActor);
@@ -411,7 +548,7 @@ void FValidInteractionSearch::VLogContext(const FColor& Color) const
 			{
 				for (int32 IndexAll = 0; IndexAll < AnimContextsNum; ++IndexAll)
 				{
-					if (const UObject* AnimContext = AnimContexts[IndexAll].Get())
+					if (const UObject* AnimContext = GetAnimContext(IndexAll))
 					{
 						UE_VLOG_SEGMENT(AnimContext, LogName, Display, Locations[IndexA], Locations[IndexB], Color, TEXT(""));
 					}
@@ -422,50 +559,34 @@ void FValidInteractionSearch::VLogContext(const FColor& Color) const
 }
 #endif // ENABLE_VISUAL_LOG
 
-// FInteractionSearchContext
-///////////////////////////////////////////////////////////
-const UE::PoseSearch::IPoseHistory* FInteractionSearchContext::GetPoseHistory(int32 Index) const
-{
-	check(PoseHistories.IsValidIndex(Index));
-	check(PoseHistories[Index]);
-	return PoseHistories[Index];
-}
-
-uint32 FInteractionSearchContext::GetHashForEquivalence() const
+void FValidInteractionSearch::AddReferencedObjects(FReferenceCollector& Collector)
 {
-	uint32 Hash = GetTypeHash(Database);
-
-	// skipping bDisableCollisions for hashing
+	FInteractionSearchContextBase::AddReferencedObjects(Collector);
 
-	for (int32 Index = 0; Index < AnimContexts.Num(); ++Index)
+	for (FDisabledCollision& DisabledCollision : DisabledCollisions)
 	{
-		Hash = HashCombineFast(Hash, GetTypeHash(AnimContexts[Index]));
+		Collector.AddReferencedObject(DisabledCollision.Key);
+		Collector.AddReferencedObject(DisabledCollision.Value);
 	}
-
-	for (int32 Index = 0; Index < PoseHistories.Num(); ++Index)
-	{
-		Hash = HashCombineFast(Hash, GetTypeHash(PoseHistories[Index]));
-	}
-
-	for (int32 Index = 0; Index < Roles.Num(); ++Index)
-	{
-		Hash = HashCombineFast(Hash, GetTypeHash(Roles[Index]));
-	}
-
-	return Hash;
 }
 
-void FInteractionSearchContext::Add(const UObject* AnimContext, const UE::PoseSearch::IPoseHistory* PoseHistory, const UE::PoseSearch::FRole Role)
+// FInteractionSearchContext
+///////////////////////////////////////////////////////////
+void FInteractionSearchContext::Add(const UObject* AnimContext, const UE::PoseSearch::IPoseHistory* PoseHistory
+	, const UE::PoseSearch::FRole Role, bool bWantsDisableCollisions, int32 TickPriority
+#if ENABLE_ANIM_DEBUG
+	, float InBroadPhaseRadius, float InBroadPhaseRadiusIncrementOnInteraction
+#endif // ENABLE_ANIM_DEBUG
+	)
 {
-	check(AnimContext && PoseHistory);
-		
-	// AnimContexts must be sorted to have deterministic searches across multiple frames
-	// as well as for faster IsEquivalent implementation
-	check(AnimContexts.IsEmpty() || GetAnimContext(AnimContexts.Num() - 1) < AnimContext);
+	FInteractionSearchContextBase::Add(AnimContext, PoseHistory, Role, bWantsDisableCollisions);
+	
+	TickPriorities.Add(TickPriority);
 
-	AnimContexts.Add(AnimContext);
-	PoseHistories.Add(PoseHistory);
-	Roles.Add(Role);
+#if ENABLE_ANIM_DEBUG
+	DebugBroadPhaseRadius.Add(InBroadPhaseRadius);
+	BroadPhaseRadiusIncrementOnInteraction.Add(InBroadPhaseRadiusIncrementOnInteraction);
+#endif // ENABLE_ANIM_DEBUG
 }
 
 FPoseSearchContinuingProperties FInteractionSearchContext::GetContinuingProperties() const
@@ -480,96 +601,31 @@ FPoseSearchContinuingProperties FInteractionSearchContext::GetContinuingProperti
 	return ContinuingProperties;
 }
 
-bool FInteractionSearchContext::IsEquivalent(const FInteractionSearchContext& Other) const
-{
-	// skipping bDisableCollisions for equality
-	return Database == Other.Database &&
-		AnimContexts == Other.AnimContexts &&
-		PoseHistories == Other.PoseHistories &&
-		Roles == Other.Roles;
-}
-
-#if DO_CHECK
-bool FInteractionSearchContext::CheckForConsistency() const
+float FInteractionSearchContext::GetDebugBroadPhaseRadius(int32 Index) const
 {
-	if (Database == nullptr)
-	{
-		return false;
-	}
-
-	const int32 Num = AnimContexts.Num();
-	if (Num < 1)
-	{
-		return false;
-	}
-
-	if (Num != PoseHistories.Num())
-	{
-		return false;
-	}
-
-	if (Num != Roles.Num())
-	{
-		return false;
-	}
-
-	for (int32 IndexA = 0; IndexA < Num; ++IndexA)
-	{
-		if (AnimContexts[IndexA] == nullptr)
-		{
-			return false;
-		}
-
-		for (int32 IndexB = IndexA + 1; IndexB < Num; ++IndexB)
-		{
-			if (AnimContexts[IndexA] == AnimContexts[IndexB])
-			{
-				return false;
-			}
-		}
-	}
-
-	for (int32 IndexA = 1; IndexA < Num; ++IndexA)
+	if (bIsContinuingInteraction)
 	{
-		// AnimContexts must be sorted to have deterministic searches across multiple frames
-		if (AnimContexts[IndexA - 1] >= AnimContexts[IndexA])
-		{
-			return false;
-		}
+		return DebugBroadPhaseRadius[Index] + BroadPhaseRadiusIncrementOnInteraction[Index];
 	}
-
-	for (int32 IndexA = 0; IndexA < Num; ++IndexA)
-	{
-		for (int32 IndexB = IndexA + 1; IndexB < Num; ++IndexB)
-		{
-			if (Roles[IndexA] == Roles[IndexB])
-			{
-				return false;
-			}
-		}
-	}
-
-	for (int32 IndexA = 0; IndexA < Num; ++IndexA)
-	{
-		if (PoseHistories[IndexA] == nullptr)
-		{
-			return false;
-		}
-
-		for (int32 IndexB = IndexA + 1; IndexB < Num; ++IndexB)
-		{
-			if (PoseHistories[IndexA] == PoseHistories[IndexB])
-			{
-				return false;
-			}
-		}
-	}
-
-	return true;
+	
+	return DebugBroadPhaseRadius[Index];
 }
 
-#endif // DO_CHECK
+void FInteractionSearchContext::SetContinuingProperties(float InPlayingAssetAccumulatedTime, const UObject* InPlayingAsset,
+	bool bInIsPlayingAssetMirrored, const FVector& InPlayingAssetBlendParameters, EPoseSearchInterruptMode InInterruptMode)
+{
+	PlayingAssetAccumulatedTime = InPlayingAssetAccumulatedTime;
+	PlayingAsset = InPlayingAsset;
+	bIsPlayingAssetMirrored = bInIsPlayingAssetMirrored;
+	PlayingAssetBlendParameters = InPlayingAssetBlendParameters;
+	InterruptMode = InInterruptMode;
+}
 
+void FInteractionSearchContext::AddReferencedObjects(FReferenceCollector& Collector)
+{
+	FInteractionSearchContextBase::AddReferencedObjects(Collector);
+	Collector.AddReferencedObject(PlayingAsset);
+}
 
 // FInteractionSearchResult
 ///////////////////////////////////////////////////////////
@@ -597,11 +653,11 @@ void FInteractionIsland::FPreTickFunction::ExecuteTick(float DeltaTime, enum ELe
 			{
 				if (const UObject* AnimContext = SearchContext.GetAnimContext(Index))
 				{
-					const IPoseHistory* PoseHistory = SearchContext.GetPoseHistory(Index);
-					check(PoseHistory && PoseHistory->AsWeak().IsValid());
-					
-					// since FInteractionIsland has a tick dependency with the USkeletalMeshComponent it's safe modify the IPoseHistory
-					const_cast<IPoseHistory*>(PoseHistory)->GenerateTrajectory(AnimContext, DeltaTime);
+					if (const IPoseHistory* PoseHistory = SearchContext.GetPoseHistory(Index))
+					{
+						// since FInteractionIsland has a tick dependency with the USkeletalMeshComponent it's safe modify the IPoseHistory
+						const_cast<IPoseHistory*>(PoseHistory)->GenerateTrajectory(AnimContext, DeltaTime);
+					}
 				}
 			}
 		}
@@ -869,7 +925,7 @@ void FInteractionIsland::RemoveTickDependencies(bool bValidateTickDependencies)
 		}
 
 		bool bMainActor = true;
-		for (TWeakObjectPtr<UActorComponent>& TickActorComponentPtr : TickActorComponents)
+		for (TObjectPtr<UActorComponent>& TickActorComponentPtr : TickActorComponents)
 		{
 			if (UActorComponent* TickActorComponent = TickActorComponentPtr.Get())
 			{
@@ -1035,13 +1091,13 @@ const AActor* FInteractionIsland::GetMainActor() const
 }
 
 #if ENABLE_ANIM_DEBUG
-void FInteractionIsland::LogTickDependencies(const TConstArrayView<TWeakObjectPtr<UActorComponent>> TickActorComponents, int32 InteractionIslandIndex)
+void FInteractionIsland::LogTickDependencies(const TConstArrayView<UActorComponent*> TickActorComponents, int32 InteractionIslandIndex)
 {
 	check(IsInGameThread());
 
-	for (const TWeakObjectPtr<UActorComponent>& TickActorComponentPtr : TickActorComponents)
+	for (UActorComponent* TickActorComponent : TickActorComponents)
 	{
-		if (UActorComponent* TickActorComponent = TickActorComponentPtr.Get())
+		if (TickActorComponent)
 		{
 			UE_LOG(LogPoseSearch, Log, TEXT("============== %s (Island %d) =============="), *TickActorComponent->GetOwner()->GetName(), InteractionIslandIndex);
 			if (IInteractionIslandDependency* IslandDependency = FindCustomDependency(TickActorComponent))
@@ -1079,7 +1135,7 @@ void FInteractionIsland::LogTickDependencies() const
 	}
 	else
 	{
-		TArray<TWeakObjectPtr<UActorComponent>> TickActorComponentsCopy = TickActorComponents;
+		TArray<UActorComponent*> TickActorComponentsCopy = TickActorComponents;
 		FFunctionGraphTask::CreateAndDispatchWhenReady([TickActorComponentsCopy, InteractionIslandIndex]()
 		{
 			LogTickDependencies(TickActorComponentsCopy, InteractionIslandIndex);
@@ -1088,83 +1144,84 @@ void FInteractionIsland::LogTickDependencies() const
 }
 #endif // ENABLE_ANIM_DEBUG
 
-bool FInteractionIsland::DoSearch_AnyThread(const UObject* AnimContext, FPoseSearchBlueprintResult& Result)
+void FInteractionIsland::AddReferencedObjects(FReferenceCollector& Collector)
 {
-	check(AnimContext);
+	for (TObjectPtr<UActorComponent>& TickActorComponent : TickActorComponents)
+	{
+		Collector.AddReferencedObject(TickActorComponent);
+	}
 
-	if (bSearchPerfomed)
+	for (TObjectPtr<const UObject>& IslandAnimContext : IslandAnimContexts)
 	{
-		// we now support multiple searches from the same AnimContext, and we return the already performed
-		// and cached result, so check(AnimContext != GetMainAnimContext()) is no longer valid!
-		return GetResult_AnyThread(AnimContext, Result);
+		Collector.AddReferencedObject(IslandAnimContext);
 	}
 
-	// searches are performed only on the MainAnimContext / MainActor
-	if (AnimContext != GetMainAnimContext())
+	for (FInteractionSearchContext& SearchContext : SearchContexts)
 	{
-		// search failed, because AnimContext is not the context from the main actor. continuing the search would lead to threading issues,
-		// since the threading model expect the main actor to perform the search! (this happens when recompiling animation blue prints..)
-		Result = FPoseSearchBlueprintResult();
-		return false;
+		SearchContext.AddReferencedObjects(Collector);
 	}
+}
 
-	QUICK_SCOPE_CYCLE_COUNTER(STAT_UPoseSearchInteractionInteractionIsland_Search);
+bool FInteractionIsland::DoSearch_AnyThread(const UObject* AnimContext, FPoseSearchBlueprintResult& Result)
+{
+	check(AnimContext);
 
-	FMemMark Mark(FMemStack::Get());
+	// searches are performed only on the MainAnimContext / MainActor
+	if (!bSearchPerfomed && AnimContext == GetMainAnimContext())
+	{
+		QUICK_SCOPE_CYCLE_COUNTER(STAT_UPoseSearchInteractionInteractionIsland_Search);
 
-	TArray<FSearchResult, TMemStackAllocator<>> PoseSearchResults;
-	TArray<FChooserEvaluationContext, TInlineAllocator<UE::PoseSearch::PreallocatedRolesNum>> Contexts;
+		FMemMark Mark(FMemStack::Get());
 
-	// SearchContexts are modified only by UPoseSearchInteractionSubsystem::Tick and constant otherwise, so it's safe to access them in a threaded environment without locks
-	PoseSearchResults.SetNum(SearchContexts.Num());
+		TArray<FSearchResult, TMemStackAllocator<>> PoseSearchResults;
+		TArray<FChooserEvaluationContext, TInlineAllocator<UE::PoseSearch::PreallocatedRolesNum>> Contexts;
 
-	for (int32 SearchIndex = 0; SearchIndex < SearchContexts.Num(); ++SearchIndex)
-	{
-		FInteractionSearchContext& InteractionSearchContext = SearchContexts[SearchIndex];
-		const UPoseSearchDatabase* Database = InteractionSearchContext.Database;
-		if (!Database)
+		// SearchContexts are modified only by UPoseSearchInteractionSubsystem::Tick and constant otherwise, so it's safe to access them in a threaded environment without locks
+		PoseSearchResults.SetNum(SearchContexts.Num());
+		for (int32 SearchIndex = 0; SearchIndex < SearchContexts.Num(); ++SearchIndex)
 		{
-			UE_LOG(LogPoseSearch, Error, TEXT("FInteractionIsland::DoSearch_AnyThread invalid context database"));
-			return false;
-		}
+			FInteractionSearchContext& InteractionSearchContext = SearchContexts[SearchIndex];
+			const UPoseSearchDatabase* Database = InteractionSearchContext.GetDatabase();
+			check(Database && Database->Schema);
+			
+			const int32 NumRoles = InteractionSearchContext.Num();
+			Contexts.Reset();
+			Contexts.SetNum(NumRoles);
 
-		if (!Database->Schema)
-		{
-			UE_LOG(LogPoseSearch, Error, TEXT("FInteractionIsland::DoSearch_AnyThread invalid schema for context database %s"), *Database->GetName());
-			return false;
-		}
+			bool bIsSearchContextValid = true;
+			FSearchContext SearchContext(0.f, FFloatInterval(0.f, 0.f), FPoseSearchEvent());
+			for (int32 RoleIndex = 0; RoleIndex < NumRoles; ++RoleIndex)
+			{
+				const UObject* SearchContextAnimContext = InteractionSearchContext.GetAnimContext(RoleIndex);
+				check(SearchContextAnimContext);
 
-		const int32 NumRoles = InteractionSearchContext.Num();
-		Contexts.Reset();
-		Contexts.SetNum(NumRoles);
+				Contexts[RoleIndex].AddObjectParam(const_cast<UObject*>(SearchContextAnimContext));
 
-		FSearchContext SearchContext(0.f, FFloatInterval(0.f, 0.f), FPoseSearchEvent());
-		for (int32 RoleIndex = 0; RoleIndex < NumRoles; ++RoleIndex)
-		{
-			const UObject* SearchContextAnimContext = InteractionSearchContext.GetAnimContext(RoleIndex);
-			if (!SearchContextAnimContext)
-			{
-				UE_LOG(LogPoseSearch, Error, TEXT("FInteractionIsland::DoSearch_AnyThread null anim context"));
-				return false;
+				if (const UE::PoseSearch::IPoseHistory* PoseHistory = InteractionSearchContext.GetPoseHistory(RoleIndex))
+				{
+					SearchContext.AddRole(InteractionSearchContext.GetRole(RoleIndex), &Contexts[RoleIndex], PoseHistory);
+				}
+				else
+				{
+					bIsSearchContextValid = false;
+					break;
+				}
 			}
 
-			Contexts[RoleIndex].AddObjectParam(const_cast<UObject*>(SearchContextAnimContext));
-
-			const UE::PoseSearch::IPoseHistory* PoseHistory = InteractionSearchContext.GetPoseHistory(RoleIndex);
-			check(PoseHistory);
-			SearchContext.AddRole(InteractionSearchContext.GetRoles()[RoleIndex], &Contexts[RoleIndex], PoseHistory);
+			if (bIsSearchContextValid)
+			{
+				const UObject* AssetsToSearch[] = { Database };
+				// @todo: we could perform multiple UPoseSearchLibrary::MotionMatch in parallel!
+				FSearchResults_Single MotionMatchSearchResults;
+				UPoseSearchLibrary::MotionMatch(SearchContext, AssetsToSearch, InteractionSearchContext.GetContinuingProperties(), MotionMatchSearchResults);
+				PoseSearchResults[SearchIndex] = MotionMatchSearchResults.GetBestResult();
+			}
 		}
 
-		const UObject* AssetsToSearch[] = { Database };
-		// @todo: we could perform multiple UPoseSearchLibrary::MotionMatch in parallel!
-		FSearchResults_Single MotionMatchSearchResults;
-		UPoseSearchLibrary::MotionMatch(SearchContext, AssetsToSearch, InteractionSearchContext.GetContinuingProperties(), MotionMatchSearchResults);
-		PoseSearchResults[SearchIndex] = MotionMatchSearchResults.GetBestResult();
+		InitSearchResults(SearchResults, PoseSearchResults, SearchContexts);
+		bSearchPerfomed = true;
 	}
-		
-	InitSearchResults(SearchResults, PoseSearchResults, SearchContexts);
-	bSearchPerfomed = true;
-	
+
 	return GetResult_AnyThread(AnimContext, Result);
 }
 
diff --git a/Engine/Plugins/Animation/PoseSearch/Source/Runtime/Private/PoseSearchInteractionSubsystem.cpp b/Engine/Plugins/Animation/PoseSearch/Source/Runtime/Private/PoseSearchInteractionSubsystem.cpp
--- a/Engine/Plugins/Animation/PoseSearch/Source/Runtime/Private/PoseSearchInteractionSubsystem.cpp
+++ b/Engine/Plugins/Animation/PoseSearch/Source/Runtime/Private/PoseSearchInteractionSubsystem.cpp
@@ -12,10 +12,6 @@
 #include "PoseSearch/Trace/PoseSearchTraceLogger.h"
 #include "VisualLogger/VisualLogger.h"
 
-#if WITH_EDITOR
-#include "Editor.h"
-#endif
-
 #include UE_INLINE_GENERATED_CPP_BY_NAME(PoseSearchInteractionSubsystem)
 
 namespace UE::PoseSearch
@@ -37,7 +33,7 @@ static FAutoConsoleVariableRef CVarPoseSearchInteractionLoglandsTickDependencies
 
 struct FAnimContextInfo
 {
-	void Init(const FPoseSearchInteractionAnimContextAvailabilities& InAnimContextAvailabilities)
+	void Init(const FInteractionAnimContextAvailabilities& InAnimContextAvailabilities)
 	{
 		check(InAnimContextAvailabilities.AnimContext && !InAnimContextAvailabilities.Availabilities.IsEmpty());
 		AnimContextAvailabilities = &InAnimContextAvailabilities;
@@ -49,7 +45,7 @@ struct FAnimContextInfo
 		Location = ContextOwningActor->GetTransform().GetLocation();
 
 		AvailabilitiesMaxBroadPhaseRadius = 0.f;
-		for (const FPoseSearchInteractionAvailabilityEx& Availability : AnimContextAvailabilities->Availabilities)
+		for (const FInteractionAvailability& Availability : AnimContextAvailabilities->Availabilities)
 		{
 			// @todo: optimize the AvailabilitiesMaxBroadPhaseRadius, since adding Availability.BroadPhaseRadiusIncrementOnInteraction is required ONLY if AnimContext is already part of an interaction
 			AvailabilitiesMaxBroadPhaseRadius = FMath::Max(AvailabilitiesMaxBroadPhaseRadius, Availability.BroadPhaseRadius + Availability.BroadPhaseRadiusIncrementOnInteraction);
@@ -61,7 +57,7 @@ struct FAnimContextInfo
 	const FRole* GetUniqueCommonRoleFilter() const
 	{
 		const FRole* UniqueCommonRoleFilter = nullptr;
-		for (const FPoseSearchInteractionAvailabilityEx& Availability : AnimContextAvailabilities->Availabilities)
+		for (const FInteractionAvailability& Availability : AnimContextAvailabilities->Availabilities)
 		{
 			if (Availability.RolesFilter.Num() == 1)
 			{
@@ -105,7 +101,7 @@ struct FAnimContextInfo
 		return AnimContextAvailabilities->AnimContext;
 	}
 
-	const TArray<FPoseSearchInteractionAvailabilityEx>& GetAvailabilities() const
+	const TArray<FInteractionAvailability>& GetAvailabilities() const
 	{
 		return AnimContextAvailabilities->Availabilities;
 	}
@@ -117,7 +113,7 @@ struct FAnimContextInfo
 	}
 
 private:
-	const FPoseSearchInteractionAnimContextAvailabilities* AnimContextAvailabilities = nullptr;
+	const FInteractionAnimContextAvailabilities* AnimContextAvailabilities = nullptr;
 
 	// cached AnimContext location
 	FVector Location = FVector::ZeroVector;
@@ -129,11 +125,9 @@ struct FAnimContextInfos : TArray<FAnimContextInfo, TMemStackAllocator<>> {};
 typedef TArray<const UPoseSearchDatabase*, TInlineAllocator<32, TMemStackAllocator<>>> FDatabasesPerTag;
 struct FTagToDatabases : TMap<FName, FDatabasesPerTag, TInlineSetAllocator<8, TMemStackSetAllocator<>>> {};
 
-} // namespace UE::PoseSearch
-
-// FPoseSearchInteractionAvailabilityEx
+// FInteractionAvailability
 ///////////////////////////////////////////////
-FString FPoseSearchInteractionAvailabilityEx::GetPoseHistoryName() const
+FString FInteractionAvailability::GetPoseHistoryName() const
 {
 	if (PoseHistory)
 	{
@@ -142,7 +136,7 @@ FString FPoseSearchInteractionAvailabilityEx::GetPoseHistoryName() const
 	return PoseHistoryName.ToString();
 }
 
-const UE::PoseSearch::IPoseHistory* FPoseSearchInteractionAvailabilityEx::GetPoseHistory(const UObject* AnimContext) const
+const UE::PoseSearch::IPoseHistory* FInteractionAvailability::GetPoseHistory(const UObject* AnimContext) const
 {
 	if (PoseHistory)
 	{
@@ -161,6 +155,8 @@ const UE::PoseSearch::IPoseHistory* FPoseSearchInteractionAvailabilityEx::GetPos
 	return nullptr;
 }
 
+} // namespace UE::PoseSearch
+
 // UPoseSearchInteractionSubsystem
 ///////////////////////////////////////////////
 UE::PoseSearch::FInteractionIsland& UPoseSearchInteractionSubsystem::CreateIsland()
@@ -289,11 +285,11 @@ void UPoseSearchInteractionSubsystem::RegenerateAllIslands(float DeltaSeconds)
 				{
 					if (int32* TickPriority = AnimContextToTickPriorityMap.Find(AnimContext))
 					{
-						*TickPriority = FMath::Max(*TickPriority, SearchContext.TickPriorities[AnimContextIndex]);
+						*TickPriority = FMath::Max(*TickPriority, SearchContext.GetTickPriority(AnimContextIndex));
 					}
 					else
 					{
-						AnimContextToTickPriorityMap.Add(AnimContext) = SearchContext.TickPriorities[AnimContextIndex];
+						AnimContextToTickPriorityMap.Add(AnimContext) = SearchContext.GetTickPriority(AnimContextIndex);
 					}
 				}
 			}
@@ -326,7 +322,7 @@ void UPoseSearchInteractionSubsystem::RegenerateAllIslands(float DeltaSeconds)
 		{
 			SearchContextsIndices.Sort([&SearchContexts](int32 A, int32 B)
 				{
-					return SearchContexts[A].AnimContexts[0] < SearchContexts[B].AnimContexts[0];
+					return SearchContexts[A].GetAnimContext(0) < SearchContexts[B].GetAnimContext(0);
 				});
 
 			AnimContextToTickPriority.Reset();
@@ -470,7 +466,7 @@ bool UPoseSearchInteractionSubsystem::ValidateAllIslands() const
 {
 	using namespace UE::PoseSearch;
 
-	TSet<TWeakObjectPtr<UActorComponent>> TickActorComponents;
+	TSet<const UActorComponent*> TickActorComponents;
 
 	typedef TSet<const UObject*> FIslandAnimContexts;
 	TArray<FIslandAnimContexts> IslandsAnimContexts;
@@ -481,7 +477,7 @@ bool UPoseSearchInteractionSubsystem::ValidateAllIslands() const
 	bool bAlreadyInSet = false;
 	for (const FInteractionIsland* Island : Islands)
 	{
-		for (const TWeakObjectPtr<UActorComponent>& TickActorComponent : Island->GetTickActorComponents())
+		for (const UActorComponent* TickActorComponent : Island->GetTickActorComponents())
 		{
 			TickActorComponents.Add(TickActorComponent, &bAlreadyInSet);
 			if (bAlreadyInSet)
@@ -550,12 +546,14 @@ void UPoseSearchInteractionSubsystem::PopulateContinuingProperties(float DeltaSe
 							check(SearchIndexAsset->GetToRealTimeFactor() > UE_KINDA_SMALL_NUMBER);
 							// in case DatabaseAsset->GetAnimationAsset() is a blendspace, SearchResult->AssetTime is a normalized time in the interval [0,1] 
 							// so we need to convert the delta time in seconds to the asset normalized time before integrating SearchResult->AssetTime
+
 							const float NormalizedDeltaTime = DeltaSeconds / SearchIndexAsset->GetToRealTimeFactor();
-							SearchContext.PlayingAssetAccumulatedTime = SearchResult->GetAssetTime() + NormalizedDeltaTime;
-							SearchContext.PlayingAsset = DatabaseAsset->GetAnimationAsset();
-							SearchContext.bIsPlayingAssetMirrored = SearchIndexAsset->IsMirrored();
-							SearchContext.PlayingAssetBlendParameters = SearchIndexAsset->GetBlendParameters();
+							const float PlayingAssetAccumulatedTime = SearchResult->GetAssetTime() + NormalizedDeltaTime;
+
 							// @todo: populate SearchContext.InterruptMode
+							EPoseSearchInterruptMode InterruptMode = EPoseSearchInterruptMode::DoNotInterrupt;
+							SearchContext.SetContinuingProperties(PlayingAssetAccumulatedTime, DatabaseAsset->GetAnimationAsset(),
+								SearchIndexAsset->IsMirrored(), SearchIndexAsset->GetBlendParameters(), InterruptMode);
 						}
 					}
 				}
@@ -577,9 +575,9 @@ UE::PoseSearch::FInteractionIsland* UPoseSearchInteractionSubsystem::FindIsland(
 
 			for (FInteractionIsland* Island : Islands)
 			{
-				for (const TWeakObjectPtr<const UObject>& IslandAnimContext : Island->GetIslandAnimContexts())
+				for (const UObject* IslandAnimContext : Island->GetIslandAnimContexts())
 				{ 
-					if (GetContextOwningActor(IslandAnimContext.Get(), false) == Actor)
+					if (GetContextOwningActor(IslandAnimContext, false) == Actor)
 					{
 						return Island;
 					}
@@ -617,45 +615,16 @@ UPoseSearchInteractionSubsystem* UPoseSearchInteractionSubsystem::GetSubsystem_A
 	return nullptr;
 }
 
-#if WITH_EDITOR
-// to address check(PoseHistory && PoseHistory->AsWeak().IsValid()) in FInteractionIsland::FPreTickFunction::ExecuteTick caused by
-// run PIE, pause (FInteractionIsland are still caching the SearchContexts calculated with the previous pose history nodes), compile BPs, unpause, crash!
-void UPoseSearchInteractionSubsystem::OnBlueprintPreCompile(UBlueprint* BlueprintToCompile)
-{
-	// @todo: reenable after fixing OnBlueprintPreCompile being called from worker thread
-
-	//check(IsInGameThread());
-	//DestroyAllIslands();
-}
-#endif // WITH_EDITOR
-
 void UPoseSearchInteractionSubsystem::Initialize(FSubsystemCollectionBase& Collection)
 {
 	check(IsInGameThread());
-	
-#if WITH_EDITOR
-	if (GEditor != nullptr)
-	{
-		GEditor->OnBlueprintPreCompile().AddUObject(this, &UPoseSearchInteractionSubsystem::OnBlueprintPreCompile);
-	}
-#endif // WITH_EDITOR
-
 	Super::Initialize(Collection);
 }
 
 void UPoseSearchInteractionSubsystem::Deinitialize()
 {
 	UpdateValidInteractionSearches();
-
 	DestroyAllIslands();
-
-#if WITH_EDITOR
-	if (GEditor != nullptr)
-	{
-		GEditor->OnBlueprintPreCompile().RemoveAll(this);
-	}
-#endif // WITH_EDITOR
-
 	Super::Deinitialize();
 }
 
@@ -669,7 +638,7 @@ void UPoseSearchInteractionSubsystem::AddAvailabilities(const TArrayView<const F
 
 	if (ReservedAnimContextsAvailabilitiesIndex < AnimContextsAvailabilitiesBuffer.Num())
 	{
-		FPoseSearchInteractionAnimContextAvailabilities& AnimContextAvailabilities = AnimContextsAvailabilitiesBuffer[ReservedAnimContextsAvailabilitiesIndex];
+		FInteractionAnimContextAvailabilities& AnimContextAvailabilities = AnimContextsAvailabilitiesBuffer[ReservedAnimContextsAvailabilitiesIndex];
 		AnimContextAvailabilities.AnimContext = AnimContext;
 		AnimContextAvailabilities.Availabilities.Reset();
 		for (const FPoseSearchInteractionAvailability& Availability : Availabilities)
@@ -686,14 +655,14 @@ void UPoseSearchInteractionSubsystem::GenerateAnimContextInfosAndTagToDatabases(
 	using namespace UE::PoseSearch;
 
 	const UWorld* SubsystemWorld = GetWorld();
-	TConstArrayView<FPoseSearchInteractionAnimContextAvailabilities> AnimContextsAvailabilities = MakeArrayView(AnimContextsAvailabilitiesBuffer.GetData(), AnimContextsAvailabilitiesNum);
+	TConstArrayView<FInteractionAnimContextAvailabilities> AnimContextsAvailabilities = MakeArrayView(AnimContextsAvailabilitiesBuffer.GetData(), AnimContextsAvailabilitiesNum);
 	
 #if DO_CHECK
 	check(SubsystemWorld);
 	check(AnimContextInfos.IsEmpty() && TagToDatabases.IsEmpty());
 #endif // DO_CHECK
 
-	for (const FPoseSearchInteractionAnimContextAvailabilities& AnimContextAvailabilities : AnimContextsAvailabilities)
+	for (const FInteractionAnimContextAvailabilities& AnimContextAvailabilities : AnimContextsAvailabilities)
 	{
 #if DO_CHECK
 		check(AnimContextAvailabilities.AnimContext);
@@ -703,7 +672,7 @@ void UPoseSearchInteractionSubsystem::GenerateAnimContextInfosAndTagToDatabases(
 #endif // DO_CHECK
 		
 		// adding AnimContext to AnimContextsWithAvailabilities only if at least one availability has a valid database or has a valid tag
-		for (const FPoseSearchInteractionAvailabilityEx& InteractionAvailabilityEx : AnimContextAvailabilities.Availabilities)
+		for (const FInteractionAvailability& InteractionAvailabilityEx : AnimContextAvailabilities.Availabilities)
 		{
 			bool bAnyValidAvailability = false;
 			if (const UPoseSearchDatabase* Database = InteractionAvailabilityEx.Database.Get())
@@ -727,7 +696,7 @@ void UPoseSearchInteractionSubsystem::GenerateAnimContextInfosAndTagToDatabases(
 			if (DistanceSquared <= MaxDistanceSquared)
 			{
 				// the AnimContext of AnimContextInfos[AnimContextInfoIndexA] can potentially interact with the one from AnimContextInfos[AnimContextInfoIndexB]:
-				// linking AnimContextInfos[AnimContextInfoIndexA] to AnimContextInfos[AnimContextInfoIndexA] and vice versa to keep track of this when evaluating the broad phase.
+				// linking AnimContextInfos[AnimContextInfoIndexA] to AnimContextInfos[AnimContextInfoIndexB] and vice versa to keep track of this when evaluating the broad phase.
 				// Since AnimContextInfos does't reallocate anymore, it's safe to store pointers to internal elements of the array!
 				AnimContextInfos[AnimContextInfoIndexA].AddNearbyAnimContextInfo(AnimContextInfos[AnimContextInfoIndexB]);
 				AnimContextInfos[AnimContextInfoIndexB].AddNearbyAnimContextInfo(AnimContextInfos[AnimContextInfoIndexA]);
@@ -841,7 +810,7 @@ void UPoseSearchInteractionSubsystem::GenerateSearchContexts(float DeltaSeconds,
 
 	struct FRoledAnimContextInfo
 	{
-		FRoledAnimContextInfo(const FPoseSearchInteractionAvailabilityEx& InAvailability, const FAnimContextInfo& InAnimContextInfo, const FRole InRole, const IPoseHistory* InPoseHistory, const UPoseSearchDatabase& InDatabase)
+		FRoledAnimContextInfo(const FInteractionAvailability& InAvailability, const FAnimContextInfo& InAnimContextInfo, const FRole InRole, const IPoseHistory* InPoseHistory, const UPoseSearchDatabase& InDatabase)
 			: Availability(&InAvailability)
 			, AnimContextInfo(&InAnimContextInfo)
 			, Role(InRole)
@@ -898,7 +867,7 @@ void UPoseSearchInteractionSubsystem::GenerateSearchContexts(float DeltaSeconds,
 		}
 
 		// Availability that spawned this FRoledAnimContextInfo
-		const FPoseSearchInteractionAvailabilityEx* Availability = nullptr;
+		const FInteractionAvailability* Availability = nullptr;
 		// AnimContextInfo containing all the information regarding the AnimContext that spawned this FRoledAnimContextInfo,
 		// including all the availabilities associated to the AnimContext as well as all the other AnimContext(s) it can potentially interact with
 		const FAnimContextInfo* AnimContextInfo = nullptr;
@@ -969,7 +938,7 @@ void UPoseSearchInteractionSubsystem::GenerateSearchContexts(float DeltaSeconds,
 			InternalData.RemoveAt(PrevIndex, Num() - PrevIndex, EAllowShrinking::No);
 		}
 
-		void AddRoledAnimContextInfos(const FPoseSearchInteractionAvailabilityEx& Availability, const FAnimContextInfo& AnimContextInfo, const IPoseHistory* PoseHistory, const UPoseSearchDatabase& Database)
+		void AddRoledAnimContextInfos(const FInteractionAvailability& Availability, const FAnimContextInfo& AnimContextInfo, const IPoseHistory* PoseHistory, const UPoseSearchDatabase& Database)
 		{
 			const UPoseSearchSchema* Schema = Database.Schema;
 			check(Schema);
@@ -1096,7 +1065,7 @@ void UPoseSearchInteractionSubsystem::GenerateSearchContexts(float DeltaSeconds,
 	public:
 		typedef FInternalMap::TIterator TIterator;
 
-		void AddRoledAnimContextInfos(const FPoseSearchInteractionAvailabilityEx& Availability, const FAnimContextInfo& AnimContextInfo, const IPoseHistory* PoseHistory, const UPoseSearchDatabase* Database)
+		void AddRoledAnimContextInfos(const FInteractionAvailability& Availability, const FAnimContextInfo& AnimContextInfo, const IPoseHistory* PoseHistory, const UPoseSearchDatabase* Database)
 		{
 			check(Database && Database->Schema);
 			FRoledAnimContextInfos& RoledAnimContextInfos = InternalMap.FindOrAdd(Database);
@@ -1126,7 +1095,7 @@ void UPoseSearchInteractionSubsystem::GenerateSearchContexts(float DeltaSeconds,
 		{
 			// analyzing all the Availability(s) associated with this AnimContext and eventually generate the associated FRoledAnimContextInfos,
 			// inserted in a per database sorted data structure (DatabaseToRoledAnimContextInfos)
-			for (const FPoseSearchInteractionAvailabilityEx& Availability : AnimContextInfo.GetAvailabilities())
+			for (const FInteractionAvailability& Availability : AnimContextInfo.GetAvailabilities())
 			{
 				const IPoseHistory* PoseHistory = Availability.GetPoseHistory(AnimContextInfo.GetAnimContext());
 				check(PoseHistory);
@@ -1223,7 +1192,7 @@ void UPoseSearchInteractionSubsystem::GenerateSearchContexts(float DeltaSeconds,
 							}
 
 							FInteractionSearchContext SearchContext;
-							SearchContext.Database = Database;
+							SearchContext.SetDatabase(Database);
 
 							// setting up a FRoledAnimContextInfo in RoledAnimContextInfos describing 
 							// this potential interaction properties about how to perform the search
@@ -1233,14 +1202,12 @@ void UPoseSearchInteractionSubsystem::GenerateSearchContexts(float DeltaSeconds,
 								const FRoledAnimContextInfo& RoledAnimContextInfo = RoledAnimContextInfos[RoledAnimContextIndex];
 
 								check(RoledAnimContextInfo.PoseHistory && RoledAnimContextInfo.Availability);
-								SearchContext.Add(RoledAnimContextInfo.AnimContextInfo->GetAnimContext(), RoledAnimContextInfo.PoseHistory, RoledAnimContextInfo.Role);
-								SearchContext.bDisableCollisions |= RoledAnimContextInfo.Availability->bDisableCollisions;
-								SearchContext.TickPriorities.Add(RoledAnimContextInfo.Availability->TickPriority);
-
+								SearchContext.Add(RoledAnimContextInfo.AnimContextInfo->GetAnimContext(), RoledAnimContextInfo.PoseHistory, RoledAnimContextInfo.Role
+									, RoledAnimContextInfo.Availability->bDisableCollisions, RoledAnimContextInfo.Availability->TickPriority
 #if ENABLE_ANIM_DEBUG
-								SearchContext.DebugBroadPhaseRadius.Add(RoledAnimContextInfo.Availability->BroadPhaseRadius);
-								SearchContext.BroadPhaseRadiusIncrementOnInteraction.Add(RoledAnimContextInfo.Availability->BroadPhaseRadiusIncrementOnInteraction);
+									, RoledAnimContextInfo.Availability->BroadPhaseRadius, RoledAnimContextInfo.Availability->BroadPhaseRadiusIncrementOnInteraction
 #endif // ENABLE_ANIM_DEBUG
+								);
 							}
 
 							// does SearchContext cover all the roles required by this interaction?
@@ -1255,7 +1222,7 @@ void UPoseSearchInteractionSubsystem::GenerateSearchContexts(float DeltaSeconds,
 
 							// looking for a preexisting valid interaction resembling SearchContext
 							const uint32 SearchContextHash = SearchContext.GetHashForEquivalence();
-							SearchContext.bIsContinuingInteraction = ValidInteractionSearches.Contains(SearchContextHash);
+							SearchContext.SetContinuingInteraction(ValidInteractionSearches.Contains(SearchContextHash));
 
 							// checking if this combination is valid for the Database:
 							for (int32 CombinationIndex = 0; CombinationIndex < CombinationCardinality; ++CombinationIndex)
@@ -1273,7 +1240,7 @@ void UPoseSearchInteractionSubsystem::GenerateSearchContexts(float DeltaSeconds,
 									const float DistanceSquared = DeltaLocation.SquaredLength();
 
 									float MaxDistanceSquared;
-									if (SearchContext.bIsContinuingInteraction)
+									if (SearchContext.IsContinuingInteraction())
 									{
 										MaxDistanceSquared = FMath::Min(RoledAnimContextInfo.IncrementedBroadPhaseRadiusSquared, OtherRoledAnimContextInfo.IncrementedBroadPhaseRadiusSquared);
 									}
@@ -1511,6 +1478,34 @@ void UPoseSearchInteractionSubsystem::GetResult_AnyThread(const UObject* AnimCon
 	}
 }
 
+void UPoseSearchInteractionSubsystem::AddReferencedObjects(UObject* This, FReferenceCollector& Collector)
+{
+	using namespace UE::PoseSearch;
+
+	Super::AddReferencedObjects(This, Collector);
+
+	UPoseSearchInteractionSubsystem* Subsystem = CastChecked<UPoseSearchInteractionSubsystem>(This);
+
+	for (FInteractionAnimContextAvailabilities& AnimContextAvailabilities : Subsystem->AnimContextsAvailabilitiesBuffer)
+	{
+		Collector.AddReferencedObject(AnimContextAvailabilities.AnimContext);
+		for (FInteractionAvailability& Availability : AnimContextAvailabilities.Availabilities)
+		{
+			Collector.AddReferencedObject(Availability.Database);
+		}
+	}
+
+	for (FInteractionIsland* Island : Subsystem->Islands)
+	{
+		Island->AddReferencedObjects(Collector);
+	}
+
+	for (TPair<uint32, FValidInteractionSearch>& ValidInteractionSearchPair : Subsystem->ValidInteractionSearches)
+	{
+		ValidInteractionSearchPair.Value.AddReferencedObjects(Collector);
+	}
+}
+
 #if ENABLE_ANIM_DEBUG
 void UPoseSearchInteractionSubsystem::DebugDrawIslands() const
 {
@@ -1548,9 +1543,9 @@ void UPoseSearchInteractionSubsystem::DebugDrawIslands() const
 	TArray<const UObject*, TInlineAllocator<256>> AllAnimContexts;
 	for (const FInteractionIsland* Island : Islands)
 	{
-		for (const TWeakObjectPtr<const UObject>& IslandAnimContextPtr : Island->GetIslandAnimContexts())
+		for (const UObject* IslandAnimContext : Island->GetIslandAnimContexts())
 		{
-			if (const UObject* IslandAnimContext = IslandAnimContextPtr.Get())
+			if (IslandAnimContext)
 			{
 				AllAnimContexts.Add(IslandAnimContext);
 			}
@@ -1569,16 +1564,7 @@ void UPoseSearchInteractionSubsystem::DebugDrawIslands() const
 				{
 					if (const UObject* AnimContext = SearchContext.GetAnimContext(Index))
 					{
-						float MaxBroadPhaseRadius;
-						if (SearchContext.bIsContinuingInteraction)
-						{
-							MaxBroadPhaseRadius = SearchContext.DebugBroadPhaseRadius[Index] + SearchContext.BroadPhaseRadiusIncrementOnInteraction[Index];
-						}
-						else
-						{
-							MaxBroadPhaseRadius = SearchContext.DebugBroadPhaseRadius[Index];
-						}
-
+						const float MaxBroadPhaseRadius = SearchContext.GetDebugBroadPhaseRadius(Index);
 						if (MaxBroadPhaseRadius > UE_SMALL_NUMBER)
 						{
 							const FTransform& Transform = GetContextTransform(AnimContext);
diff --git a/Engine/Plugins/Animation/PoseSearch/Source/Runtime/Public/PoseSearch/PoseSearchInteractionAvailability.h b/Engine/Plugins/Animation/PoseSearch/Source/Runtime/Public/PoseSearch/PoseSearchInteractionAvailability.h
--- a/Engine/Plugins/Animation/PoseSearch/Source/Runtime/Public/PoseSearch/PoseSearchInteractionAvailability.h
+++ b/Engine/Plugins/Animation/PoseSearch/Source/Runtime/Public/PoseSearch/PoseSearchInteractionAvailability.h
@@ -59,13 +59,12 @@ struct FPoseSearchInteractionAvailability
 	bool operator==(const FPoseSearchInteractionAvailability& Other) const = default;
 };
 
+namespace UE::PoseSearch
+{
 // Experimental, this feature might be removed without warning, not for production use
-USTRUCT(Experimental)
-struct FPoseSearchInteractionAvailabilityEx : public FPoseSearchInteractionAvailability
+struct FInteractionAvailability : public FPoseSearchInteractionAvailability
 {
-	GENERATED_BODY()
-
-	void Init(const FPoseSearchInteractionAvailability& InAvailability, FName InPoseHistoryName, const UE::PoseSearch::IPoseHistory* InPoseHistory)
+	void Init(const FPoseSearchInteractionAvailability& InAvailability, FName InPoseHistoryName, const IPoseHistory* InPoseHistory)
 	{
 		check(!InPoseHistory || InPoseHistory->AsWeak().IsValid());
 
@@ -75,22 +74,18 @@ struct FPoseSearchInteractionAvailabilityEx : public FPoseSearchInteractionAvail
 	}
 
 	FString GetPoseHistoryName() const;
-	const UE::PoseSearch::IPoseHistory* GetPoseHistory(const UObject* AnimContext) const;
+	const IPoseHistory* GetPoseHistory(const UObject* AnimContext) const;
 
 private:
 	FName PoseHistoryName;
-	const UE::PoseSearch::IPoseHistory* PoseHistory = nullptr;
+	const IPoseHistory* PoseHistory = nullptr;
 };
 
 // Experimental, this feature might be removed without warning, not for production use
-USTRUCT(Experimental)
-struct FPoseSearchInteractionAnimContextAvailabilities
+struct FInteractionAnimContextAvailabilities
 {
-	GENERATED_BODY()
-
-	UPROPERTY(Transient)
 	TObjectPtr<const UObject> AnimContext;
-
-	UPROPERTY(Transient)
-	TArray<FPoseSearchInteractionAvailabilityEx> Availabilities;
+	TArray<FInteractionAvailability> Availabilities;
 };
+
+} // namespace UE::PoseSearch
diff --git a/Engine/Plugins/Animation/PoseSearch/Source/Runtime/Public/PoseSearch/PoseSearchInteractionIsland.h b/Engine/Plugins/Animation/PoseSearch/Source/Runtime/Public/PoseSearch/PoseSearchInteractionIsland.h
--- a/Engine/Plugins/Animation/PoseSearch/Source/Runtime/Public/PoseSearch/PoseSearchInteractionIsland.h
+++ b/Engine/Plugins/Animation/PoseSearch/Source/Runtime/Public/PoseSearch/PoseSearchInteractionIsland.h
@@ -19,10 +19,48 @@ struct IInteractionIslandDependency;
 struct FInteractionSearchContext;
 
 // Experimental, this feature might be removed without warning, not for production use
-typedef TArray<TPair<TWeakObjectPtr<AActor>, TWeakObjectPtr<AActor>>> FDisabledCollisions;
+typedef TPair<TObjectPtr<AActor>, TObjectPtr<AActor>> FDisabledCollision;
+// Experimental, this feature might be removed without warning, not for production use
+typedef TArray<FDisabledCollision> FDisabledCollisions;
 
 // Experimental, this feature might be removed without warning, not for production use
-struct FValidInteractionSearch
+struct FInteractionSearchContextBase
+{
+	int32 Num() const { return AnimContexts.Num(); }
+	const UObject* GetAnimContext(int32 Index) const { return AnimContexts[Index]; }
+	const UE::PoseSearch::IPoseHistory* GetPoseHistory(int32 Index) const;
+	const UE::PoseSearch::FRole GetRole(int32 Index) const { return Roles[Index]; }
+	const UPoseSearchDatabase* GetDatabase() const { return Database; }
+	const TConstArrayView<UE::PoseSearch::FRole> GetRoles() const { return Roles; }
+
+	bool IsDisableCollisions() const { return bDisableCollisions; }
+
+	uint32 GetHashForEquivalence() const;
+	bool IsEquivalent(const FInteractionSearchContextBase& Other) const;
+	
+	void SetDatabase(const UPoseSearchDatabase* InDatabase) { Database = InDatabase; }
+	void AddReferencedObjects(FReferenceCollector& Collector);
+	void Add(const UObject* AnimContext, const UE::PoseSearch::IPoseHistory* PoseHistory, const UE::PoseSearch::FRole Role, bool bWantsDisableCollisions);
+
+#if DO_CHECK
+	bool CheckForConsistency() const;
+#endif // DO_CHECK
+
+private:
+	TArray<TObjectPtr<const UObject>, TInlineAllocator<PreallocatedRolesNum>> AnimContexts;
+
+	// storing PoseHistories as TWeakPtr since they are from nodes or traits that can be invalidated within the same frame.
+	// for performance reasons are not used to determine equivalence IsEquivalent / GetHashForEquivalence
+	TArray<TWeakPtr<const UE::PoseSearch::IPoseHistory, ESPMode::ThreadSafe>, TInlineAllocator<PreallocatedRolesNum>> PoseHistories;
+	
+	TArray<UE::PoseSearch::FRole, TInlineAllocator<PreallocatedRolesNum>> Roles;
+
+	TObjectPtr<const UPoseSearchDatabase> Database;
+	bool bDisableCollisions = false;
+};
+
+// Experimental, this feature might be removed without warning, not for production use
+struct FValidInteractionSearch : public FInteractionSearchContextBase
 {
 	void Init(const FInteractionSearchContext& Other);
 	void DoDisableCollisions();
@@ -32,44 +70,38 @@ struct FValidInteractionSearch
 	void VLogContext(const FColor& Color) const;
 #endif // ENABLE_VISUAL_LOG
 
-private:
-	TWeakObjectPtr<const UPoseSearchDatabase> Database;
-	bool bDisableCollisions = false;
+	void AddReferencedObjects(FReferenceCollector& Collector);
 
-	TArray<TWeakObjectPtr<const UObject>, TInlineAllocator<PreallocatedRolesNum>> AnimContexts;
-	TArray<TWeakPtr<const UE::PoseSearch::IPoseHistory, ESPMode::ThreadSafe>, TInlineAllocator<PreallocatedRolesNum>> PoseHistories;
-	TArray<UE::PoseSearch::FRole, TInlineAllocator<PreallocatedRolesNum>> Roles;
+private:
 	FDisabledCollisions DisabledCollisions;
 };
 
 // Experimental, this feature might be removed without warning, not for production use
-// stack allocated temporary struct used to generate interaction contexts, then stored using weak pointers as FValidInteractionSearch
-struct FInteractionSearchContext
+// stack allocated temporary struct used to generate interaction contexts, then the FInteractionSearchContextBase portion will be 
+// stored including DisabledCollisions as FValidInteractionSearch
+struct FInteractionSearchContext : public FInteractionSearchContextBase
 {
-	uint32 GetHashForEquivalence() const;
-	void Add(const UObject* AnimContext, const UE::PoseSearch::IPoseHistory* PoseHistory, const UE::PoseSearch::FRole Role);
-	int32 Num() const { return AnimContexts.Num(); }
-	const UObject* GetAnimContext(int32 Index) const { return AnimContexts[Index]; }
-	const UE::PoseSearch::IPoseHistory* GetPoseHistory(int32 Index) const;
-	const UE::PoseSearch::FRole GetRole(int32 Index) const { return Roles[Index]; }
-	const TConstArrayView<UE::PoseSearch::FRole> GetRoles() const { return Roles; }
-	FPoseSearchContinuingProperties GetContinuingProperties() const;
+	void Add(const UObject* AnimContext, const UE::PoseSearch::IPoseHistory* PoseHistory
+		, const UE::PoseSearch::FRole Role, bool bWantsDisableCollisions, int32 TickPriority
+#if ENABLE_ANIM_DEBUG
+		, float InBroadPhaseRadius, float InBroadPhaseRadiusIncrementOnInteraction
+#endif // ENABLE_ANIM_DEBUG
+	);
 
-	bool IsValid() const { return !AnimContexts.IsEmpty(); }
-	bool IsEquivalent(const FInteractionSearchContext& Other) const;
+	int32 GetTickPriority(int32 Index) const { return TickPriorities[Index]; }
 
-#if DO_CHECK
-	bool CheckForConsistency() const;
-#endif // DO_CHECK
+	FPoseSearchContinuingProperties GetContinuingProperties() const;
 
-	const UPoseSearchDatabase* Database = nullptr;
-	bool bDisableCollisions = false;
+	bool IsContinuingInteraction() const { return bIsContinuingInteraction; }
+	float GetDebugBroadPhaseRadius(int32 Index) const;
 
-	TArray<const UObject*, TInlineAllocator<PreallocatedRolesNum, TMemStackAllocator<>>> AnimContexts;
-	TArray<const UE::PoseSearch::IPoseHistory*, TInlineAllocator<PreallocatedRolesNum, TMemStackAllocator<>>> PoseHistories;
-	TArray<UE::PoseSearch::FRole, TInlineAllocator<PreallocatedRolesNum, TMemStackAllocator<>>> Roles;
+	void SetContinuingInteraction(bool bInIsContinuingInteraction) { bIsContinuingInteraction = bInIsContinuingInteraction; }
+	void AddReferencedObjects(FReferenceCollector& Collector);
+	void SetContinuingProperties(float InPlayingAssetAccumulatedTime, const UObject* InPlayingAsset,
+		bool bInIsPlayingAssetMirrored, const FVector& InPlayingAssetBlendParameters, EPoseSearchInterruptMode InInterruptMode);
 
-	const UObject* PlayingAsset = nullptr;
+private:
+	TObjectPtr<const UObject> PlayingAsset;
 	float PlayingAssetAccumulatedTime = 0.f;
 	bool bIsPlayingAssetMirrored = false;
 	FVector PlayingAssetBlendParameters = FVector::ZeroVector;
@@ -117,8 +149,8 @@ struct FInteractionIsland
 	bool DoSearch_AnyThread(const UObject* AnimContext, FPoseSearchBlueprintResult& Result);
 	bool GetResult_AnyThread(const UObject* AnimContext, FPoseSearchBlueprintResult& Result, bool bCompareOwningActors = false);
 
-	const TArray<TWeakObjectPtr<UActorComponent>>& GetTickActorComponents() const { return TickActorComponents; }
-	const TArray<TWeakObjectPtr<const UObject>>& GetIslandAnimContexts() const { return IslandAnimContexts; }
+	const TArray<UActorComponent*>& GetTickActorComponents() const { return TickActorComponents; }
+	const TArray<const UObject*>& GetIslandAnimContexts() const { return IslandAnimContexts; }
 	const TArray<FInteractionSearchContext>& GetSearchContexts() const { return SearchContexts; }
 	const TArray<FInteractionSearchResult>& GetSearchResults() const { return SearchResults; }
 
@@ -135,9 +167,11 @@ struct FInteractionIsland
 	void LogTickDependencies() const;
 #endif // ENABLE_ANIM_DEBUG
 
+	void AddReferencedObjects(FReferenceCollector& Collector);
+
 private:
 #if ENABLE_ANIM_DEBUG
-	static void LogTickDependencies(const TConstArrayView<TWeakObjectPtr<UActorComponent>> TickActorComponents, int32 InteractionIslandIndex);
+	static void LogTickDependencies(const TConstArrayView<UActorComponent*> TickActorComponents, int32 InteractionIslandIndex);
 #endif // ENABLE_ANIM_DEBUG
 
 	void AddTickDependencies(UActorComponent* TickActorComponent, bool bInIsMainActor);
@@ -164,9 +198,9 @@ private:
 	FPostTickFunction PostTickFunction;
 	bool bHasTickDependencies = false;
 
-	TArray<TWeakObjectPtr<UActorComponent>> TickActorComponents;
+	TArray<TObjectPtr<UActorComponent>> TickActorComponents;
 	// all the AnimContext in this island. Each SearchContexts will contain a subset of IslandAnimContexts 
-	TArray<TWeakObjectPtr<const UObject>> IslandAnimContexts;
+	TArray<TObjectPtr<const UObject>> IslandAnimContexts;
 
 	// there's one FSearchContext for each search we need to perform (including all the possible roles combinations). Added by UPoseSearchInteractionSubsystem::Tick
 	// islands don't get deallocated, so this array once warmed up will not hit the allocator and waste cycles
diff --git a/Engine/Plugins/Animation/PoseSearch/Source/Runtime/Public/PoseSearch/PoseSearchInteractionSubsystem.h b/Engine/Plugins/Animation/PoseSearch/Source/Runtime/Public/PoseSearch/PoseSearchInteractionSubsystem.h
--- a/Engine/Plugins/Animation/PoseSearch/Source/Runtime/Public/PoseSearch/PoseSearchInteractionSubsystem.h
+++ b/Engine/Plugins/Animation/PoseSearch/Source/Runtime/Public/PoseSearch/PoseSearchInteractionSubsystem.h
@@ -90,6 +90,9 @@ public:
 
 	TConstArrayView<UE::PoseSearch::FInteractionIsland*> GetInteractionIslands() const { return Islands; }
 
+	// garbage collection handling
+	static UE_API void AddReferencedObjects(UObject* This, FReferenceCollector& Collector);
+
 private:
 #if ENABLE_ANIM_DEBUG
 	friend struct UE::PoseSearch::FInteractionValidator;
@@ -124,13 +127,8 @@ private:
 	static void OnInteractionContinuing(UE::PoseSearch::FValidInteractionSearch& ValidInteractionSearch);
 	static void OnInteractionEnd(UE::PoseSearch::FValidInteractionSearch& ValidInteractionSearch);
 
-#if WITH_EDITOR
-	void OnBlueprintPreCompile(UBlueprint* BlueprintToCompile);
-#endif // WITH_EDITOR
-
 	// it could contain duplicated availabilities
-	UPROPERTY(Transient)
-	TArray<FPoseSearchInteractionAnimContextAvailabilities> AnimContextsAvailabilitiesBuffer;
+	TArray<UE::PoseSearch::FInteractionAnimContextAvailabilities> AnimContextsAvailabilitiesBuffer;
 	FThreadSafeCounter AnimContextsAvailabilitiesIndex = 0;
 	int32 AnimContextsAvailabilitiesNum = 0;
 
