PoseSearch - fix for missing #if ENABLE_ANIM_DEBUG wrapping FInteractionSearchContext::GetDebugBroadPhaseRadius

diff --git a/Engine/Plugins/Animation/PoseSearch/Source/Runtime/Private/PoseSearchInteractionIsland.cpp b/Engine/Plugins/Animation/PoseSearch/Source/Runtime/Private/PoseSearchInteractionIsland.cpp
--- a/Engine/Plugins/Animation/PoseSearch/Source/Runtime/Private/PoseSearchInteractionIsland.cpp
+++ b/Engine/Plugins/Animation/PoseSearch/Source/Runtime/Private/PoseSearchInteractionIsland.cpp
@@ -601,6 +601,7 @@ FPoseSearchContinuingProperties FInteractionSearchContext::GetContinuingProperti
 	return ContinuingProperties;
 }
 
+#if ENABLE_ANIM_DEBUG
 float FInteractionSearchContext::GetDebugBroadPhaseRadius(int32 Index) const
 {
 	if (bIsContinuingInteraction)
@@ -610,6 +611,7 @@ float FInteractionSearchContext::GetDebugBroadPhaseRadius(int32 Index) const
 	
 	return DebugBroadPhaseRadius[Index];
 }
+#endif // ENABLE_ANIM_DEBUG
 
 void FInteractionSearchContext::SetContinuingProperties(float InPlayingAssetAccumulatedTime, const UObject* InPlayingAsset,
 	bool bInIsPlayingAssetMirrored, const FVector& InPlayingAssetBlendParameters, EPoseSearchInterruptMode InInterruptMode)
diff --git a/Engine/Plugins/Animation/PoseSearch/Source/Runtime/Public/PoseSearch/PoseSearchInteractionIsland.h b/Engine/Plugins/Animation/PoseSearch/Source/Runtime/Public/PoseSearch/PoseSearchInteractionIsland.h
--- a/Engine/Plugins/Animation/PoseSearch/Source/Runtime/Public/PoseSearch/PoseSearchInteractionIsland.h
+++ b/Engine/Plugins/Animation/PoseSearch/Source/Runtime/Public/PoseSearch/PoseSearchInteractionIsland.h
@@ -93,8 +93,9 @@ struct FInteractionSearchContext : public FInteractionSearchContextBase
 	FPoseSearchContinuingProperties GetContinuingProperties() const;
 
 	bool IsContinuingInteraction() const { return bIsContinuingInteraction; }
+#if ENABLE_ANIM_DEBUG
 	float GetDebugBroadPhaseRadius(int32 Index) const;
-
+#endif // ENABLE_ANIM_DEBUG
 	void SetContinuingInteraction(bool bInIsContinuingInteraction) { bIsContinuingInteraction = bInIsContinuingInteraction; }
 	void AddReferencedObjects(FReferenceCollector& Collector);
 	void SetContinuingProperties(float InPlayingAssetAccumulatedTime, const UObject* InPlayingAsset,
