PoseSearch - removing some 5.6 deprecations

diff --git a/Engine/Plugins/Animation/PoseSearch/Source/Runtime/Private/AnimNode_MotionMatching.cpp b/Engine/Plugins/Animation/PoseSearch/Source/Runtime/Private/AnimNode_MotionMatching.cpp
--- a/Engine/Plugins/Animation/PoseSearch/Source/Runtime/Private/AnimNode_MotionMatching.cpp
+++ b/Engine/Plugins/Animation/PoseSearch/Source/Runtime/Private/AnimNode_MotionMatching.cpp
@@ -231,12 +231,6 @@ void FAnimNode_MotionMatching::UpdateAssetPlayer(const FAnimationUpdateContext&
 				// const casting here is safe since we're in the thread owning the pose history, and it's the correct place to update the previously selected poses
 				const_cast<FPoseIndicesHistory*>(PoseIndicesHistory)->Update(SearchResult, DeltaTime, PoseReselectHistory);
 			}
-
-			PRAGMA_DISABLE_DEPRECATION_WARNINGS
-			MotionMatchingState.CurrentSearchResult = SearchResult;
-			MotionMatchingState.WantedPlayRate = WantedPlayRate;
-			MotionMatchingState.bJumpedToPose = MotionMatchingState.SearchResult.SelectedAnim && !MotionMatchingState.SearchResult.bIsContinuingPoseSearch;
-			PRAGMA_ENABLE_DEPRECATION_WARNINGS
 		}
 
 		float DesiredPlayRate = MotionMatchingState.SearchResult.WantedPlayRate * PlayRateMultiplier;
diff --git a/Engine/Plugins/Animation/PoseSearch/Source/Runtime/Private/PoseSearchAssetSampler.cpp b/Engine/Plugins/Animation/PoseSearch/Source/Runtime/Private/PoseSearchAssetSampler.cpp
--- a/Engine/Plugins/Animation/PoseSearch/Source/Runtime/Private/PoseSearchAssetSampler.cpp
+++ b/Engine/Plugins/Animation/PoseSearch/Source/Runtime/Private/PoseSearchAssetSampler.cpp
@@ -636,19 +636,6 @@ void FAnimationAssetSampler::Process()
 	}
 }
 
-void FAnimationAssetSampler::ExtractPoseSearchNotifyStates(float Time, const TFunction<bool(UAnimNotifyState_PoseSearchBase*)>& ProcessPoseSearchBase) const
-{
-	FAnimNotifyContext PreAllocatedNotifyContext;
-	ExtractAnimNotifyStates(Time, PreAllocatedNotifyContext, [&ProcessPoseSearchBase](UAnimNotifyState* AnimNotifyState)
-		{
-			if (UAnimNotifyState_PoseSearchBase* PoseSearchAnimNotify = Cast<UAnimNotifyState_PoseSearchBase>(AnimNotifyState))
-			{
-				return ProcessPoseSearchBase(PoseSearchAnimNotify);
-			}
-			return true;
-		});
-}
-
 void FAnimationAssetSampler::ExtractAnimNotifyStates(float Time, FAnimNotifyContext& PreAllocatedNotifyContext, const TFunction<bool(UAnimNotifyState*)>& ProcessAnimNotifyState) const
 {
 	float SampleTime = Time;
diff --git a/Engine/Plugins/Animation/PoseSearch/Source/Runtime/Private/PoseSearchContext.cpp b/Engine/Plugins/Animation/PoseSearch/Source/Runtime/Private/PoseSearchContext.cpp
--- a/Engine/Plugins/Animation/PoseSearch/Source/Runtime/Private/PoseSearchContext.cpp
+++ b/Engine/Plugins/Animation/PoseSearch/Source/Runtime/Private/PoseSearchContext.cpp
@@ -189,24 +189,6 @@ static const UWorld* GetWorld(TConstArrayView<FChooserEvaluationContext*> AnimCo
 
 //////////////////////////////////////////////////////////////////////////
 // FDebugDrawParams
-PRAGMA_DISABLE_DEPRECATION_WARNINGS
-FDebugDrawParams::FDebugDrawParams(TArrayView<FAnimInstanceProxy*> InAnimInstanceProxies, TConstArrayView<const IPoseHistory*> InPoseHistories, const FRoleToIndex& InRoleToIndex, const UPoseSearchDatabase* InDatabase, EDebugDrawFlags InFlags)
-: AnimContexts()
-, PoseHistories(InPoseHistories)
-, RoleToIndex(InRoleToIndex)
-, Database(InDatabase)
-{
-}
-
-FDebugDrawParams::FDebugDrawParams(TArrayView<const USkinnedMeshComponent*> InMeshes, TConstArrayView<const IPoseHistory*> InPoseHistories, const FRoleToIndex& InRoleToIndex, const UPoseSearchDatabase* InDatabase, EDebugDrawFlags InFlags)
-: AnimContexts()
-, PoseHistories(InPoseHistories)
-, RoleToIndex(InRoleToIndex)
-, Database(InDatabase)
-{
-}
-PRAGMA_ENABLE_DEPRECATION_WARNINGS
-
 FDebugDrawParams::FDebugDrawParams(TConstArrayView<FChooserEvaluationContext*> InAnimContexts, TConstArrayView<const IPoseHistory*> InPoseHistories, const FRoleToIndex& InRoleToIndex, const UPoseSearchDatabase* InDatabase)
 : FDebugDrawParams()
 {
@@ -865,13 +847,6 @@ FSearchContext::FPoseCandidate FSearchContext::FBestPoseCandidates::GetUnsortedC
 
 //////////////////////////////////////////////////////////////////////////
 // FSearchContext
-FSearchContext::FSearchContext(float InDesiredPermutationTimeOffset, const FPoseIndicesHistory* InPoseIndicesHistory,
-	const FSearchResult& InCurrentResult, const FFloatInterval& InPoseJumpThresholdTime, bool bInUseCachedChannelData)
-: FSearchContext(InDesiredPermutationTimeOffset, InPoseJumpThresholdTime, FPoseSearchEvent())
-{
-	UpdateContinuingPoseSearchResult(InCurrentResult, InCurrentResult);
-}
-
 FSearchContext::FSearchContext(float InDesiredPermutationTimeOffset, const FPoseIndicesHistory* InPoseIndicesHistory,
 	const FSearchResult& InCurrentResult, const FFloatInterval& InPoseJumpThresholdTime, const FPoseSearchEvent& InEventToSearch)
 : FSearchContext(InDesiredPermutationTimeOffset, InPoseJumpThresholdTime, FPoseSearchEvent())
@@ -1107,16 +1082,6 @@ const IPoseHistory* FSearchContext::GetPoseHistory(const FRole& Role) const
 	return PoseHistories[RoleToIndex[Role]];
 }
 
-const UAnimInstance* FSearchContext::GetAnimInstance(const FRole& Role) const
-{
-	return Cast<UAnimInstance>(AnimContexts[RoleToIndex[Role]]->GetFirstObjectParam());
-}
-
-const UObject* FSearchContext::GetAnimContext(const FRole& Role) const
-{
-	return AnimContexts[RoleToIndex[Role]]->GetFirstObjectParam();
-}
-
 const FChooserEvaluationContext* FSearchContext::GetContext(const FRole& Role) const
 {
 	return AnimContexts[RoleToIndex[Role]];
diff --git a/Engine/Plugins/Animation/PoseSearch/Source/Runtime/Private/PoseSearchDatabase.cpp b/Engine/Plugins/Animation/PoseSearch/Source/Runtime/Private/PoseSearchDatabase.cpp
--- a/Engine/Plugins/Animation/PoseSearch/Source/Runtime/Private/PoseSearchDatabase.cpp
+++ b/Engine/Plugins/Animation/PoseSearch/Source/Runtime/Private/PoseSearchDatabase.cpp
@@ -398,19 +398,6 @@ float FPoseSearchDatabaseAnimationAssetBase::GetPlayLength(const FVector& BlendP
 }
 
 #if WITH_EDITOR
-int32 FPoseSearchDatabaseAnimationAssetBase::GetFrameAtTime(float Time) const
-{
-	if (const UAnimSequenceBase* SequenceBase = Cast<UAnimSequenceBase>(GetAnimationAsset()))
-	{
-		return SequenceBase->GetFrameAtTime(Time);
-	}
-
-	// estimating the frame for blend spaces in a non precise way. Anyways this method is deprecated and should not be used
-	const float RealAssetTime = GetPlayLength(FVector::ZeroVector);
-	const FFrameRate& DefaultFrameRate = UAnimationSettings::Get()->GetDefaultFrameRate();
-	return DefaultFrameRate.AsFrameTime(RealAssetTime).RoundToFrame().Value;
-}
-
 bool FPoseSearchDatabaseAnimationAssetBase::IsSkeletonCompatible(TObjectPtr<const UPoseSearchSchema> InSchema) const
 {
 	if (InSchema)
@@ -500,11 +487,6 @@ FFloatInterval FPoseSearchDatabaseAnimationAssetBase::GetEffectiveSamplingRange(
 	return GetEffectiveSamplingRange(GetPlayLength(BlendParameters), GetSamplingRange());
 }
 
-FFloatInterval FPoseSearchDatabaseAnimationAssetBase::GetEffectiveSamplingRange() const
-{
-	return GetEffectiveSamplingRange(GetPlayLength(FVector::ZeroVector), GetSamplingRange());
-}
-
 FFloatInterval FPoseSearchDatabaseAnimationAssetBase::GetEffectiveSamplingRange(float PlayLength, const FFloatInterval& SamplingRange)
 {
 	const bool bSampleAll = (SamplingRange.Min == 0.0f) && (SamplingRange.Max == 0.0f);
@@ -822,73 +804,6 @@ void FPoseSearchDatabaseBlendSpace::IterateOverSamplingParameter(const TFunction
 	}
 }
 
-void FPoseSearchDatabaseBlendSpace::GetBlendSpaceParameterSampleRanges(int32& HorizontalBlendNum, int32& VerticalBlendNum) const
-{
-	check(BlendSpace);
-
-	if (bUseSingleSample)
-	{
-		HorizontalBlendNum = 1;
-		VerticalBlendNum = 1;
-	}
-	else if (bUseGridForSampling)
-	{
-		HorizontalBlendNum = BlendSpace->GetBlendParameter(0).GridNum + 1;
-		VerticalBlendNum = BlendSpace->IsA<UBlendSpace1D>() ? 1 : BlendSpace->GetBlendParameter(1).GridNum + 1;
-	}
-	else
-	{
-		HorizontalBlendNum = FMath::Max(NumberOfHorizontalSamples, 1);
-		VerticalBlendNum = BlendSpace->IsA<UBlendSpace1D>() ? 1 : FMath::Max(NumberOfVerticalSamples, 1);
-	}
-
-	check(HorizontalBlendNum >= 1 && VerticalBlendNum >= 1);
-}
-
-FVector FPoseSearchDatabaseBlendSpace::BlendParameterForSampleRanges(int32 HorizontalBlendIndex, int32 VerticalBlendIndex) const
-{
-	check(BlendSpace);
-
-	if (bUseSingleSample)
-	{
-		check(HorizontalBlendIndex == 0 && VerticalBlendIndex == 0);
-		return FVector(BlendParamX, BlendParamY, 0.f);
-	}
-	
-	const bool bWrapInputOnHorizontalAxis = BlendSpace->GetBlendParameter(0).bWrapInput;
-	const bool bWrapInputOnVerticalAxis = BlendSpace->GetBlendParameter(1).bWrapInput;
-
-	int32 HorizontalBlendNum, VerticalBlendNum;
-	GetBlendSpaceParameterSampleRanges(HorizontalBlendNum, VerticalBlendNum);
-
-	if (bWrapInputOnHorizontalAxis)
-	{
-		++HorizontalBlendNum;
-	}
-
-	if (bWrapInputOnVerticalAxis)
-	{
-		++VerticalBlendNum;
-	}
-
-	const float HorizontalBlendMin = BlendSpace->GetBlendParameter(0).Min;
-	const float HorizontalBlendMax = BlendSpace->GetBlendParameter(0).Max;
-
-	const float VerticalBlendMin = BlendSpace->GetBlendParameter(1).Min;
-	const float VerticalBlendMax = BlendSpace->GetBlendParameter(1).Max;
-
-	return FVector(
-		HorizontalBlendNum > 1 ? 
-			HorizontalBlendMin + (HorizontalBlendMax - HorizontalBlendMin) * 
-			((float)HorizontalBlendIndex) / (HorizontalBlendNum - 1) : 
-		HorizontalBlendMin,
-		VerticalBlendNum > 1 ? 
-			VerticalBlendMin + (VerticalBlendMax - VerticalBlendMin) * 
-			((float)VerticalBlendIndex) / (VerticalBlendNum - 1) : 
-		VerticalBlendMin,
-		0.f);
-}
-
 #endif // WITH_EDITORONLY_DATA
 
 //////////////////////////////////////////////////////////////////////////
@@ -1136,11 +1051,6 @@ const UE::PoseSearch::FSearchIndex& UPoseSearchDatabase::GetSearchIndex() const
 	return SearchIndexPrivate;
 }
 
-int32 UPoseSearchDatabase::GetPoseIndexFromTime(float RealTimeInSeconds, const UE::PoseSearch::FSearchIndexAsset& SearchIndexAsset) const
-{
-	return SearchIndexAsset.GetPoseIndexFromTime(RealTimeInSeconds, Schema->SampleRate);
-}
-
 int32 UPoseSearchDatabase::GetPoseIndex(const UObject* AnimationAsset, float AnimationAssetTime, bool bMirrored, const FVector& BlendParameters) const
 {
 	using namespace UE::PoseSearch;
diff --git a/Engine/Plugins/Animation/PoseSearch/Source/Runtime/Private/PoseSearchDerivedData.cpp b/Engine/Plugins/Animation/PoseSearch/Source/Runtime/Private/PoseSearchDerivedData.cpp
--- a/Engine/Plugins/Animation/PoseSearch/Source/Runtime/Private/PoseSearchDerivedData.cpp
+++ b/Engine/Plugins/Animation/PoseSearch/Source/Runtime/Private/PoseSearchDerivedData.cpp
@@ -564,9 +564,6 @@ static Eigen::ComputationInfo PreprocessSearchIndexPCAData(FSearchIndex& SearchI
 #if WITH_EDITORONLY_DATA
 	SearchIndex.PCAExplainedVarianceEditorOnly = 0.f;
 #endif
-	PRAGMA_DISABLE_DEPRECATION_WARNINGS
-	SearchIndex.PCAExplainedVariance = 0.f;
-	PRAGMA_ENABLE_DEPRECATION_WARNINGS
 
 	SearchIndex.PCAValues.Reset();
 	SearchIndex.Mean.Reset();
@@ -648,17 +645,12 @@ static Eigen::ComputationInfo PreprocessSearchIndexPCAData(FSearchIndex& SearchI
 				AccumulatedVariance += EigenValues[Indexer[PCAComponentIndex]];
 			}
 
-			PRAGMA_DISABLE_DEPRECATION_WARNINGS
-			// @todo: move this code under WITH_EDITORONLY_DATA once SearchIndex.PCAExplainedVariance has been removed
+#if WITH_EDITORONLY_DATA
 			// calculating the total variance knowing that eigen values measure variance along the principal components:
 			const float TotalVariance = EigenValues.sum();
 			// and explained variance as ratio between AccumulatedVariance and TotalVariance: https://ro-che.info/articles/2017-12-11-pca-explained-variance
-			SearchIndex.PCAExplainedVariance = TotalVariance > UE_KINDA_SMALL_NUMBER ? AccumulatedVariance / TotalVariance : 0.f;
-
-#if WITH_EDITORONLY_DATA
-			SearchIndex.PCAExplainedVarianceEditorOnly = SearchIndex.PCAExplainedVariance;
+			SearchIndex.PCAExplainedVarianceEditorOnly = TotalVariance > UE_KINDA_SMALL_NUMBER ? AccumulatedVariance / TotalVariance : 0.f;
 #endif // WITH_EDITORONLY_DATA
-			PRAGMA_ENABLE_DEPRECATION_WARNINGS
 
 			MapPCAValues = CenteredValues * PCAProjectionMatrix;
 
@@ -1396,13 +1388,6 @@ static void CompareSearchIndex(const FSearchIndex& A, const FSearchIndex& B, con
 		StringBuilder.Append(TEXT("PCAExplainedVarianceEditorOnly mismatch\n"));
 	}
 #endif // WITH_EDITORONLY_DATA
-
-	PRAGMA_DISABLE_DEPRECATION_WARNINGS
-	if (A.PCAExplainedVariance != B.PCAExplainedVariance)
-	{
-		StringBuilder.Append(TEXT("PCAExplainedVariance mismatch\n"));
-	}
-	PRAGMA_ENABLE_DEPRECATION_WARNINGS
 }
 
 #endif // ENABLE_ANIM_DEBUG
diff --git a/Engine/Plugins/Animation/PoseSearch/Source/Runtime/Private/PoseSearchHistory.cpp b/Engine/Plugins/Animation/PoseSearch/Source/Runtime/Private/PoseSearchHistory.cpp
--- a/Engine/Plugins/Animation/PoseSearch/Source/Runtime/Private/PoseSearchHistory.cpp
+++ b/Engine/Plugins/Animation/PoseSearch/Source/Runtime/Private/PoseSearchHistory.cpp
@@ -876,13 +876,6 @@ const TConstArrayView<FName> FPoseHistory::GetCollectedCurves() const
 	return MakeConstArrayView(CollectedCurves);
 }
 
-PRAGMA_DISABLE_DEPRECATION_WARNINGS
-void FPoseHistory::SetTrajectory(const FPoseSearchQueryTrajectory& InTrajectory, float InTrajectorySpeedMultiplier)
-{
-	SetTrajectory(FTransformTrajectory(InTrajectory), InTrajectorySpeedMultiplier);
-}
-PRAGMA_ENABLE_DEPRECATION_WARNINGS
-	
 int32 FPoseHistory::GetNumEntries() const
 {
 	UE_MT_SCOPED_READ_ACCESS(PoseDataThreadSafeCounter);
@@ -1047,16 +1040,6 @@ void FPoseHistory::EvaluateComponentSpace_AnyThread(float DeltaTime, IComponentS
 	}
 }
 
-void FPoseHistory::EvaluateComponentSpace_AnyThread(float DeltaTime, FCSPose<FCompactPose>& ComponentSpacePose, bool bStoreScales,
-	float RootBoneRecoveryTime, float RootBoneTranslationRecoveryRatio, float RootBoneRotationRecoveryRatio,
-	bool bNeedsReset, bool bCacheBones, const TArray<FBoneIndexType>& RequiredBones,
-	const FBlendedCurve& Curves, const TConstArrayView<FName>& InCollectedCurves)
-{
-	FComponentSpacePoseProvider ComponentSpacePoseProvider(ComponentSpacePose);
-	EvaluateComponentSpace_AnyThread(DeltaTime, ComponentSpacePoseProvider, bStoreScales, RootBoneRecoveryTime, RootBoneTranslationRecoveryRatio,
-		RootBoneRotationRecoveryRatio, bNeedsReset, bCacheBones, RequiredBones, Curves, InCollectedCurves);
-}
-
 #if ENABLE_DRAW_DEBUG && ENABLE_ANIM_DEBUG
 void FPoseHistory::DebugDraw(FAnimInstanceProxy& AnimInstanceProxy, FColor Color) const
 {
@@ -1293,12 +1276,6 @@ void FMemStackPoseHistory::ExtractAndAddFuturePoses(const UAnimationAsset* Anima
 	}
 }
 
-void FMemStackPoseHistory::AddFuturePose(float Time, FCSPose<FCompactPose>& ComponentSpacePose, const FBlendedCurve& Curves)
-{
-	FComponentSpacePoseProvider ComponentSpacePoseProvider(ComponentSpacePose);
-	AddFuturePose(Time, ComponentSpacePoseProvider, Curves);
-}
-
 int32 FMemStackPoseHistory::GetNumEntries() const
 {
 	check(PoseHistory);
diff --git a/Engine/Plugins/Animation/PoseSearch/Source/Runtime/Private/PoseSearchIndex.cpp b/Engine/Plugins/Animation/PoseSearch/Source/Runtime/Private/PoseSearchIndex.cpp
--- a/Engine/Plugins/Animation/PoseSearch/Source/Runtime/Private/PoseSearchIndex.cpp
+++ b/Engine/Plugins/Animation/PoseSearch/Source/Runtime/Private/PoseSearchIndex.cpp
@@ -225,6 +225,7 @@ FArchive& operator<<(FArchive& Ar, FSearchIndexAsset& IndexAsset)
 	return Ar;
 }
 
+#if WITH_EDITORONLY_DATA
 //////////////////////////////////////////////////////////////////////////
 // FSearchStats
 void FSearchStats::Reset()
@@ -252,6 +253,7 @@ FArchive& operator<<(FArchive& Ar, FSearchStats& Stats)
 	Ar << Stats.MaxAcceleration;
 	return Ar;
 }
+#endif // WITH_EDITORONLY_DATA
 
 //////////////////////////////////////////////////////////////////////////
 // FEventDataCollector
@@ -393,9 +395,9 @@ FSearchIndexBase::FSearchIndexBase()
 	, Assets()
 	, EventData()
 	, MinCostAddend(-MAX_FLT)
-	PRAGMA_DISABLE_DEPRECATION_WARNINGS
+#if WITH_EDITORONLY_DATA
 	, Stats()
-	PRAGMA_ENABLE_DEPRECATION_WARNINGS
+#endif // WITH_EDITORONLY_DATA
 {
 }
 
@@ -411,9 +413,9 @@ FSearchIndexBase::FSearchIndexBase(const FSearchIndexBase& Other)
 	, Assets(Other.Assets)
 	, EventData(Other.EventData)
 	, MinCostAddend(Other.MinCostAddend)
-	PRAGMA_DISABLE_DEPRECATION_WARNINGS
+#if WITH_EDITORONLY_DATA
 	, Stats(Other.Stats)
-	PRAGMA_ENABLE_DEPRECATION_WARNINGS
+#endif // WITH_EDITORONLY_DATA
 {
 }
 
@@ -425,9 +427,9 @@ FSearchIndexBase::FSearchIndexBase(FSearchIndexBase&& Other)
 	, Assets(MoveTemp(Other.Assets))
 	, EventData(MoveTemp(Other.EventData))
 	, MinCostAddend(MoveTemp(Other.MinCostAddend))
-	PRAGMA_DISABLE_DEPRECATION_WARNINGS
+#if WITH_EDITORONLY_DATA
 	, Stats(MoveTemp(Other.Stats))
-	PRAGMA_ENABLE_DEPRECATION_WARNINGS
+#endif // WITH_EDITORONLY_DATA
 {
 }
 
@@ -485,9 +487,9 @@ void FSearchIndexBase::Reset()
 	Assets.Reset();
 	EventData.Reset();
 	MinCostAddend = -MAX_FLT;
-	PRAGMA_DISABLE_DEPRECATION_WARNINGS
+#if WITH_EDITORONLY_DATA
 	Stats.Reset();
-	PRAGMA_ENABLE_DEPRECATION_WARNINGS
+#endif // WITH_EDITORONLY_DATA
 }
 
 void FSearchIndexBase::PruneDuplicateValues(float SimilarityThreshold, int32 DataCardinality, bool bDoNotGenerateValuesVectorToPoseIndexes)
@@ -563,17 +565,17 @@ void FSearchIndexBase::AllocateData(int32 DataCardinality, int32 NumPoses)
 
 bool FSearchIndexBase::operator==(const FSearchIndexBase& Other) const
 {
-	return
-		Values == Other.Values &&
-		ValuesVectorToPoseIndexes == Other.ValuesVectorToPoseIndexes &&
-		PoseMetadata == Other.PoseMetadata &&
-		bAnyBlockTransition == Other.bAnyBlockTransition &&
-		Assets == Other.Assets &&
-		EventData == Other.EventData &&
-		MinCostAddend == Other.MinCostAddend &&
-		PRAGMA_DISABLE_DEPRECATION_WARNINGS
-		Stats == Other.Stats;
-		PRAGMA_ENABLE_DEPRECATION_WARNINGS
+	return Values == Other.Values
+		&& ValuesVectorToPoseIndexes == Other.ValuesVectorToPoseIndexes
+		&& PoseMetadata == Other.PoseMetadata
+		&& bAnyBlockTransition == Other.bAnyBlockTransition
+		&& Assets == Other.Assets
+		&& EventData == Other.EventData
+		&& MinCostAddend == Other.MinCostAddend
+#if WITH_EDITORONLY_DATA
+		&& Stats == Other.Stats
+#endif // WITH_EDITORONLY_DATA
+		;
 }
 
 FArchive& operator<<(FArchive& Ar, FSearchIndexBase& Index)
@@ -585,9 +587,14 @@ FArchive& operator<<(FArchive& Ar, FSearchIndexBase& Index)
 	Ar << Index.Assets;
 	Ar << Index.EventData;
 	Ar << Index.MinCostAddend;
-	PRAGMA_DISABLE_DEPRECATION_WARNINGS
-	Ar << Index.Stats;
-	PRAGMA_ENABLE_DEPRECATION_WARNINGS
+	
+#if WITH_EDITORONLY_DATA
+	if (!Ar.IsFilterEditorOnly())
+	{
+		Ar << Index.Stats;
+	}
+#endif // WITH_EDITORONLY_DATA
+
 	return Ar;
 }
 
@@ -606,9 +613,6 @@ FSearchIndex::FSearchIndex(const FSearchIndex& Other)
 	, DeviationEditorOnly(Other.DeviationEditorOnly)
 	, PCAExplainedVarianceEditorOnly(Other.PCAExplainedVarianceEditorOnly)
 	#endif
-	PRAGMA_DISABLE_DEPRECATION_WARNINGS
-	, PCAExplainedVariance(Other.PCAExplainedVariance)
-	PRAGMA_ENABLE_DEPRECATION_WARNINGS
 {
 	check(!PCAValues.IsEmpty() || KDTree.DataSource.PointCount == 0);
 	KDTree.DataSource.Data = PCAValues.IsEmpty() ? nullptr : PCAValues.GetData();
@@ -627,9 +631,6 @@ FSearchIndex::FSearchIndex(FSearchIndex&& Other)
 	, DeviationEditorOnly(MoveTemp(Other.DeviationEditorOnly))
 	, PCAExplainedVarianceEditorOnly(MoveTemp(Other.PCAExplainedVarianceEditorOnly))
 	#endif
-	PRAGMA_DISABLE_DEPRECATION_WARNINGS
-	, PCAExplainedVariance(MoveTemp(Other.PCAExplainedVariance))
-	PRAGMA_ENABLE_DEPRECATION_WARNINGS
 {
 	check(!PCAValues.IsEmpty() || KDTree.DataSource.PointCount == 0);
 	KDTree.DataSource.Data = PCAValues.IsEmpty() ? nullptr : PCAValues.GetData();
@@ -869,25 +870,6 @@ void FSearchIndex::PruneDuplicatePCAValues(float SimilarityThreshold, int32 Numb
 	}
 }
 
-TArray<float> FSearchIndex::GetPoseValuesSafe(int32 PoseIdx) const
-{
-	TArray<float> PoseValues;
-	if (PoseIdx >= 0 && PoseIdx < GetNumPoses())
-	{
-		if (IsValuesEmpty())
-		{
-			const int32 NumDimensions = GetNumDimensions();
-			PoseValues.SetNumUninitialized(NumDimensions);
-			GetReconstructedPoseValues(PoseIdx, PoseValues);
-		}
-		else
-		{
-			PoseValues = GetPoseValues(PoseIdx);
-		}
-	}
-	return PoseValues;
-}
-
 TConstArrayView<float> FSearchIndex::GetPoseValuesSafe(int32 PoseIdx, TArray<float>& BufferUsedForReconstruction) const
 {
 	if (PoseIdx < 0 || PoseIdx >= GetNumPoses())
@@ -925,16 +907,6 @@ TConstArrayView<float> FSearchIndex::GetPCAPoseValues(int32 PCAValuesVectorIdx)
 	return MakeArrayView(&PCAValues[ValueOffset], NumberOfPrincipalComponents);
 }
 
-FPoseSearchCost FSearchIndex::ComparePoses(int32 PoseIdx, float ContinuingPoseCostBias, TConstArrayView<float> PoseValues, TConstArrayView<float> QueryValues) const
-{
-	return FPoseSearchCost(CompareFeatureVectors<false>(PoseValues, QueryValues, WeightsSqrt), PoseMetadata[PoseIdx].GetCostAddend(), ContinuingPoseCostBias, 0.f);
-}
-
-FPoseSearchCost FSearchIndex::CompareAlignedPoses(int32 PoseIdx, float ContinuingPoseCostBias, TConstArrayView<float> PoseValues, TConstArrayView<float> QueryValues) const
-{
-	return FPoseSearchCost(CompareFeatureVectors<true>(PoseValues, QueryValues, WeightsSqrt), PoseMetadata[PoseIdx].GetCostAddend(), ContinuingPoseCostBias, 0.f);
-}
-
 void FSearchIndex::GetPoseToPCAValuesVectorIndexes(TArray<uint32>& PoseToPCAValuesVectorIndexes) const
 {
 	if (PCAValuesVectorToPoseIndexes.Num() > 0)
@@ -985,9 +957,6 @@ bool FSearchIndex::operator==(const FSearchIndex& Other) const
 		&& Mean == Other.Mean
 		&& KDTree == Other.KDTree
 		&& VPTree == Other.VPTree
-		PRAGMA_DISABLE_DEPRECATION_WARNINGS
-		&& PCAExplainedVariance == Other.PCAExplainedVariance
-		PRAGMA_ENABLE_DEPRECATION_WARNINGS
 #if WITH_EDITORONLY_DATA
 		&& DeviationEditorOnly == Other.DeviationEditorOnly
 		&& PCAExplainedVarianceEditorOnly == Other.PCAExplainedVarianceEditorOnly
@@ -1008,10 +977,6 @@ FArchive& operator<<(FArchive& Ar, FSearchIndex& Index)
 
 	Serialize(Ar, Index.KDTree, Index.PCAValues.GetData());
 
-	PRAGMA_DISABLE_DEPRECATION_WARNINGS
-	Ar << Index.PCAExplainedVariance;
-	PRAGMA_ENABLE_DEPRECATION_WARNINGS
-
 #if WITH_EDITORONLY_DATA
 	if (!Ar.IsFilterEditorOnly())
 	{
diff --git a/Engine/Plugins/Animation/PoseSearch/Source/Runtime/Private/PoseSearchLibrary.cpp b/Engine/Plugins/Animation/PoseSearch/Source/Runtime/Private/PoseSearchLibrary.cpp
--- a/Engine/Plugins/Animation/PoseSearch/Source/Runtime/Private/PoseSearchLibrary.cpp
+++ b/Engine/Plugins/Animation/PoseSearch/Source/Runtime/Private/PoseSearchLibrary.cpp
@@ -574,30 +574,13 @@ namespace UE::PoseSearch
 //////////////////////////////////////////////////////////////////////////
 // FMotionMatchingState
 
-void FMotionMatchingState::Reset(const FTransform& ComponentTransform)
-{
-	Reset();
-}
-
 void FMotionMatchingState::Reset()
 {
-PRAGMA_DISABLE_DEPRECATION_WARNINGS
-	bJumpedToPose = false;
-	WantedPlayRate = 1.f;
-PRAGMA_ENABLE_DEPRECATION_WARNINGS
-
 	SearchResult = FPoseSearchBlueprintResult();
 	// Set the elapsed time to INFINITY to trigger a search right away
 	ElapsedPoseSearchTime = std::numeric_limits<float>::infinity();
 }
 
-void FMotionMatchingState::AdjustAssetTime(float AssetTime)
-{
-PRAGMA_DISABLE_DEPRECATION_WARNINGS
-	CurrentSearchResult.UpdateWithNormalizedTime(AssetTime);
-PRAGMA_ENABLE_DEPRECATION_WARNINGS
-}
-
 FVector FMotionMatchingState::GetEstimatedFutureRootMotionVelocity() const
 {
 	using namespace UE::PoseSearch;
@@ -618,13 +601,6 @@ FVector FMotionMatchingState::GetEstimatedFutureRootMotionVelocity() const
 	return FVector::ZeroVector;
 }
 
-void FMotionMatchingState::UpdateWantedPlayRate(const UE::PoseSearch::FSearchContext& SearchContext, const FFloatInterval& PlayRate, float TrajectorySpeedMultiplier, const FPoseSearchEvent& EventToSearch)
-{
-PRAGMA_DISABLE_DEPRECATION_WARNINGS
-	WantedPlayRate = CalculateWantedPlayRate(CurrentSearchResult, SearchContext, PlayRate, TrajectorySpeedMultiplier, EventToSearch);
-PRAGMA_ENABLE_DEPRECATION_WARNINGS
-}
-
 //////////////////////////////////////////////////////////////////////////
 // FPoseSearchContinuingProperties
 
@@ -894,80 +870,6 @@ void UPoseSearchLibrary::TraceMotionMatching(
 }
 #endif // UE_POSE_SEARCH_TRACE_ENABLED
 
-void UPoseSearchLibrary::UpdateMotionMatchingState(
-	const FAnimationUpdateContext& Context,
-	const TArray<TObjectPtr<const UPoseSearchDatabase>>& Databases,
-	float BlendTime,
-	int32 MaxActiveBlends,
-	const FFloatInterval& PoseJumpThresholdTime,
-	float PoseReselectHistory,
-	float SearchThrottleTime,
-	const FFloatInterval& PlayRate,
-	FMotionMatchingState& InOutMotionMatchingState,
-	EPoseSearchInterruptMode InterruptMode,
-	bool bShouldSearch,
-	bool bShouldUseCachedChannelData,
-	bool bDebugDrawQuery,
-	bool bDebugDrawCurResult)
-{
-	using namespace UE::PoseSearch;
-
-	if (Databases.IsEmpty())
-	{
-		Context.LogMessage(
-			EMessageSeverity::Error,
-			LOCTEXT("NoDatabases", "No database assets provided for motion matching."));
-		return;
-	}
-
-	const IPoseHistory* PoseHistory = nullptr;
-	if (const FPoseHistoryProvider* PoseHistoryProvider = Context.GetMessage<FPoseHistoryProvider>())
-	{
-		PoseHistory = &PoseHistoryProvider->GetPoseHistory();
-	}
-
-	check(Context.AnimInstanceProxy);
-	FChooserEvaluationContext AnimContext(Context.AnimInstanceProxy->GetAnimInstanceObject());
-	PRAGMA_DISABLE_DEPRECATION_WARNINGS
-	UpdateMotionMatchingState(&AnimContext, PoseHistory, Databases, Context.GetDeltaTime(),
-		PoseJumpThresholdTime, PoseReselectHistory, bShouldSearch ? SearchThrottleTime : UE_BIG_NUMBER, PlayRate, InOutMotionMatchingState,
-		InterruptMode, bShouldUseCachedChannelData, bDebugDrawQuery, bDebugDrawCurResult);
-	PRAGMA_ENABLE_DEPRECATION_WARNINGS
-}
-
-void UPoseSearchLibrary::UpdateMotionMatchingState(
-	const UObject* AnimContext,
-	const UE::PoseSearch::IPoseHistory* PoseHistory,
-	const TConstArrayView<TObjectPtr<const UPoseSearchDatabase>> Databases,
-	float DeltaTime,
-	const FFloatInterval& PoseJumpThresholdTime,
-	float PoseReselectHistory,
-	float SearchThrottleTime,
-	const FFloatInterval& PlayRate,
-	FMotionMatchingState& InOutMotionMatchingState,
-	EPoseSearchInterruptMode InterruptMode,
-	bool bShouldUseCachedChannelData,
-	bool bDebugDrawQuery,
-	bool bDebugDrawCurResult,
-	const FPoseSearchEvent& EventToSearch)
-{
-	FChooserEvaluationContext Context(const_cast<UObject*>(AnimContext));
-	PRAGMA_DISABLE_DEPRECATION_WARNINGS
-	UpdateMotionMatchingState(&Context,
-		PoseHistory,
-		Databases,
-		DeltaTime,
-		PoseJumpThresholdTime,
-		PoseReselectHistory,
-		SearchThrottleTime,
-		PlayRate,
-		InOutMotionMatchingState,
-		InterruptMode,
-		bShouldUseCachedChannelData,
-		bDebugDrawQuery, bDebugDrawCurResult);
-	PRAGMA_ENABLE_DEPRECATION_WARNINGS
-}
-
 void UPoseSearchLibrary::UpdateMotionMatchingState(
 	FChooserEvaluationContext* AnimContext,
 	const UE::PoseSearch::IPoseHistory* PoseHistory,
@@ -1110,11 +1012,6 @@ void UPoseSearchLibrary::UpdateMotionMatchingState(
 
 	InOutMotionMatchingState.SearchResult.InitFrom(SearchResult, WantedPlayRate);
 
-PRAGMA_DISABLE_DEPRECATION_WARNINGS
-	InOutMotionMatchingState.WantedPlayRate = WantedPlayRate;
-	InOutMotionMatchingState.CurrentSearchResult = SearchResult;
-PRAGMA_ENABLE_DEPRECATION_WARNINGS
-
 #if UE_POSE_SEARCH_TRACE_ENABLED
 	TraceMotionMatching(SearchContext, SearchResult, InOutMotionMatchingState.ElapsedPoseSearchTime, DeltaTime, bSearch, InOutMotionMatchingState.SearchResult.WantedPlayRate, InterruptMode);
 #endif // UE_POSE_SEARCH_TRACE_ENABLED
@@ -1183,78 +1080,44 @@ void UPoseSearchLibrary::MotionMatch(
 	const FPoseSearchFutureProperties Future,
 	FPoseSearchBlueprintResult& Result)
 {
-	using namespace UE::PoseSearch;
-
-	FMemMark Mark(FMemStack::Get());
-
-	TArray<UAnimInstance*, TInlineAllocator<PreallocatedRolesNum, TMemStackAllocator<>>> AnimInstances;
-	AnimInstances.Add(AnimInstance);
-
-	TArray<FName, TInlineAllocator<PreallocatedRolesNum, TMemStackAllocator<>>> Roles;
-	Roles.Add(DefaultRole);
-
-	TArray<const UObject*>& AssetsToSearchConst = reinterpret_cast<TArray<const UObject*>&>(AssetsToSearch);
-	PRAGMA_DISABLE_DEPRECATION_WARNINGS
-	MotionMatch(AnimInstances, Roles, AssetsToSearchConst, PoseHistoryName, ContinuingProperties, Future, Result);
-	PRAGMA_ENABLE_DEPRECATION_WARNINGS
-}
-
-void UPoseSearchLibrary::MotionMatch(
-	const TArrayView<UAnimInstance*> AnimInstances,
-	const TArrayView<const UE::PoseSearch::FRole> Roles,
-	const TArrayView<const UObject*> AssetsToSearch,
-	const FName PoseHistoryName,
-	const FPoseSearchContinuingProperties& ContinuingProperties,
-	const FPoseSearchFutureProperties& Future,
-	FPoseSearchBlueprintResult& Result)
-{
-	using namespace UE::Anim;
+	//using namespace UE::Anim;
 	using namespace UE::PoseSearch;
 
 	Result = FPoseSearchBlueprintResult();
 
-	if (AnimInstances.IsEmpty() || AnimInstances.Num() != Roles.Num())
+	if (!AnimInstance)
 	{
-		UE_LOG(LogPoseSearch, Error, TEXT("UPoseSearchLibrary::MotionMatch - invalid input AnimInstances or Roles"));
+		UE_LOG(LogPoseSearch, Error, TEXT("UPoseSearchLibrary::MotionMatch - null AnimInstances"));
 		return;
 	}
-	
-	for (UAnimInstance* AnimInstance : AnimInstances)
-	{
-		if (!AnimInstance)
-		{
-			UE_LOG(LogPoseSearch, Error, TEXT("UPoseSearchLibrary::MotionMatch - null AnimInstances"));
-			return;
-		}
 
-		if (!AnimInstance->CurrentSkeleton)
-		{
-			UE_LOG(LogPoseSearch, Error, TEXT("UPoseSearchLibrary::MotionMatch - null AnimInstances->CurrentSkeleton"));
-			return;
-		}
+	if (!AnimInstance->CurrentSkeleton)
+	{
+		UE_LOG(LogPoseSearch, Error, TEXT("UPoseSearchLibrary::MotionMatch - null AnimInstances->CurrentSkeleton"));
+		return;
 	}
 
 	FMemMark Mark(FMemStack::Get());
 
 	TArray<const IPoseHistory*, TInlineAllocator<PreallocatedRolesNum, TMemStackAllocator<>>> PoseHistories;
-	TArray<const UObject*, TInlineAllocator<PreallocatedRolesNum, TMemStackAllocator<>>> AnimContexts;
-	for (UAnimInstance* AnimInstance : AnimInstances)
+	if (const FAnimNode_PoseSearchHistoryCollector_Base* PoseHistoryNode = FindPoseHistoryNode(PoseHistoryName, AnimInstance))
 	{
-		if (const FAnimNode_PoseSearchHistoryCollector_Base* PoseHistoryNode = FindPoseHistoryNode(PoseHistoryName, AnimInstance))
-		{
-			PoseHistories.Add(&PoseHistoryNode->GetPoseHistory());
-		}
-
-		AnimContexts.Add(AnimInstance);
+		PoseHistories.Add(&PoseHistoryNode->GetPoseHistory());
 	}
-
-	if (PoseHistories.Num() != AnimInstances.Num())
+	else
 	{
 		UE_LOG(LogPoseSearch, Error, TEXT("UPoseSearchLibrary::MotionMatch - Couldn't find pose history with name '%s'"), *PoseHistoryName.ToString());
 		return;
 	}
 
-	const FSearchResult SearchResult = MotionMatch(AnimContexts, Roles, PoseHistories, AssetsToSearch, ContinuingProperties, Future, FPoseSearchEvent());
+	TArray<const UObject*, TInlineAllocator<PreallocatedRolesNum, TMemStackAllocator<>>> AnimContexts;
+	AnimContexts.Add(AnimInstance);
+
+	TArray<FName, TInlineAllocator<PreallocatedRolesNum, TMemStackAllocator<>>> Roles;
+	Roles.Add(DefaultRole);
+
+	TArray<const UObject*>& AssetsToSearchConst = reinterpret_cast<TArray<const UObject*>&>(AssetsToSearch);
+	const FSearchResult SearchResult = MotionMatch(AnimContexts, Roles, PoseHistories, AssetsToSearchConst, ContinuingProperties, Future, FPoseSearchEvent());
 	if (SearchResult.IsValid())
 	{
 		const UPoseSearchDatabase* Database = SearchResult.Database.Get();
diff --git a/Engine/Plugins/Animation/PoseSearch/Source/Runtime/Private/PoseSearchResult.cpp b/Engine/Plugins/Animation/PoseSearch/Source/Runtime/Private/PoseSearchResult.cpp
--- a/Engine/Plugins/Animation/PoseSearch/Source/Runtime/Private/PoseSearchResult.cpp
+++ b/Engine/Plugins/Animation/PoseSearch/Source/Runtime/Private/PoseSearchResult.cpp
@@ -494,10 +494,6 @@ bool FPoseSearchBlueprintResult::InitFrom(const UE::PoseSearch::FSearchResult& S
 		const FPoseSearchDatabaseAnimationAssetBase* DatabaseAsset = Database->GetDatabaseAnimationAsset(*SearchIndexAsset);
 		check(DatabaseAsset);
 		
-#if WITH_EDITORONLY_DATA
-		SelectedAnimation_DEPRECATED = DatabaseAsset->GetAnimationAsset();
-#endif // WITH_EDITORONLY_DATA
-
 		SelectedAnim = DatabaseAsset->GetAnimationAsset();
 		SelectedTime = SearchResult.GetAssetTime();
 		bIsContinuingPoseSearch = SearchResult.bIsContinuingPoseSearch;
@@ -538,9 +534,6 @@ UAnimationAsset* FPoseSearchBlueprintResult::GetAnimationAssetForRole()
 bool FPoseSearchBlueprintResult::operator==(const FPoseSearchBlueprintResult& Other) const
 {
 	return	
-#if WITH_EDITORONLY_DATA
-		SelectedAnimation_DEPRECATED == Other.SelectedAnimation_DEPRECATED &&
-#endif // WITH_EDITORONLY_DATA
 		SelectedAnim == Other.SelectedAnim &&
 		SelectedTime == Other.SelectedTime &&
 		bIsContinuingPoseSearch == Other.bIsContinuingPoseSearch &&
diff --git a/Engine/Plugins/Animation/PoseSearch/Source/Runtime/Private/PoseSearchSchema.cpp b/Engine/Plugins/Animation/PoseSearch/Source/Runtime/Private/PoseSearchSchema.cpp
--- a/Engine/Plugins/Animation/PoseSearch/Source/Runtime/Private/PoseSearchSchema.cpp
+++ b/Engine/Plugins/Animation/PoseSearch/Source/Runtime/Private/PoseSearchSchema.cpp
@@ -83,35 +83,6 @@ void UPoseSearchSchema::AddDefaultChannels()
 	Finalize();
 }
 
-void UPoseSearchSchema::InitBoneContainersFromRoledSkeleton(TMap<FName, FBoneContainer>& RoledBoneContainers) const
-{
-	RoledBoneContainers.Reset();
-	RoledBoneContainers.Reserve(Skeletons.Num());
-
-	for (const FPoseSearchRoledSkeleton& RoledSkeleton : Skeletons)
-	{
-		FBoneContainer& RoledBoneContainer = RoledBoneContainers.Add(RoledSkeleton.Role);
-		// Add a curve filter to our bone container to only eval curves actually used by the schema.
-		const UE::Anim::FCurveFilterSettings CurveFilterSettings(UE::Anim::ECurveFilterMode::AllowOnlyFiltered, &RoledSkeleton.RequiredCurves);
-
-		PRAGMA_DISABLE_DEPRECATION_WARNINGS
-		RoledBoneContainer.InitializeTo(RoledSkeleton.BoneIndicesWithParents_DEPRECATED, CurveFilterSettings, *RoledSkeleton.Skeleton);
-		PRAGMA_ENABLE_DEPRECATION_WARNINGS
-	}
-}
-
-bool UPoseSearchSchema::AllRoledSkeletonHaveMirrorDataTable() const
-{
-	for (const FPoseSearchRoledSkeleton& RoledSkeleton : Skeletons)
-	{
-		if (!RoledSkeleton.MirrorDataTable)
-		{
-			return false;
-		}
-	}
-	return true;
-}
-
 const FPoseSearchRoledSkeleton* UPoseSearchSchema::GetRoledSkeleton(const UE::PoseSearch::FRole& Role) const
 {
 	for (const FPoseSearchRoledSkeleton& RoledSkeleton : Skeletons)
@@ -251,9 +222,6 @@ void UPoseSearchSchema::ResetFinalize()
 	for (FPoseSearchRoledSkeleton& RoledSkeleton : Skeletons)
 	{
 		RoledSkeleton.BoneReferences.Reset();
-		PRAGMA_DISABLE_DEPRECATION_WARNINGS
-		RoledSkeleton.BoneIndicesWithParents_DEPRECATED.Reset();
-		PRAGMA_ENABLE_DEPRECATION_WARNINGS
 	}
 
 	FinalizedChannels.Reset();
@@ -341,39 +309,6 @@ void UPoseSearchSchema::Finalize()
 			UPoseSearchFeatureChannel_Padding::AddToSchema(this, PaddingSize);
 		}
 	}
-
-	PRAGMA_DISABLE_DEPRECATION_WARNINGS
-	// Initialize references to obtain bone indices and fill out bone index array
-	for (FPoseSearchRoledSkeleton& RoledSkeleton : Skeletons)
-	{
-		for (FBoneReference& BoneRef : RoledSkeleton.BoneReferences)
-		{
-			check(BoneRef.HasValidSetup());
-			RoledSkeleton.BoneIndicesWithParents_DEPRECATED.AddUnique(BoneRef.BoneIndex);
-
-			if (RoledSkeleton.MirrorDataTable)
-			{
-				if (RoledSkeleton.MirrorDataTable->BoneToMirrorBoneIndex.IsValidIndex(BoneRef.BoneIndex))
-				{
-					const FSkeletonPoseBoneIndex MirroredBoneIndex = RoledSkeleton.MirrorDataTable->BoneToMirrorBoneIndex[BoneRef.BoneIndex];
-					if (MirroredBoneIndex.IsValid())
-					{
-						RoledSkeleton.BoneIndicesWithParents_DEPRECATED.AddUnique(MirroredBoneIndex.GetInt());
-					}
-				}
-				else
-				{
-					UE_LOG(LogPoseSearch, Warning, TEXT("UPoseSearchSchema::Finalize: couldn't Finalize '%s' because bone index doest not exist in mirror table or mirror table is empty."), *GetNameSafe(this));
-				}
-			}
-		}
-
-		// Build separate index array with parent indices guaranteed to be present. Sort for EnsureParentsPresent.
-		check(!RoledSkeleton.BoneIndicesWithParents_DEPRECATED.IsEmpty());
-		RoledSkeleton.BoneIndicesWithParents_DEPRECATED.Sort();
-		FAnimationRuntime::EnsureParentsPresent(RoledSkeleton.BoneIndicesWithParents_DEPRECATED, RoledSkeleton.Skeleton->GetReferenceSkeleton());
-	}
-	PRAGMA_ENABLE_DEPRECATION_WARNINGS
 }
 
 void UPoseSearchSchema::PostLoad()
diff --git a/Engine/Plugins/Animation/PoseSearch/Source/Runtime/Private/Trace/PoseSearchTraceLogger.cpp b/Engine/Plugins/Animation/PoseSearch/Source/Runtime/Private/Trace/PoseSearchTraceLogger.cpp
--- a/Engine/Plugins/Animation/PoseSearch/Source/Runtime/Private/Trace/PoseSearchTraceLogger.cpp
+++ b/Engine/Plugins/Animation/PoseSearch/Source/Runtime/Private/Trace/PoseSearchTraceLogger.cpp
@@ -140,78 +140,4 @@ const FTraceMotionMatchingStatePoseEntry* FTraceMotionMatchingStateMessage::GetC
 	return nullptr;
 }
 
-FText FTraceMotionMatchingStateMessage::GenerateSearchName() const
-{
-	PRAGMA_DISABLE_DEPRECATION_WARNINGS
-	TStringBuilder<256> StringBuilder;
-	bool bAddDatabasesSeparator = false;
-	for (const FTraceMotionMatchingStateDatabaseEntry& DbEntry : DatabaseEntries)
-	{
-		if (bAddDatabasesSeparator)
-		{
-			StringBuilder.Append(" - ");
-		}
-
-		if (const UPoseSearchDatabase* Database = GetObjectFromId<UPoseSearchDatabase>(DbEntry.DatabaseId))
-		{
-			StringBuilder.Append(Database->GetName());
-		}
-		else
-		{
-			StringBuilder.Append("Unknown");
-		}
-		
-		bAddDatabasesSeparator = true;
-	}
-
-	if (Roles.Num() > 1)
-	{
-		if (Roles.Num() != SkeletalMeshComponentIds.Num())
-		{
-			StringBuilder.Append("Error!");
-		}
-		else
-		{
-			StringBuilder.Append(" [");
-
-			bool bAddRolesSeparator = false;
-
-			for (int32 RoleIndex = 0; RoleIndex < Roles.Num(); ++RoleIndex)
-			{
-				if (bAddRolesSeparator)
-				{
-					StringBuilder.Append(" - ");
-				}
-
-				StringBuilder.Append(Roles[RoleIndex].ToString());
-
-				StringBuilder.Append(": ");
-
-				if (const USkeletalMeshComponent* SkeletalMeshComponent = GetObjectFromId<USkeletalMeshComponent>(SkeletalMeshComponentIds[RoleIndex]))
-				{
-					if (const AActor* Actor = SkeletalMeshComponent->GetOwner())
-					{
-						StringBuilder.Append(Actor->GetName());
-					}
-					else
-					{
-						StringBuilder.Append("Error!");
-					}
-				}
-				else
-				{
-					StringBuilder.Append("Error!");
-				}
-
-				bAddRolesSeparator = true;
-			}
-
-			StringBuilder.Append("]");
-		}
-	}
-	PRAGMA_ENABLE_DEPRECATION_WARNINGS
-
-	return FText::FromString(StringBuilder.ToString());
-}
-
 } // namespace UE::PoseSearch
\ No newline at end of file
diff --git a/Engine/Plugins/Animation/PoseSearch/Source/Runtime/Public/PoseSearch/PoseHistoryProvider.h b/Engine/Plugins/Animation/PoseSearch/Source/Runtime/Public/PoseSearch/PoseHistoryProvider.h
--- a/Engine/Plugins/Animation/PoseSearch/Source/Runtime/Public/PoseSearch/PoseHistoryProvider.h
+++ b/Engine/Plugins/Animation/PoseSearch/Source/Runtime/Public/PoseSearch/PoseHistoryProvider.h
@@ -21,9 +21,6 @@ public:
 	FPoseHistoryProvider(const FAnimNode_PoseSearchHistoryCollector_Base* InHistoryCollector) : HistoryCollector(InHistoryCollector) { check(HistoryCollector); }
 	UE_API const IPoseHistory& GetPoseHistory() const;
 	
-	UE_DEPRECATED(5.6, "Use GetPoseHistory instead")
-	const FAnimNode_PoseSearchHistoryCollector_Base* GetHistoryCollector() const { return HistoryCollector; }
-
 private:
 	const FAnimNode_PoseSearchHistoryCollector_Base* HistoryCollector = nullptr;
 };
diff --git a/Engine/Plugins/Animation/PoseSearch/Source/Runtime/Public/PoseSearch/PoseSearchAssetSampler.h b/Engine/Plugins/Animation/PoseSearch/Source/Runtime/Public/PoseSearch/PoseSearchAssetSampler.h
--- a/Engine/Plugins/Animation/PoseSearch/Source/Runtime/Public/PoseSearch/PoseSearchAssetSampler.h
+++ b/Engine/Plugins/Animation/PoseSearch/Source/Runtime/Public/PoseSearch/PoseSearchAssetSampler.h
@@ -46,9 +46,6 @@ struct FAnimationAssetSampler
 	// sequence limits when Time is less than zero or greater than the sequence length.
 	POSESEARCH_API FTransform ExtractRootTransform(float Time) const;
 	
-	UE_DEPRECATED(5.6, "Use ExtractAnimNotifyStates instead")
-	POSESEARCH_API void ExtractPoseSearchNotifyStates(float Time, const TFunction<bool(UAnimNotifyState_PoseSearchBase*)>& ProcessPoseSearchBase) const;
-
 	// Extracts notify states present in the sequence at Time
 	POSESEARCH_API void ExtractAnimNotifyStates(float Time, FAnimNotifyContext& PreAllocatedNotifyContext, const TFunction<bool(UAnimNotifyState*)>& ProcessAnimNotifyState) const;
 	
diff --git a/Engine/Plugins/Animation/PoseSearch/Source/Runtime/Public/PoseSearch/PoseSearchContext.h b/Engine/Plugins/Animation/PoseSearch/Source/Runtime/Public/PoseSearch/PoseSearchContext.h
--- a/Engine/Plugins/Animation/PoseSearch/Source/Runtime/Public/PoseSearch/PoseSearchContext.h
+++ b/Engine/Plugins/Animation/PoseSearch/Source/Runtime/Public/PoseSearch/PoseSearchContext.h
@@ -37,17 +37,6 @@ const FBoneContainer GetBoneContainer(const UObject* AnimContext, bool bCheckIsI
 struct FPoseIndicesHistory;
 struct IPoseHistory;
 
-enum class UE_DEPRECATED(5.6, "EDebugDrawFlags is no longer used") EDebugDrawFlags : uint32
-{
-	None = 0,
-
-	// used to differenciate channels debug drawing of the query
-	DrawQuery = 1 << 0,
-};
-PRAGMA_DISABLE_DEPRECATION_WARNINGS
-ENUM_CLASS_FLAGS(EDebugDrawFlags);
-PRAGMA_ENABLE_DEPRECATION_WARNINGS
-
 enum class EPoseCandidateFlags : uint32
 {
 	None = 0,
@@ -73,13 +62,7 @@ ENUM_CLASS_FLAGS(EPoseCandidateFlags);
 #if ENABLE_DRAW_DEBUG
 struct FDebugDrawParams
 {
-	PRAGMA_DISABLE_DEPRECATION_WARNINGS
-	UE_DEPRECATED(5.6, "Use other FDebugDrawParams signatures instead")
-	UE_API FDebugDrawParams(TArrayView<FAnimInstanceProxy*> InAnimInstanceProxies, TConstArrayView<const IPoseHistory*> InPoseHistories, const FRoleToIndex& InRoleToIndex, const UPoseSearchDatabase* InDatabase, EDebugDrawFlags InFlags = EDebugDrawFlags::None);
-	UE_DEPRECATED(5.6, "Use other FDebugDrawParams signatures instead")
-	UE_API FDebugDrawParams(TArrayView<const USkinnedMeshComponent*> InMeshes, TConstArrayView<const IPoseHistory*> InPoseHistories, const FRoleToIndex& InRoleToIndex, const UPoseSearchDatabase* InDatabase, EDebugDrawFlags InFlags = EDebugDrawFlags::None);
-	UE_API PRAGMA_ENABLE_DEPRECATION_WARNINGS
-	FDebugDrawParams(TConstArrayView<FChooserEvaluationContext*> InAnimContexts, TConstArrayView<const IPoseHistory*> InPoseHistories, const FRoleToIndex& InRoleToIndex, const UPoseSearchDatabase* InDatabase);
+	UE_API FDebugDrawParams(TConstArrayView<FChooserEvaluationContext*> InAnimContexts, TConstArrayView<const IPoseHistory*> InPoseHistories, const FRoleToIndex& InRoleToIndex, const UPoseSearchDatabase* InDatabase);
 	FDebugDrawParams() = default;
 
 	UE_API FDebugDrawParams(const FDebugDrawParams& Other);
@@ -171,10 +154,6 @@ struct FCachedChannel
 
 struct FSearchContext
 {
-	UE_DEPRECATED(5.6, "Use other signaure instead")
-	UE_API FSearchContext(float InDesiredPermutationTimeOffset = 0.f, const FPoseIndicesHistory* InPoseIndicesHistory = nullptr,
-		const FSearchResult& InCurrentResult = FSearchResult(), const FFloatInterval& InPoseJumpThresholdTime = FFloatInterval(0.f, 0.f), bool bInUseCachedChannelData = false);
-
 	UE_DEPRECATED(5.7, "Use other signaure instead")
 	UE_API FSearchContext(float InDesiredPermutationTimeOffset, const FPoseIndicesHistory* InPoseIndicesHistory,
 		const FSearchResult& InCurrentResult, const FFloatInterval& InPoseJumpThresholdTime, const FPoseSearchEvent& InEventToSearch);
@@ -256,34 +235,8 @@ struct FSearchContext
 
 	float GetDesiredPermutationTimeOffset() const { return DesiredPermutationTimeOffset; }
 	
-	UE_DEPRECATED(5.6, "Use GetContext instead")
-	UE_API const UAnimInstance* GetAnimInstance(const FRole& Role) const;
-	UE_DEPRECATED(5.6, "Use GetContext instead")
-	UE_API const UObject* GetAnimContext(const FRole& Role) const;
-
 	UE_API const FChooserEvaluationContext* GetContext(const FRole& Role) const;
 	UE_API FChooserEvaluationContext* GetContext(const FRole& Role);
-	UE_DEPRECATED(5.6, "Use GetContexts instead")
-	const TArray<const UAnimInstance*, TInlineAllocator<PreallocatedRolesNum, TMemStackAllocator<>>> GetAnimInstances() const 
-	{
-	 	TArray<const UAnimInstance*, TInlineAllocator<PreallocatedRolesNum, TMemStackAllocator<>>> AnimInstances;
-	 	for (const FChooserEvaluationContext* AnimContext : AnimContexts)
-	 	{
-	 		AnimInstances.Add(Cast<UAnimInstance>(AnimContext->GetFirstObjectParam()));
-	 	}
-	 	return AnimInstances;
-	}
-
-	UE_DEPRECATED(5.6, "Use GetContexts instead")
-	const TArray<const UAnimInstance*, TInlineAllocator<PreallocatedRolesNum, TMemStackAllocator<>>> GetAnimContexts() const 
-	{
-		TArray<const UAnimInstance*, TInlineAllocator<PreallocatedRolesNum, TMemStackAllocator<>>> AnimInstances;
-		for (const FChooserEvaluationContext* AnimContext : AnimContexts)
-		{
-			AnimInstances.Add(Cast<UAnimInstance>(AnimContext->GetFirstObjectParam()));
-		}
-		return AnimInstances;
-	}
 	
 	TArrayView<FChooserEvaluationContext*> GetContexts() { return AnimContexts; }
 
diff --git a/Engine/Plugins/Animation/PoseSearch/Source/Runtime/Public/PoseSearch/PoseSearchCost.h b/Engine/Plugins/Animation/PoseSearch/Source/Runtime/Public/PoseSearch/PoseSearchCost.h
--- a/Engine/Plugins/Animation/PoseSearch/Source/Runtime/Public/PoseSearch/PoseSearchCost.h
+++ b/Engine/Plugins/Animation/PoseSearch/Source/Runtime/Public/PoseSearch/PoseSearchCost.h
@@ -10,16 +10,12 @@ struct FPoseSearchCost
 	GENERATED_BODY()
 public:
 
-PRAGMA_DISABLE_DEPRECATION_WARNINGS
-	// Note: We need to explicitly disable warnings on these constructors/operators for clang to be happy with deprecated variables
-	// this is a requirement for clang to compile without warnings.
 	FPoseSearchCost() = default;
 	~FPoseSearchCost() = default;
 	FPoseSearchCost(const FPoseSearchCost&) = default;
 	FPoseSearchCost(FPoseSearchCost&&) = default;
 	FPoseSearchCost& operator=(const FPoseSearchCost&) = default;
 	FPoseSearchCost& operator=(FPoseSearchCost&&) = default;
-PRAGMA_ENABLE_DEPRECATION_WARNINGS
 
 	FPoseSearchCost(float DissimilarityCost, float InNotifyCostAddend, float InContinuingPoseCostAddend, float InContinuingInteractionCostAddend = 0.f)
 		: TotalCost(DissimilarityCost + InNotifyCostAddend + InContinuingPoseCostAddend + InContinuingInteractionCostAddend)
@@ -27,9 +23,6 @@ PRAGMA_ENABLE_DEPRECATION_WARNINGS
 		, NotifyCostAddend(InNotifyCostAddend)
 		, ContinuingPoseCostAddend(InContinuingPoseCostAddend)
 		, ContinuingInteractionCostAddend(InContinuingInteractionCostAddend)
-PRAGMA_DISABLE_DEPRECATION_WARNINGS
-		, CostAddend(InNotifyCostAddend + InContinuingPoseCostAddend + InContinuingInteractionCostAddend)
-PRAGMA_ENABLE_DEPRECATION_WARNINGS
 #endif // WITH_EDITORONLY_DATA
 	{
 	}
@@ -37,9 +30,6 @@ PRAGMA_ENABLE_DEPRECATION_WARNINGS
 	static bool IsCostValid(const float Cost) { return Cost != MAX_flt; }
 	bool IsValid() const { return IsCostValid(TotalCost); }
 	
-	UE_DEPRECATED(5.6, "Use operator float() instead.")
-	float GetTotalCost() const { return TotalCost; }
-
 	operator float() const { return TotalCost; }
 
 #if WITH_EDITORONLY_DATA
@@ -58,11 +48,6 @@ PRAGMA_ENABLE_DEPRECATION_WARNINGS
 		Ar << PoseSearchCost.NotifyCostAddend;
 		Ar << PoseSearchCost.ContinuingPoseCostAddend;
 		Ar << PoseSearchCost.ContinuingInteractionCostAddend;
-
-		PRAGMA_DISABLE_DEPRECATION_WARNINGS
-		Ar << PoseSearchCost.CostAddend;
-		PRAGMA_ENABLE_DEPRECATION_WARNINGS
-
 #endif // WITH_EDITORONLY_DATA
 			
 		return Ar;
@@ -88,11 +73,6 @@ protected:
 	UPROPERTY()
 	float ContinuingInteractionCostAddend = 0.f;
 
-public:
-
-	UE_DEPRECATED(5.6, "Use GetCostAddend instead")
-	UPROPERTY()
-	float CostAddend = 0.f;
 #endif // WITH_EDITORONLY_DATA
 };
 
diff --git a/Engine/Plugins/Animation/PoseSearch/Source/Runtime/Public/PoseSearch/PoseSearchDatabase.h b/Engine/Plugins/Animation/PoseSearch/Source/Runtime/Public/PoseSearch/PoseSearchDatabase.h
--- a/Engine/Plugins/Animation/PoseSearch/Source/Runtime/Public/PoseSearch/PoseSearchDatabase.h
+++ b/Engine/Plugins/Animation/PoseSearch/Source/Runtime/Public/PoseSearch/PoseSearchDatabase.h
@@ -59,8 +59,6 @@ struct FPoseSearchDatabaseAnimationAssetBase
 	virtual ~FPoseSearchDatabaseAnimationAssetBase() = default;
 	virtual UObject* GetAnimationAsset() const { return nullptr; }
 	
-	UE_DEPRECATED(5.6, "Use other signatures")
-	virtual float GetPlayLength() const { return GetPlayLength(FVector::ZeroVector); }
 	UE_API virtual float GetPlayLength(const FVector& BlendParameters) const;
 
 	virtual int32 GetNumRoles() const { return 1; }
@@ -68,8 +66,6 @@ struct FPoseSearchDatabaseAnimationAssetBase
 	UE_API virtual UAnimationAsset* GetAnimationAssetForRole(const UE::PoseSearch::FRole& Role) const;
 	UE_API virtual FTransform GetRootTransformOriginForRole(const UE::PoseSearch::FRole& Role) const;
 #if WITH_EDITOR
-	UE_DEPRECATED(5.6, "GetFrameAtTime is no longer supported")
-	UE_API virtual int32 GetFrameAtTime(float Time) const;
 	UE_API virtual bool IsSkeletonCompatible(TObjectPtr<const UPoseSearchSchema> InSchema) const;
 	// Experimental, this feature might be removed without warning, not for production use
 	UE_API virtual USkeletalMesh* GetPreviewMeshForRole(const UE::PoseSearch::FRole& Role) const;
@@ -93,8 +89,6 @@ struct FPoseSearchDatabaseAnimationAssetBase
 	virtual void SetSamplingRange(const FFloatInterval& NewRange) { }
 	UE_API virtual FFloatInterval GetEffectiveSamplingRange(const FVector& BlendParameters) const;
 	
-	UE_DEPRECATED(5.6, "Use other signatures")
-	UE_API FFloatInterval GetEffectiveSamplingRange() const;
 	static UE_API FFloatInterval GetEffectiveSamplingRange(float PlayLength, const FFloatInterval& SamplingRange);
 
 	UE_API virtual int64 GetEditorMemSize() const;
@@ -190,11 +184,6 @@ struct FPoseSearchDatabaseAnimationAsset : public FPoseSearchDatabaseAnimationAs
 	UE_API virtual FFloatInterval GetEffectiveSamplingRange(const FVector& BlendParameters) const override;
 	UE_API virtual void IterateOverSamplingParameter(const TFunction<void(const FVector& BlendParameters)>& ProcessSamplingParameter) const override;
 	virtual void SetSamplingRange(const FFloatInterval& NewRange) override { SamplingRange = NewRange; }
-
-	UE_DEPRECATED(5.6, "Use IterateOverSamplingParameter to iterate over all the sampling parameters for the blend space")
-	UE_API void GetBlendSpaceParameterSampleRanges(int32& HorizontalBlendNum, int32& VerticalBlendNum) const;
-	UE_DEPRECATED(5.6, "Use IterateOverSamplingParameter to iterate over all the sampling parameters for the blend space")
-	UE_API FVector BlendParameterForSampleRanges(int32 HorizontalBlendIndex, int32 VerticalBlendIndex) const;
 #endif // WITH_EDITORONLY_DATA
 
 #if WITH_EDITOR
@@ -325,11 +314,6 @@ struct UE_DEPRECATED(5.7, "Use FPoseSearchDatabaseAnimationAsset instead. To con
 	UE_API virtual FFloatInterval GetEffectiveSamplingRange(const FVector& BlendParameters) const override;
 	UE_API virtual void IterateOverSamplingParameter(const TFunction<void(const FVector& BlendParameters)>& ProcessSamplingParameter) const override;
 	virtual void SetSamplingRange(const FFloatInterval& NewRange) override { SamplingRange = NewRange; }
-
-	UE_DEPRECATED(5.6, "Use IterateOverSamplingParameter to iterate over all the sampling parameters for the blend space")
-	UE_API void GetBlendSpaceParameterSampleRanges(int32& HorizontalBlendNum, int32& VerticalBlendNum) const;
-	UE_DEPRECATED(5.6, "Use IterateOverSamplingParameter to iterate over all the sampling parameters for the blend space")
-	UE_API FVector BlendParameterForSampleRanges(int32 HorizontalBlendIndex, int32 VerticalBlendIndex) const;
 #endif // WITH_EDITORONLY_DATA
 
 	UE_API virtual UObject* GetAnimationAsset() const override;
@@ -647,11 +631,6 @@ public:
 	
 	UE_API bool GetSkipSearchIfPossible() const;
 
-	UE_DEPRECATED(5.6, "Use SearchIndexAsset.GetPoseIndexFromTime instead")
-	// given RealTimeInSeconds in seconds (not a normalized time, so for blend spaces needs to be multiplied by the 
-	// blendspace length given the blend parameters). it'll return the associated pose index
-	UE_API int32 GetPoseIndexFromTime(float RealTimeInSeconds, const UE::PoseSearch::FSearchIndexAsset& SearchIndexAsset) const;
-
 	// given an AnimationAsset, a normalized time AnimationAssetTime (in [0-1] range for blend spaces), and mirror state bMirrored
 	// will retrieve the associated pose index from this database closest in distance to BlendParameters.
 	// It'll return INDEX_NONE if it cannot find a pose
diff --git a/Engine/Plugins/Animation/PoseSearch/Source/Runtime/Public/PoseSearch/PoseSearchDefines.h b/Engine/Plugins/Animation/PoseSearch/Source/Runtime/Public/PoseSearch/PoseSearchDefines.h
--- a/Engine/Plugins/Animation/PoseSearch/Source/Runtime/Public/PoseSearch/PoseSearchDefines.h
+++ b/Engine/Plugins/Animation/PoseSearch/Source/Runtime/Public/PoseSearch/PoseSearchDefines.h
@@ -43,7 +43,7 @@ static constexpr int32 PreallocatedSearchesNum = 4;
 static constexpr float FiniteDelta = 1 / 60.0f; // Time delta used for computing pose derivatives
 static constexpr int32 MaxNumberOfCollectedPoseCandidatesPerDatabase = 256;
 static constexpr int32 InvalidSearchId = INDEX_NONE; // @todo: it should be 0, since the search id is an hash!
-static constexpr int32 DatabaseIndexDerivedDataCacheKeyVersion = 2719; // bump this value to force reindex ALL the UPoseSearchDatabases, without having to change the global FDevSystemGuids::POSESEARCHDB_DERIVEDDATA_VER
+static constexpr int32 DatabaseIndexDerivedDataCacheKeyVersion = 2723; // bump this value to force reindex ALL the UPoseSearchDatabases, without having to change the global FDevSystemGuids::POSESEARCHDB_DERIVEDDATA_VER
 static constexpr int32 DefaultTickPriority = 100;
 
 template<uint32 Alignment = DEFAULT_ALIGNMENT> class TMemStackSetAllocator : public TSetAllocator<TSparseArrayAllocator<TMemStackAllocator<Alignment>, TMemStackAllocator<Alignment>>, TMemStackAllocator<Alignment>> {};
diff --git a/Engine/Plugins/Animation/PoseSearch/Source/Runtime/Public/PoseSearch/PoseSearchHistory.h b/Engine/Plugins/Animation/PoseSearch/Source/Runtime/Public/PoseSearch/PoseSearchHistory.h
--- a/Engine/Plugins/Animation/PoseSearch/Source/Runtime/Public/PoseSearch/PoseSearchHistory.h
+++ b/Engine/Plugins/Animation/PoseSearch/Source/Runtime/Public/PoseSearch/PoseSearchHistory.h
@@ -192,21 +192,12 @@ struct FPoseHistory : public IPoseHistory
 	// @todo: deprecate and delete this method in favor of providing the trajectory via SetTrajectory
 	UE_API void GenerateTrajectory(const UObject* AnimContext, float DeltaTime, const FPoseSearchTrajectoryData& TrajectoryData, const FPoseSearchTrajectoryData::FSampling& TrajectoryDataSampling);
 	
-	UE_DEPRECATED(5.6, "No longer necessary method")
-	void PreUpdate() {}
-
 	UE_API void Initialize_AnyThread(int32 InNumPoses, float InSamplingInterval);
 	UE_API void EvaluateComponentSpace_AnyThread(float DeltaTime, IComponentSpacePoseProvider& ComponentSpacePoseProvider, bool bStoreScales,
 		float RootBoneRecoveryTime, float RootBoneTranslationRecoveryRatio, float RootBoneRotationRecoveryRatio,
 		bool bNeedsReset, bool bCacheBones, const TArray<FBoneIndexType>& RequiredBones, 
 		const FBlendedCurve& Curves = FBlendedCurve(), const TConstArrayView<FName>& InCollectedCurves = TConstArrayView<FName>());
 
-	UE_DEPRECATED(5.6, "Use other EvaluateComponentSpace_AnyThread signatures instead")
-	UE_API void EvaluateComponentSpace_AnyThread(float DeltaTime, FCSPose<FCompactPose>& ComponentSpacePose, bool bStoreScales,
-		float RootBoneRecoveryTime, float RootBoneTranslationRecoveryRatio, float RootBoneRotationRecoveryRatio,
-		bool bNeedsReset, bool bCacheBones, const TArray<FBoneIndexType>& RequiredBones,
-		const FBlendedCurve& Curves = FBlendedCurve(), const TConstArrayView<FName>& InCollectedCurves = TConstArrayView<FName>());
-		
 	// IPoseHistory interface
 	UE_API virtual bool GetTransformAtTime(float Time, FTransform& OutBoneTransform, const USkeleton* BoneIndexSkeleton = nullptr, FBoneIndexType BoneIndexType = RootBoneIndexType, FBoneIndexType ReferenceBoneIndexType = ComponentSpaceIndexType, bool bExtrapolate = false) const override;
 	UE_API virtual bool GetCurveValueAtTime(float Time, const FName& CurveName, float& outcurvevalue, bool bExtrapolate = false) const override;
@@ -220,13 +211,7 @@ struct FPoseHistory : public IPoseHistory
 	UE_API virtual const FBoneToTransformMap& GetBoneToTransformMap() const override;
 	UE_API virtual const TConstArrayView<FName> GetCollectedCurves() const override;
 
-	// We need to handle the changes from CL 39422514 since licensees still haven't gotten the move of this function to the IPoseHistory, as of 5.5, therefore this function still needs to be marked as deprecated even tho it was deleted later on.
-	PRAGMA_DISABLE_DEPRECATION_WARNINGS
-	UE_DEPRECATED(5.6, "SetTrajectory() that takes in FPoseSearchQueryTrajectory is deprecated. Use version that takes FTransformTrajectory instead.")
-	UE_API void SetTrajectory(const FPoseSearchQueryTrajectory& InTrajectory, float InTrajectorySpeedMultiplier = 1.f);
-	UE_API PRAGMA_ENABLE_DEPRECATION_WARNINGS
-	
-	virtual int32 GetNumEntries() const override;
+	UE_API virtual int32 GetNumEntries() const override;
 	UE_API virtual const FPoseHistoryEntry& GetEntry(int32 EntryIndex) const override;
 	UE_API virtual const FPoseIndicesHistory* GetPoseIndicesHistory() const override;
 
@@ -320,9 +305,6 @@ struct FMemStackPoseHistory : public IPoseHistory
 	// if FiniteDeltaTime is > FiniteDeltaTime a second pose will be extracted at AnimationTime - FiniteDeltaTime and placed at IntervalTime - FiniteDeltaTime (for MM to be able to calculate velocities)
 	UE_API void ExtractAndAddFuturePoses(const UAnimationAsset* AnimationAsset, float AnimationTime, float FiniteDeltaTime, const FVector& BlendParameters, float IntervalTime, const USkeleton* OverrideSkeleton = nullptr, bool bUseRefPoseRootBone = false);
 
-	UE_DEPRECATED(5.6, "Use other AddFuturePose signatures instead")
-	UE_API void AddFuturePose(float Time, FCSPose<FCompactPose>& ComponentSpacePose, const FBlendedCurve& Curves = FBlendedCurve());
-
 	const IPoseHistory* GetThisOrPoseHistory() const { return FutureEntries.IsEmpty() ? PoseHistory : this; }
 
 private:
diff --git a/Engine/Plugins/Animation/PoseSearch/Source/Runtime/Public/PoseSearch/PoseSearchIndex.h b/Engine/Plugins/Animation/PoseSearch/Source/Runtime/Public/PoseSearch/PoseSearchIndex.h
--- a/Engine/Plugins/Animation/PoseSearch/Source/Runtime/Public/PoseSearch/PoseSearchIndex.h
+++ b/Engine/Plugins/Animation/PoseSearch/Source/Runtime/Public/PoseSearch/PoseSearchIndex.h
@@ -290,6 +290,7 @@ private:
 	const float ToRealTimeFactor = 1.f;
 };
 
+#if WITH_EDITORONLY_DATA
 struct FSearchStats
 {
 	float AverageSpeed = 0.f;
@@ -301,6 +302,7 @@ struct FSearchStats
 	bool operator==(const FSearchStats& Other) const;
 	friend FArchive& operator<<(FArchive& Ar, FSearchStats& Stats);
 };
+#endif // WITH_EDITORONLY_DATA
 
 // compact representation of an array of arrays
 template <typename Type = int32>
@@ -512,9 +514,9 @@ struct FSearchIndexBase
 	// minimum of the database metadata CostAddend: it represents the minimum cost of any search for the associated database (we'll skip the search in case the search result total cost is already less than MinCostAddend)
 	float MinCostAddend = -MAX_FLT;
 
-	// @todo: this property should be editor only
-	UE_DEPRECATED(5.6, "this property will become editor only. wrap its usage with #if WITH_EDITORONLY_DATA")
+#if WITH_EDITORONLY_DATA
 	FSearchStats Stats;
+#endif // WITH_EDITORONLY_DATA
 
 	// overriding ctor, copy, move to allow deprecations.
 	FSearchIndexBase();
@@ -595,8 +597,6 @@ struct FSearchIndex : public FSearchIndexBase
 	TAlignedArray<float> DeviationEditorOnly;
 	float PCAExplainedVarianceEditorOnly = 0.f;
 #endif
-	UE_DEPRECATED(5.6, "Use PCAExplainedVarianceEditorOnly")
-	float PCAExplainedVariance = 0.f;
 
 	FSearchIndex() = default;
 	~FSearchIndex() = default;
@@ -612,8 +612,6 @@ struct FSearchIndex : public FSearchIndexBase
 	int32 GetNumberOfPrincipalComponents() const;
 	POSESEARCH_API TConstArrayView<float> PCAProject(TConstArrayView<float> PoseValues, TArrayView<float> BufferUsedForProjection) const;
 
-	UE_DEPRECATED(5.6, "Use other GetPoseValuesSafe signatures instead")
-	POSESEARCH_API TArray<float> GetPoseValuesSafe(int32 PoseIdx) const;
 	POSESEARCH_API TConstArrayView<float> GetPoseValuesSafe(int32 PoseIdx, TArray<float>& BufferUsedForReconstruction) const;
 
 	// since PCAValues (as well as Values can be pruned out from duplicate data, we lose the 1:1 mapping between PoseIdx and PCAValuesVectorIdx
@@ -628,11 +626,6 @@ struct FSearchIndex : public FSearchIndexBase
 		return PCAValues.Num() / DataCardinality;
 	}
 
-	UE_DEPRECATED(5.6, "Use the bAlignedAndPadded templated way of calculating FPoseSearchCost(CompareFeatureVectors<false>(PoseValues, QueryValues, WeightsSqrt), PoseMetadata[PoseIdx].GetCostAddend(), ContinuingPoseCostBias, 0.f) instead")
-	POSESEARCH_API FPoseSearchCost ComparePoses(int32 PoseIdx, float ContinuingPoseCostBias, TConstArrayView<float> PoseValues, TConstArrayView<float> QueryValues) const;
-	UE_DEPRECATED(5.6, "Use the bAlignedAndPadded templated way of calculating FPoseSearchCost(CompareFeatureVectors<true>(PoseValues, QueryValues, WeightsSqrt), PoseMetadata[PoseIdx].GetCostAddend(), ContinuingPoseCostBias, 0.f) instead")
-	POSESEARCH_API FPoseSearchCost CompareAlignedPoses(int32 PoseIdx, float ContinuingPoseCostBias, TConstArrayView<float> PoseValues, TConstArrayView<float> QueryValues) const;
-
 	void PruneDuplicatePCAValues(float SimilarityThreshold, int32 NumberOfPrincipalComponents);
 	void PrunePCAValuesFromBlockTransitionPoses(int32 NumberOfPrincipalComponents);
 
diff --git a/Engine/Plugins/Animation/PoseSearch/Source/Runtime/Public/PoseSearch/PoseSearchLibrary.h b/Engine/Plugins/Animation/PoseSearch/Source/Runtime/Public/PoseSearch/PoseSearchLibrary.h
--- a/Engine/Plugins/Animation/PoseSearch/Source/Runtime/Public/PoseSearch/PoseSearchLibrary.h
+++ b/Engine/Plugins/Animation/PoseSearch/Source/Runtime/Public/PoseSearch/PoseSearchLibrary.h
@@ -66,38 +66,17 @@ PRAGMA_DISABLE_DEPRECATION_WARNINGS
 	~FMotionMatchingState() = default;
 PRAGMA_ENABLE_DEPRECATION_WARNINGS
 
-	// Reset the state to a default state using the current Database
-	UE_DEPRECATED(5.6, "Use other Reset signatures instead")
-	void Reset(const FTransform& ComponentTransform);
-
 	// Reset the state to a default state using the current Database
 	POSESEARCH_API void Reset();
 
-	UE_DEPRECATED(5.6, "Use CurrentSearchResult.UpdateWithNormalizedTime or CurrentSearchResult.UpdateWithRealTime instead")
-	POSESEARCH_API void AdjustAssetTime(float AssetTime);
-
-	UE_DEPRECATED(5.6, "This API is deprecated and it'll be made private")
-	POSESEARCH_API void UpdateWantedPlayRate(const UE::PoseSearch::FSearchContext& SearchContext, const FFloatInterval& PlayRate, float TrajectorySpeedMultiplier, const FPoseSearchEvent& EventToSearch /*= FPoseSearchEvent()*/);
-
 	FVector GetEstimatedFutureRootMotionVelocity() const;
 
-	UE_DEPRECATED(5.6, "Use SearchResult instead")
-	UE::PoseSearch::FSearchResult CurrentSearchResult;
-
 	UPROPERTY(Transient)
 	FPoseSearchBlueprintResult SearchResult;
 
 	// Time since the last pose jump (real time in seconds, not a normalized time)
 	float ElapsedPoseSearchTime = 0.f;
 
-	UE_DEPRECATED(5.6, "Use SearchResult.WantedPlayRate instead")
-	// wanted PlayRate to have the selected animation playing at the estimated requested speed from the query.
-	float WantedPlayRate = 1.f;
-
-	// true if a new animation has been selected
-	UE_DEPRECATED(5.6, "Derive this value via 'CurrentSearchResult.IsValid() && !CurrentSearchResult.bIsContinuingPoseSearch' if needed")
-	bool bJumpedToPose = false;
-
 	UE_DEPRECATED(5.7, "this property has been moved into the pose history. Use IPoseHistory::GetPoseIndicesHistory to access it")
 	UE::PoseSearch::FPoseIndicesHistory PoseIndicesHistory;
 };
@@ -179,31 +158,6 @@ class UPoseSearchLibrary : public UBlueprintFunctionLibrary
 		float ElapsedPoseSearchTime,
 		float WantedPlayRate,
 		EPoseSearchInterruptMode InterruptMode);
-
-	UE_DEPRECATED(5.6, "Use TraceMotionMatching with different signature instead.")
-	static void TraceMotionMatching(
-		UE::PoseSearch::FSearchContext& SearchContext,
-		const UE::PoseSearch::FSearchResult& SearchResult,
-		float ElapsedPoseSearchTime,
-		float DeltaTime,
-		bool bSearch,
-		float WantedPlayRate)
-	{
-		TraceMotionMatching(SearchContext, SearchResult, ElapsedPoseSearchTime, DeltaTime, bSearch, WantedPlayRate, EPoseSearchInterruptMode::DoNotInterrupt);
-	}
-
-	UE_DEPRECATED(5.6, "Use TraceMotionMatching with different signature instead.")
-	static void TraceMotionMatching(
-		UE::PoseSearch::FSearchContext& SearchContext,
-		FMotionMatchingState& CurrentState,
-		float DeltaTime,
-		bool bSearch,
-		float RecordingTime)
-	{
-		PRAGMA_DISABLE_DEPRECATION_WARNINGS
-		TraceMotionMatching(SearchContext, CurrentState.CurrentSearchResult, CurrentState.ElapsedPoseSearchTime, DeltaTime, bSearch, CurrentState.WantedPlayRate, EPoseSearchInterruptMode::DoNotInterrupt);
-		PRAGMA_ENABLE_DEPRECATION_WARNINGS
-	}
 #endif // UE_POSE_SEARCH_TRACE_ENABLED
 
 public:
@@ -241,40 +195,6 @@ public:
 		bool bDebugDrawCurResult = false,
 		const FPoseSearchEvent& EventToSearch = FPoseSearchEvent());
 	
-	UE_DEPRECATED(5.6, "Use MotionMatch API instead")
-	static UE_API void UpdateMotionMatchingState(
-		const UObject* AnimContext,
-		const UE::PoseSearch::IPoseHistory* PoseHistory,
-		const TConstArrayView<TObjectPtr<const UPoseSearchDatabase>> Databases,
-		float DeltaTime,
-		const FFloatInterval& PoseJumpThresholdTime,
-		float PoseReselectHistory,
-		float SearchThrottleTime,
-		const FFloatInterval& PlayRate,
-		FMotionMatchingState& InOutMotionMatchingState,
-		EPoseSearchInterruptMode InterruptMode,
-		bool bShouldUseCachedChannelData,
-		bool bDebugDrawQuery = false,
-		bool bDebugDrawCurResult = false,
-		const FPoseSearchEvent& EventToSearch = FPoseSearchEvent());
-	
-	UE_DEPRECATED(5.6, "Use MotionMatch API instead")
-	static UE_API void UpdateMotionMatchingState(
-		const FAnimationUpdateContext& Context,
-		const TArray<TObjectPtr<const UPoseSearchDatabase>>& Databases,
-		float BlendTime,
-		int32 MaxActiveBlends,
-		const FFloatInterval& PoseJumpThresholdTime,
-		float PoseReselectHistory,
-		float SearchThrottleTime,
-		const FFloatInterval& PlayRate,
-		FMotionMatchingState& InOutMotionMatchingState,
-		EPoseSearchInterruptMode InterruptMode = EPoseSearchInterruptMode::DoNotInterrupt,
-		bool bShouldSearch = true,
-		bool bShouldUseCachedChannelData = true,
-		bool bDebugDrawQuery = false,
-		bool bDebugDrawCurResult = false);
-
 	/**
 	* Implementation of the core motion matching algorithm
 	*
@@ -295,16 +215,6 @@ public:
 		const FPoseSearchFutureProperties Future,
 		FPoseSearchBlueprintResult& Result);
 
-	UE_DEPRECATED(5.6, "Use other MotionMatch signatures instead")
-	static UE_API void MotionMatch(
-		const TArrayView<UAnimInstance*> AnimInstances,
-		const TArrayView<const UE::PoseSearch::FRole> Roles,
-		const TArrayView<const UObject*> AssetsToSearch,
-		const FName PoseHistoryName,
-		const FPoseSearchContinuingProperties& ContinuingProperties,
-		const FPoseSearchFutureProperties& Future,
-		FPoseSearchBlueprintResult& Result);
-
 	// deprecate in favor of MotionMatch with the FSearchResults parameter once we settle on FSearchResults API signatures
 	static UE_API UE::PoseSearch::FSearchResult MotionMatch(
 		const TArrayView<const UObject*> AnimContexts,
diff --git a/Engine/Plugins/Animation/PoseSearch/Source/Runtime/Public/PoseSearch/PoseSearchResult.h b/Engine/Plugins/Animation/PoseSearch/Source/Runtime/Public/PoseSearch/PoseSearchResult.h
--- a/Engine/Plugins/Animation/PoseSearch/Source/Runtime/Public/PoseSearch/PoseSearchResult.h
+++ b/Engine/Plugins/Animation/PoseSearch/Source/Runtime/Public/PoseSearch/PoseSearchResult.h
@@ -194,12 +194,6 @@ struct FPoseSearchBlueprintResult
 {
 	GENERATED_BODY()
 
-#if WITH_EDITORONLY_DATA
-	// Deprecated in 5.6. Use "SelectedAnim" property instead
-	UPROPERTY(Transient, meta = (DeprecatedProperty, DeprecationMessage = "Use SelectedAnim instead"))
-	TObjectPtr<const UObject> SelectedAnimation_DEPRECATED = nullptr;
-#endif // WITH_EDITORONLY_DATA
-
 	// selected animation selected by the motion matching alghoritm that matches the query (trajectory + historical pose) the best
 	UPROPERTY(Transient, VisibleAnywhere, BlueprintReadOnly, Category=State)
 	TObjectPtr<UObject> SelectedAnim = nullptr;
diff --git a/Engine/Plugins/Animation/PoseSearch/Source/Runtime/Public/PoseSearch/PoseSearchSchema.h b/Engine/Plugins/Animation/PoseSearch/Source/Runtime/Public/PoseSearch/PoseSearchSchema.h
--- a/Engine/Plugins/Animation/PoseSearch/Source/Runtime/Public/PoseSearch/PoseSearchSchema.h
+++ b/Engine/Plugins/Animation/PoseSearch/Source/Runtime/Public/PoseSearch/PoseSearchSchema.h
@@ -55,10 +55,6 @@ struct FPoseSearchRoledSkeleton
 	UPROPERTY(Transient)
 	TArray<FBoneReference> BoneReferences;
 
-	UE_DEPRECATED(5.6, "no longer necessary property")
-	UPROPERTY(Transient)
-	TArray<uint16> BoneIndicesWithParents_DEPRECATED;
-
 	UPROPERTY(Transient)
 	TArray<FName> RequiredCurves;
 };
@@ -173,11 +169,6 @@ public:
 	UE_API bool AreSkeletonsCompatible(const UPoseSearchSchema* Other) const;
 	UE_API void AddDefaultChannels();
 
-	UE_DEPRECATED(5.6, "No longer supported API")
-	UE_API void InitBoneContainersFromRoledSkeleton(TMap<FName, FBoneContainer>& RoledBoneContainers) const;
-
-	UE_DEPRECATED(5.6, "access the Skeletons via GetRoledSkeletons() and query for their MirrorDataTable instead")
-	UE_API bool AllRoledSkeletonHaveMirrorDataTable() const;
 	UE_API const FPoseSearchRoledSkeleton* GetRoledSkeleton(const UE::PoseSearch::FRole& Role) const;
 	UE_API FPoseSearchRoledSkeleton* GetRoledSkeleton(const UE::PoseSearch::FRole& Role);
 
diff --git a/Engine/Plugins/Animation/PoseSearch/Source/Runtime/Public/PoseSearch/Trace/PoseSearchTraceLogger.h b/Engine/Plugins/Animation/PoseSearch/Source/Runtime/Public/PoseSearch/Trace/PoseSearchTraceLogger.h
--- a/Engine/Plugins/Animation/PoseSearch/Source/Runtime/Public/PoseSearch/Trace/PoseSearchTraceLogger.h
+++ b/Engine/Plugins/Animation/PoseSearch/Source/Runtime/Public/PoseSearch/Trace/PoseSearchTraceLogger.h
@@ -133,19 +133,6 @@ struct FTraceMotionMatchingStateMessage : public FTraceMessage
 		return nullptr;
 	}
 
-	UE_DEPRECATED(5.6, "Use FObjectTrace::GetObjectId instead")
-	static uint64 GetIdFromObject(const UObject* Object)
-	{
-#if OBJECT_TRACE_ENABLED
-		return FObjectTrace::GetObjectId(Object);
-#else
-		return 0;
-#endif
-	}
-	
-	UE_DEPRECATED(5.6, "Use FText GenerateSearchName(const FTraceMotionMatchingStateMessage&, const IGameplayProvider*) instead")
-	UE_API FText GenerateSearchName() const;
-
 	UE_DEPRECATED(5.7, "FTraceMotionMatchingStateMessage::Name is no longer necessary and it'll be removed")
 	static UE_API const FName Name;
 };
diff --git a/Engine/Plugins/Experimental/UAF/UAFPoseSearch/Source/UAFPoseSearch/Private/MotionMatchingTrait.cpp b/Engine/Plugins/Experimental/UAF/UAFPoseSearch/Source/UAFPoseSearch/Private/MotionMatchingTrait.cpp
--- a/Engine/Plugins/Experimental/UAF/UAFPoseSearch/Source/UAFPoseSearch/Private/MotionMatchingTrait.cpp
+++ b/Engine/Plugins/Experimental/UAF/UAFPoseSearch/Source/UAFPoseSearch/Private/MotionMatchingTrait.cpp
@@ -282,12 +282,6 @@ PRAGMA_ENABLE_DEPRECATION_WARNINGS
 					// const casting here is safe since we're in the thread owning the pose history, and it's the correct place to update the previously selected poses
 					const_cast<FPoseIndicesHistory*>(PoseIndicesHistory)->Update(SearchResult, DeltaTime, SharedData->GetPoseReselectHistory(Binding));
 				}
-
-				PRAGMA_DISABLE_DEPRECATION_WARNINGS
-				MotionMatchingState.CurrentSearchResult = SearchResult;
-				MotionMatchingState.WantedPlayRate = WantedPlayRate;
-				MotionMatchingState.bJumpedToPose = MotionMatchingState.SearchResult.SelectedAnim && !MotionMatchingState.SearchResult.bIsContinuingPoseSearch;
-				PRAGMA_ENABLE_DEPRECATION_WARNINGS
 			}
 		}
 
