PoseSearch - improvements in displaying invalid poses in the pose search debugger

diff --git a/Engine/Plugins/Animation/PoseSearch/Source/Editor/Private/PoseSearchDatabaseDataDetails.cpp b/Engine/Plugins/Animation/PoseSearch/Source/Editor/Private/PoseSearchDatabaseDataDetails.cpp
--- a/Engine/Plugins/Animation/PoseSearch/Source/Editor/Private/PoseSearchDatabaseDataDetails.cpp
+++ b/Engine/Plugins/Animation/PoseSearch/Source/Editor/Private/PoseSearchDatabaseDataDetails.cpp
@@ -576,7 +576,7 @@ void SDatabaseDataDetails::Reconstruct(int32 MaxPreviewActors)
 		for (const FDatabasePreviewActor& PreviewActor : PreviewActorGroup)
 		{
 			HeaderRow->AddColumn(SHeaderRow::Column(*GetNameSafe(PreviewActor.GetActor())) 
-				.DefaultLabel(FText::FromString(PreviewActor.GetSampler().GetAsset()->GetName())));
+				.DefaultLabel(FText::FromString(GetNameSafe(PreviewActor.GetSampler().GetAsset()))));
 
 			++AddedColumns;
 			if (AddedColumns >= MaxPreviewActors)
diff --git a/Engine/Plugins/Animation/PoseSearch/Source/Editor/Private/PoseSearchDebuggerDatabaseColumns.h b/Engine/Plugins/Animation/PoseSearch/Source/Editor/Private/PoseSearchDebuggerDatabaseColumns.h
--- a/Engine/Plugins/Animation/PoseSearch/Source/Editor/Private/PoseSearchDebuggerDatabaseColumns.h
+++ b/Engine/Plugins/Animation/PoseSearch/Source/Editor/Private/PoseSearchDebuggerDatabaseColumns.h
@@ -125,12 +125,19 @@ struct FPoseIdx : ITextColumn
 
 	virtual FSortPredicate GetSortPredicate() const override
 	{
-		return [](const FRowDataRef& Row0, const FRowDataRef& Row1) -> bool { return Row0->PoseIdx < Row1->PoseIdx; };
+		return [](const FRowDataRef& Row0, const FRowDataRef& Row1) -> bool
+			{
+				return Row0->PoseIdx < Row1->PoseIdx;
+			};
 	}
 
 	virtual FText GetRowText(const FRowDataRef& Row) const override
 	{
-		return FText::AsNumber(Row->PoseIdx, &FNumberFormattingOptions::DefaultNoGrouping());
+		if (Row->IsValid())
+		{
+			return FText::AsNumber(Row->PoseIdx, &FNumberFormattingOptions::DefaultNoGrouping());
+		}
+		return FText::FromString(TEXT("-"));
 	}
 };
 
@@ -150,12 +157,19 @@ struct FAssetIdx : ITextColumn
 
 	virtual FSortPredicate GetSortPredicate() const override
 	{
-		return [](const FRowDataRef& Row0, const FRowDataRef& Row1) -> bool { return Row0->DbAssetIdx < Row1->DbAssetIdx; };
+		return [](const FRowDataRef& Row0, const FRowDataRef& Row1) -> bool
+			{
+				return Row0->DbAssetIdx < Row1->DbAssetIdx;
+			};
 	}
 
 	virtual FText GetRowText(const FRowDataRef& Row) const override
 	{
-		return FText::AsNumber(Row->DbAssetIdx, &FNumberFormattingOptions::DefaultNoGrouping());
+		if (Row->IsValid())
+		{
+			return FText::AsNumber(Row->DbAssetIdx, &FNumberFormattingOptions::DefaultNoGrouping());
+		}
+		return FText::FromString(TEXT("-"));
 	}
 };
 
@@ -182,45 +196,57 @@ struct FDatabaseName : IColumn
 
 	virtual FSortPredicate GetSortPredicate() const override
 	{
-		return [](const FRowDataRef& Row0, const FRowDataRef& Row1) -> bool { return GetNameSafe(Row0->SharedData->SourceDatabase.Get()) < GetNameSafe(Row1->SharedData->SourceDatabase.Get()); };
+		return [](const FRowDataRef& Row0, const FRowDataRef& Row1) -> bool
+			{
+				return GetNameSafe(Row0->SharedData->SourceDatabase.Get()) < GetNameSafe(Row1->SharedData->SourceDatabase.Get());
+			};
 	}
 
 	virtual TSharedRef<SWidget> GenerateWidget(const FRowDataRef& RowData) const override
 	{
-		return SNew(SHyperlink)
-		.Text_Lambda([RowData]() -> FText
+		if (RowData->IsValid())
 		{
-			return FText::FromString(GetNameSafe(RowData->SharedData->SourceDatabase.Get()));
-		})
-		.TextStyle(&FCoreStyle::Get().GetWidgetStyle<FTextBlockStyle>("SmallText"))
-		.ToolTipText_Lambda([RowData]() -> FText
-		{
-			return FText::Format(
-				LOCTEXT("DatabaseHyperlinkTooltipFormat", "Open database '{0}'"),
-				FText::FromString(GetPathNameSafe(RowData->SharedData->SourceDatabase.Get())));
-		})
-		.OnNavigate_Lambda([this, RowData]()
-		{
-			// Open editor
-			if (UPoseSearchDatabase* Database = RowData->SharedData->SourceDatabase.Get())
-			{
-				if (UAssetEditorSubsystem* AssetEditorSS = GEditor->GetEditorSubsystem<UAssetEditorSubsystem>())
-				{
-					AssetEditorSS->OpenEditorForAsset(Database);
-								
-					if (IAssetEditorInstance* Editor = AssetEditorSS->FindEditorForAsset(Database, true))
+			return SNew(SHyperlink)
+				.Text_Lambda([RowData]() -> FText
 					{
-						if (Editor->GetEditorName() == FName("PoseSearchDatabaseEditor"))
+						return FText::FromString(GetNameSafe(RowData->SharedData->SourceDatabase.Get()));
+					})
+				.TextStyle(&FCoreStyle::Get().GetWidgetStyle<FTextBlockStyle>("SmallText"))
+				.ToolTipText_Lambda([RowData]() -> FText
+					{
+						return FText::Format(
+							LOCTEXT("DatabaseHyperlinkTooltipFormat", "Open database '{0}'"),
+							FText::FromString(GetPathNameSafe(RowData->SharedData->SourceDatabase.Get())));
+					})
+				.OnNavigate_Lambda([this, RowData]()
+					{
+						// Open editor
+						if (UPoseSearchDatabase* Database = RowData->SharedData->SourceDatabase.Get())
 						{
-							FDatabaseEditor* DatabaseEditor = static_cast<FDatabaseEditor*>(Editor);
-									
-							// Open asset paused and at specific time as seen on the pose search debugger.
-							DatabaseEditor->SetSelectedPoseIdx(RowData->PoseIdx, DebuggerViewModel->GetDrawQuery(), RowData->SharedData->QueryVector);
+							if (UAssetEditorSubsystem* AssetEditorSS = GEditor->GetEditorSubsystem<UAssetEditorSubsystem>())
+							{
+								AssetEditorSS->OpenEditorForAsset(Database);
+
+								if (IAssetEditorInstance* Editor = AssetEditorSS->FindEditorForAsset(Database, true))
+								{
+									if (Editor->GetEditorName() == FName("PoseSearchDatabaseEditor"))
+									{
+										FDatabaseEditor* DatabaseEditor = static_cast<FDatabaseEditor*>(Editor);
+
+										// Open asset paused and at specific time as seen on the pose search debugger.
+										DatabaseEditor->SetSelectedPoseIdx(RowData->PoseIdx, DebuggerViewModel->GetDrawQuery(), RowData->SharedData->QueryVector);
+									}
+								}
+							}
 						}
-					}
-				}
-			}
-		});
+					});
+		}
+
+		return SNew(STextBlock)
+            .Font(FAppStyle::Get().GetFontStyle("DetailsView.CategoryTextStyle"))
+			.Text(FText::FromString(TEXT("-")))
+            .Justification(ETextJustify::Center)
+			.ColorAndOpacity(FSlateColor(FLinearColor::White));
 	}
 
 	virtual float ComputeColumnWidth(TConstArrayView<FRowDataRef> Rows) const override
@@ -255,61 +281,73 @@ struct FAssetName : IColumn
 
 	virtual FSortPredicate GetSortPredicate() const override
 	{
-		return [](const FRowDataRef& Row0, const FRowDataRef& Row1) -> bool { return Row0->AssetName < Row1->AssetName; };
+		return [](const FRowDataRef& Row0, const FRowDataRef& Row1) -> bool
+			{
+				return Row0->AssetName < Row1->AssetName;
+			};
 	}
 
 	virtual TSharedRef<SWidget> GenerateWidget(const FRowDataRef& RowData) const override
 	{
-		return SNew(SHyperlink)
-			.Text_Lambda([RowData]() -> FText { return FText::FromString(RowData->AssetName); })
-			.TextStyle(&FCoreStyle::Get().GetWidgetStyle<FTextBlockStyle>("SmallText"))
-			.ToolTipText_Lambda([RowData]() -> FText 
-				{ 
-					return FText::Format(
-						LOCTEXT("AssetHyperlinkTooltipFormat", "Open asset '{0}'"), 
-						FText::FromString(RowData->AssetPath)); 
-				})
-			.OnNavigate_Lambda([RowData]()
-				{
-					UObject* Asset = nullptr;
-
-					// Load asset
-					if (UPackage* Package = LoadPackage(NULL, *RowData->AssetPath, LOAD_NoRedirects))
+		if (RowData->IsValid())
+		{
+			return SNew(SHyperlink)
+				.Text_Lambda([RowData]() -> FText { return FText::FromString(RowData->AssetName); })
+				.TextStyle(&FCoreStyle::Get().GetWidgetStyle<FTextBlockStyle>("SmallText"))
+				.ToolTipText_Lambda([RowData]() -> FText
 					{
-						Package->FullyLoad();
-
-						const FString AssetName = FPaths::GetBaseFilename(RowData->AssetPath);
-						Asset = FindObject<UObject>(Package, *AssetName);
-					}
-					else
+						return FText::Format(
+							LOCTEXT("AssetHyperlinkTooltipFormat", "Open asset '{0}'"),
+							FText::FromString(RowData->AssetPath));
+					})
+				.OnNavigate_Lambda([RowData]()
 					{
-						// Fallback for unsaved assets
-						Asset = FindObject<UObject>(nullptr, *RowData->AssetPath);
-					}
+						UObject* Asset = nullptr;
 
-					// Open editor
-					if (Asset != nullptr)
-					{
-						if (UAssetEditorSubsystem* AssetEditorSS = GEditor->GetEditorSubsystem<UAssetEditorSubsystem>())
+						// Load asset
+						if (UPackage* Package = LoadPackage(NULL, *RowData->AssetPath, LOAD_NoRedirects))
 						{
-							AssetEditorSS->OpenEditorForAsset(Asset);
-							
-							if (IAssetEditorInstance* Editor = AssetEditorSS->FindEditorForAsset(Asset, true))
+							Package->FullyLoad();
+
+							const FString AssetName = FPaths::GetBaseFilename(RowData->AssetPath);
+							Asset = FindObject<UObject>(Package, *AssetName);
+						}
+						else
+						{
+							// Fallback for unsaved assets
+							Asset = FindObject<UObject>(nullptr, *RowData->AssetPath);
+						}
+
+						// Open editor
+						if (Asset != nullptr)
+						{
+							if (UAssetEditorSubsystem* AssetEditorSS = GEditor->GetEditorSubsystem<UAssetEditorSubsystem>())
 							{
-								if (Editor->GetEditorName() == "AnimationEditor")
-								{
-									const IAnimationEditor* AnimationEditor = static_cast<IAnimationEditor*>(Editor);
-									const UDebugSkelMeshComponent* PreviewComponent = AnimationEditor->GetPersonaToolkit()->GetPreviewMeshComponent();
+								AssetEditorSS->OpenEditorForAsset(Asset);
 
-									// Open asset paused and at specific time as seen on the pose search debugger.
-									PreviewComponent->PreviewInstance->SetPosition(RowData->AssetTime);
-									PreviewComponent->PreviewInstance->SetPlaying(false);
-									PreviewComponent->PreviewInstance->SetBlendSpacePosition(RowData->BlendParameters);
+								if (IAssetEditorInstance* Editor = AssetEditorSS->FindEditorForAsset(Asset, true))
+								{
+									if (Editor->GetEditorName() == "AnimationEditor")
+									{
+										const IAnimationEditor* AnimationEditor = static_cast<IAnimationEditor*>(Editor);
+										const UDebugSkelMeshComponent* PreviewComponent = AnimationEditor->GetPersonaToolkit()->GetPreviewMeshComponent();
+
+										// Open asset paused and at specific time as seen on the pose search debugger.
+										PreviewComponent->PreviewInstance->SetPosition(RowData->AssetTime);
+										PreviewComponent->PreviewInstance->SetPlaying(false);
+										PreviewComponent->PreviewInstance->SetBlendSpacePosition(RowData->BlendParameters);
+									}
 								}
 							}
 						}
-					}
-				});
+					});
+		}
+
+		return SNew(STextBlock)
+            .Font(FAppStyle::Get().GetFontStyle("DetailsView.CategoryTextStyle"))
+			.Text(LOCTEXT("ColumnAssetNameInvalidPose", "Invalid Pose"))
+            .Justification(ETextJustify::Center)
+			.ColorAndOpacity(FSlateColor(FLinearColor::White));
 	}
 
 	virtual float ComputeColumnWidth(TConstArrayView<FRowDataRef> Rows) const override
@@ -344,12 +382,19 @@ struct FFrame : ITextColumn
 
 	virtual FSortPredicate GetSortPredicate() const override
 	{
-		return [](const FRowDataRef& Row0, const FRowDataRef& Row1) -> bool { return Row0->AnimFrame < Row1->AnimFrame; };
+		return [](const FRowDataRef& Row0, const FRowDataRef& Row1) -> bool
+			{
+				return Row0->AnimFrame < Row1->AnimFrame;
+			};
 	}
 
 	virtual FText GetRowText(const FRowDataRef& Row) const override
 	{
-		return FText::AsNumber(Row->AnimFrame, &FNumberFormattingOptions::DefaultNoGrouping());
+		if (Row->IsValid())
+		{
+			return FText::AsNumber(Row->AnimFrame, &FNumberFormattingOptions::DefaultNoGrouping());
+		}
+		return FText::FromString(TEXT("-"));
 	}
 };
 
@@ -369,12 +414,19 @@ struct FTime : ITextColumn
 
 	virtual FSortPredicate GetSortPredicate() const override
 	{
-		return [](const FRowDataRef& Row0, const FRowDataRef& Row1) -> bool { return Row0->AssetTime < Row1->AssetTime; };
+		return [](const FRowDataRef& Row0, const FRowDataRef& Row1) -> bool
+			{
+				return Row0->AssetTime < Row1->AssetTime;
+			};
 	}
 
 	virtual FText GetRowText(const FRowDataRef& Row) const override
 	{
-		return FText::AsNumber(Row->AssetTime, &FNumberFormattingOptions().SetUseGrouping(false).SetMaximumFractionalDigits(2));
+		if (Row->IsValid())
+		{
+			return FText::AsNumber(Row->AssetTime, &FNumberFormattingOptions().SetUseGrouping(false).SetMaximumFractionalDigits(2));
+		}
+		return FText::FromString(TEXT("-"));
 	}
 };
 
@@ -394,12 +446,19 @@ struct FPercentage : ITextColumn
 
 	virtual FSortPredicate GetSortPredicate() const override
 	{
-		return [](const FRowDataRef& Row0, const FRowDataRef& Row1) -> bool { return Row0->AnimPercentage < Row1->AnimPercentage; };
+		return [](const FRowDataRef& Row0, const FRowDataRef& Row1) -> bool
+			{
+				return Row0->AnimPercentage < Row1->AnimPercentage;
+			};
 	}
 
 	virtual FText GetRowText(const FRowDataRef& Row) const override
 	{
-		return FText::AsPercent(Row->AnimPercentage, &FNumberFormattingOptions().SetMaximumFractionalDigits(2));
+		if (Row->IsValid())
+		{
+			return FText::AsPercent(Row->AnimPercentage, &FNumberFormattingOptions().SetMaximumFractionalDigits(2));
+		}
+		return FText::FromString(TEXT("-"));
 	}
 };
 
@@ -419,19 +478,19 @@ struct FCost : ITextColumn
 
 	virtual FSortPredicate GetSortPredicate() const override
 	{
-		return [](const FRowDataRef& Row0, const FRowDataRef& Row1) -> bool { return Row0->PoseCost < Row1->PoseCost; };
+		return [](const FRowDataRef& Row0, const FRowDataRef& Row1) -> bool
+			{
+				return Row0->PoseCost < Row1->PoseCost;
+			};
 	}
 		
 	virtual FText GetRowText(const FRowDataRef& Row) const override
 	{
-		if (Row->PoseCost.IsValid())
+		if (Row->IsValid() && Row->PoseCost.IsValid())
 		{
 			return FText::AsNumber(Row->PoseCost);
 		}
-		else
-		{
-			return FText::FromString(TEXT("-"));
-		}
+		return FText::FromString(TEXT("-"));
     }
 
 	virtual FSlateColor GetColorAndOpacity(const FRowDataRef& Row) const override
@@ -456,12 +515,19 @@ struct FPCACost : ITextColumn
 
 	virtual FSortPredicate GetSortPredicate() const override
 	{
-		return [](const FRowDataRef& Row0, const FRowDataRef& Row1) -> bool { return Row0->PosePCACost < Row1->PosePCACost; };
+		return [](const FRowDataRef& Row0, const FRowDataRef& Row1) -> bool
+			{
+				return Row0->PosePCACost < Row1->PosePCACost;
+			};
 	}
 		
 	virtual FText GetRowText(const FRowDataRef& Row) const override
 	{
-		return FText::AsNumber(Row->PosePCACost);
+		if (Row->IsValid())
+		{
+			return FText::AsNumber(Row->PosePCACost);
+		}
+		return FText::FromString(TEXT("-"));
     }
 
 	virtual FSlateColor GetColorAndOpacity(const FRowDataRef& Row) const override
@@ -492,24 +558,39 @@ struct FChannelBreakdownCostColumn : ITextColumn
 	virtual FSortPredicate GetSortPredicate() const override
 	{
 		return [this](const FRowDataRef& Row0, const FRowDataRef& Row1) -> bool
-		{
-			return Row0->CostBreakdowns[BreakdownCostIndex] < Row1->CostBreakdowns[BreakdownCostIndex];
-		};
+			{
+				if (!Row0->IsValid())
+				{
+					return true;
+				}
+				if (!Row1->IsValid())
+				{
+					return false;
+				}
+				return Row0->CostBreakdowns[BreakdownCostIndex] < Row1->CostBreakdowns[BreakdownCostIndex];
+			};
 	}
 
 	virtual FText GetRowText(const FRowDataRef& Row) const override
 	{
-		const float LabelCost = Row->CostBreakdowns[BreakdownCostIndex];
-		if (LabelCost != UE_MAX_FLT)
+		if (Row->IsValid())
 		{
-			return FText::AsNumber(LabelCost);
+			const float LabelCost = Row->CostBreakdowns[BreakdownCostIndex];
+			if (LabelCost != UE_MAX_FLT)
+			{
+				return FText::AsNumber(LabelCost);
+			}
 		}
 		return FText::FromString(TEXT("-"));
 	}
 
 	virtual FSlateColor GetColorAndOpacity(const FRowDataRef& Row) const override
 	{
-		return Row->CostBreakdownsColors[BreakdownCostIndex];
+		if (Row->CostBreakdownsColors.IsValidIndex(BreakdownCostIndex))
+		{
+			return Row->CostBreakdownsColors[BreakdownCostIndex];
+		}
+		return FLinearColor::White;
 	}
 
 	FText Label;
@@ -532,12 +613,19 @@ struct FCostModifier : ITextColumn
 
 	virtual FSortPredicate GetSortPredicate() const override
 	{
-		return [](const FRowDataRef& Row0, const FRowDataRef& Row1) -> bool { return Row0->PoseCost.GetCostAddend() < Row1->PoseCost.GetCostAddend(); };
+		return [](const FRowDataRef& Row0, const FRowDataRef& Row1) -> bool
+			{
+				return Row0->PoseCost.GetCostAddend() < Row1->PoseCost.GetCostAddend();
+			};
 	}
 
 	virtual FText GetRowText(const FRowDataRef& Row) const override
 	{
-		return FText::AsNumber(Row->PoseCost.GetCostAddend());
+		if (Row->IsValid())
+		{
+			return FText::AsNumber(Row->PoseCost.GetCostAddend());
+		}
+		return FText::FromString(TEXT("-"));
 	}
 };
 
@@ -557,12 +645,19 @@ struct FContinuingPoseCost : ITextColumn
 
 	virtual FSortPredicate GetSortPredicate() const override
 	{
-		return [](const FRowDataRef& Row0, const FRowDataRef& Row1) -> bool { return Row0->PoseCost.GetContinuingPoseCostAddend() < Row1->PoseCost.GetContinuingPoseCostAddend(); };
+		return [](const FRowDataRef& Row0, const FRowDataRef& Row1) -> bool
+			{
+				return Row0->PoseCost.GetContinuingPoseCostAddend() < Row1->PoseCost.GetContinuingPoseCostAddend();
+			};
 	}
 
 	virtual FText GetRowText(const FRowDataRef& Row) const override
 	{
-		return FText::AsNumber(Row->PoseCost.GetContinuingPoseCostAddend());
+		if (Row->IsValid())
+		{
+			return FText::AsNumber(Row->PoseCost.GetContinuingPoseCostAddend());
+		}
+		return FText::FromString(TEXT("-"));
 	}
 };
 
@@ -582,12 +677,19 @@ struct FContinuingInteractionCost : ITextColumn
 
 	virtual FSortPredicate GetSortPredicate() const override
 	{
-		return [](const FRowDataRef& Row0, const FRowDataRef& Row1) -> bool { return Row0->PoseCost.GetContinuingInteractionCostAddend() < Row1->PoseCost.GetContinuingInteractionCostAddend(); };
+		return [](const FRowDataRef& Row0, const FRowDataRef& Row1) -> bool
+			{
+				return Row0->PoseCost.GetContinuingInteractionCostAddend() < Row1->PoseCost.GetContinuingInteractionCostAddend();
+			};
 	}
 
 	virtual FText GetRowText(const FRowDataRef& Row) const override
 	{
-		return FText::AsNumber(Row->PoseCost.GetContinuingInteractionCostAddend());
+		if (Row->IsValid())
+		{
+			return FText::AsNumber(Row->PoseCost.GetContinuingInteractionCostAddend());
+		}
+		return FText::FromString(TEXT("-"));
 	}
 };
 
@@ -607,12 +709,19 @@ struct FNotifyCost : ITextColumn
 
 	virtual FSortPredicate GetSortPredicate() const override
 	{
-		return [](const FRowDataRef& Row0, const FRowDataRef& Row1) -> bool { return Row0->PoseCost.GetNotifyCostAddend() < Row1->PoseCost.GetNotifyCostAddend(); };
+		return [](const FRowDataRef& Row0, const FRowDataRef& Row1) -> bool
+			{
+				return Row0->PoseCost.GetNotifyCostAddend() < Row1->PoseCost.GetNotifyCostAddend();
+			};
 	}
 
 	virtual FText GetRowText(const FRowDataRef& Row) const override
 	{
-		return FText::AsNumber(Row->PoseCost.GetNotifyCostAddend());
+		if (Row->IsValid())
+		{
+			return FText::AsNumber(Row->PoseCost.GetNotifyCostAddend());
+		}
+		return FText::FromString(TEXT("-"));
 	}
 };
 
@@ -633,12 +742,19 @@ struct FMirrored : ITextColumn
 
 	virtual FSortPredicate GetSortPredicate() const override
 	{
-		return [](const FRowDataRef& Row0, const FRowDataRef& Row1) -> bool { return Row0->bMirrored < Row1->bMirrored; };
+		return [](const FRowDataRef& Row0, const FRowDataRef& Row1) -> bool
+			{
+				return Row0->bMirrored < Row1->bMirrored;
+			};
 	}
 
 	virtual FText GetRowText(const FRowDataRef& Row) const override
 	{
-		return FText::Format(LOCTEXT("Mirrored", "{0}"), { Row->bMirrored } );
+		if (Row->IsValid())
+		{
+			return FText::Format(LOCTEXT("Mirrored", "{0}"), { Row->bMirrored } );
+		}
+		return FText::FromString(TEXT("-"));
 	}
 };
 
@@ -658,12 +774,19 @@ struct FLooping : ITextColumn
 
 	virtual FSortPredicate GetSortPredicate() const override
 	{
-		return [](const FRowDataRef& Row0, const FRowDataRef& Row1) -> bool { return Row0->bLooping < Row1->bLooping; };
+		return [](const FRowDataRef& Row0, const FRowDataRef& Row1) -> bool
+			{
+				return Row0->bLooping < Row1->bLooping;
+			};
 	}
 
 	virtual FText GetRowText(const FRowDataRef& Row) const override
 	{
-		return FText::Format(LOCTEXT("Looping", "{0}"), { Row->bLooping });
+		if (Row->IsValid())
+		{
+			return FText::Format(LOCTEXT("Looping", "{0}"), { Row->bLooping });
+		}
+		return FText::FromString(TEXT("-"));
 	}
 };
 
@@ -693,7 +816,11 @@ struct FBlendParameters : ITextColumn
 
 	virtual FText GetRowText(const FRowDataRef& Row) const override
 	{
-		return FText::Format(LOCTEXT("Blend Parameters", "{0}, {1}"), FText::AsNumber(Row->BlendParameters[0]), FText::AsNumber(Row->BlendParameters[1]));
+		if (Row->IsValid())
+		{
+			return FText::Format(LOCTEXT("Blend Parameters", "{0}, {1}"), FText::AsNumber(Row->BlendParameters[0]), FText::AsNumber(Row->BlendParameters[1]));
+		}
+		return FText::FromString(TEXT("-"));
 	}
 };
 
@@ -713,12 +840,15 @@ struct FPoseCandidateFlags : ITextColumn
 
 	virtual FSortPredicate GetSortPredicate() const override
 	{
-		return [](const FRowDataRef& Row0, const FRowDataRef& Row1) -> bool { return Row0->PoseCandidateFlags < Row1->PoseCandidateFlags; };
+		return [](const FRowDataRef& Row0, const FRowDataRef& Row1) -> bool
+			{
+				return Row0->PoseCandidateFlags < Row1->PoseCandidateFlags;
+			};
 	}
 
 	virtual FText GetRowText(const FRowDataRef& Row) const override
 	{
-		if (EnumHasAnyFlags(Row->PoseCandidateFlags, EPoseCandidateFlags::AnyDiscardedMask))
+		if (Row->IsValid() && EnumHasAnyFlags(Row->PoseCandidateFlags, EPoseCandidateFlags::AnyDiscardedMask))
 		{
 			FString Sring;
 
diff --git a/Engine/Plugins/Animation/PoseSearch/Source/Editor/Private/PoseSearchDebuggerDatabaseRowData.h b/Engine/Plugins/Animation/PoseSearch/Source/Editor/Private/PoseSearchDebuggerDatabaseRowData.h
--- a/Engine/Plugins/Animation/PoseSearch/Source/Editor/Private/PoseSearchDebuggerDatabaseRowData.h
+++ b/Engine/Plugins/Animation/PoseSearch/Source/Editor/Private/PoseSearchDebuggerDatabaseRowData.h
@@ -27,8 +27,9 @@ class FDebuggerDatabaseRowData : public TSharedFromThis<FDebuggerDatabaseRowData
 {
 public:
 	explicit FDebuggerDatabaseRowData(TSharedRef<FDebuggerDatabaseSharedData> InSharedData) : SharedData(InSharedData) {}
+	bool IsValid() const { return PoseIdx != INDEX_NONE; }
 
-	int32 PoseIdx = 0;
+	int32 PoseIdx = INDEX_NONE;
 	EPoseCandidateFlags PoseCandidateFlags = EPoseCandidateFlags::None;
 	FString AssetName = "";
 	FString AssetPath = "";
diff --git a/Engine/Plugins/Animation/PoseSearch/Source/Editor/Private/PoseSearchDebuggerDatabaseView.cpp b/Engine/Plugins/Animation/PoseSearch/Source/Editor/Private/PoseSearchDebuggerDatabaseView.cpp
--- a/Engine/Plugins/Animation/PoseSearch/Source/Editor/Private/PoseSearchDebuggerDatabaseView.cpp
+++ b/Engine/Plugins/Animation/PoseSearch/Source/Editor/Private/PoseSearchDebuggerDatabaseView.cpp
@@ -273,6 +273,9 @@ void SDebuggerDatabaseView::Update(const FTraceMotionMatchingStateMessage& State
 	TArray<uint32> PoseToPCAValuesVectorIndexes;
 	TAlignedArray<float> DynamicWeightsSqrtBuffer;
 
+	TSharedRef<FDebuggerDatabaseSharedData> InvalidDatabaseRowSharedData = MakeShared<FDebuggerDatabaseSharedData>();
+	TSharedRef<FDebuggerDatabaseRowData>& InvalidDatabaseRow = InvalidDatabaseRows.Add_GetRef(MakeShared<FDebuggerDatabaseRowData>(InvalidDatabaseRowSharedData));
+
 	FAsyncPoseSearchDatabasesManagementCache DatabasesManagementCache;
 	for (const FTraceMotionMatchingStateDatabaseEntry& DbEntry : State.DatabaseEntries)
 	{
@@ -815,39 +818,14 @@ void SDebuggerDatabaseView::PopulateViewRows()
 	ContinuingPoseView.ListView->RequestListRefresh();
 	FilteredDatabaseView.ListView->RequestListRefresh();
 
-	if (ActiveView.Rows.Num() > 0)
-	{
-		ReasonForNoActivePose = FText::GetEmpty();
-	}
-	else
-	{
-		ReasonForNoActivePose = LOCTEXT("ReasonForNoActivePose", "Database search didn't find any candidates, or the search has not been performed");
-	}
-
-	if (ContinuingPoseView.Rows.Num() > 0)
+	if (ActiveView.Rows.IsEmpty())
 	{
-		ReasonForNoContinuingPose = FText::GetEmpty();
-	}
-	else
-	{
-		ReasonForNoContinuingPose = LOCTEXT("ReasonForNoContinuingPose", "Invalid continuing pose");
+		ActiveView.Rows.Add(InvalidDatabaseRows[0]);
 	}
 
-	if (FilteredDatabaseView.Rows.Num() > 0)
-	{
-		ReasonForNoCandidates = FText::GetEmpty();
-	}
-	else if (UnfilteredDatabaseRows.Num() == 0)
-	{
-		ReasonForNoCandidates = LOCTEXT("ReasonForNoCandidates_NoSearch", "Database search didn't find any candidates, or the search has not been performed");
-	}
-	else if (UnfilteredDatabaseRows.Num() == 1)
+	if (ContinuingPoseView.Rows.IsEmpty())
 	{
-		ReasonForNoCandidates = LOCTEXT("ReasonForNoCandidates_OnlyContinuingPose", "The continuing pose cost cannot be lowered by searching the databases, so the search has been skipped");
-	}
-	else
-	{
-		ReasonForNoCandidates = FText::Format(LOCTEXT("ReasonForNoCandidates_AllFilteredOut", "All {0} databases poses have been filtered out"), UnfilteredDatabaseRows.Num());
+		ContinuingPoseView.Rows.Add(InvalidDatabaseRows[0]);
 	}
 }
 
@@ -990,35 +968,11 @@ void SDebuggerDatabaseView::Construct(const FArguments& InArgs)
 					.VAlign(VAlign_Fill)
 					.Padding(0.0f)
 					[
-						SNew(SOverlay)
-						+ SOverlay::Slot()
-						[
-							SNew(SBorder)
-							.BorderImage(FAppStyle::GetBrush("NoBorder"))
-							.Padding(0.0f)
-							[
-								ActiveView.ListView.ToSharedRef()
-							]
-						]
-						+ SOverlay::Slot()
+						SNew(SBorder)
+						.BorderImage(FAppStyle::GetBrush("NoBorder"))
+						.Padding(0.0f)
 						[
-							SNew(SBorder)
-							.Visibility(EVisibility::SelfHitTestInvisible)
-							.Padding(FMargin(5.f, 5.f, 5.f, 5.f))
-							.HAlign(HAlign_Center)
-							.VAlign(VAlign_Fill)
-							[
-								SNew(STextBlock)
-								.Visibility_Lambda([this]()
-								{
-									return ReasonForNoActivePose.IsEmpty() ? EVisibility::Collapsed : EVisibility::HitTestInvisible;
-								})
-								.Margin(FMargin(5.f, 5.f, 5.f, 5.f))
-								.Text_Lambda([this]()
-								{
-									return ReasonForNoActivePose;
-								})
-							]
+							ActiveView.ListView.ToSharedRef()
 						]
 					]
 
@@ -1072,35 +1026,12 @@ void SDebuggerDatabaseView::Construct(const FArguments& InArgs)
 					.VAlign(VAlign_Fill)
 					.Padding(0.0f)
 					[
-						SNew(SOverlay)
-						+ SOverlay::Slot()
-						[
-							SNew(SBorder)
-							.BorderImage(FAppStyle::GetBrush("NoBorder"))
-							.Padding(0.0f)
-							[
-								ContinuingPoseView.ListView.ToSharedRef()
-							]
-						]
-						+ SOverlay::Slot()
+						
+						SNew(SBorder)
+						.BorderImage(FAppStyle::GetBrush("NoBorder"))
+						.Padding(0.0f)
 						[
-							SNew(SBorder)
-							.Visibility(EVisibility::SelfHitTestInvisible)
-							.Padding(FMargin(5.f, 5.f, 5.f, 5.f))
-							.HAlign(HAlign_Center)
-							.VAlign(VAlign_Fill)
-							[
-								SNew(STextBlock)
-								.Visibility_Lambda([this]()
-								{
-									return ReasonForNoContinuingPose.IsEmpty() ? EVisibility::Collapsed : EVisibility::HitTestInvisible;
-								})
-								.Margin(FMargin(5.f, 5.f, 5.f, 5.f))
-								.Text_Lambda([this]()
-								{
-									return ReasonForNoContinuingPose;
-								})
-							]
+							ContinuingPoseView.ListView.ToSharedRef()
 						]
 					]
 
@@ -1249,35 +1180,11 @@ void SDebuggerDatabaseView::Construct(const FArguments& InArgs)
 					+ SHorizontalBox::Slot()
 					.Padding(0.0f)
 					[
-						SNew(SOverlay)
-						+SOverlay::Slot()
-						[
-							SNew(SBorder)
-							.BorderImage(FAppStyle::GetBrush("NoBorder"))
-							.Padding(0.0f)
-							[
-								FilteredDatabaseView.ListView.ToSharedRef()
-							]
-						]
-						+ SOverlay::Slot()
+						SNew(SBorder)
+						.BorderImage(FAppStyle::GetBrush("NoBorder"))
+						.Padding(0.0f)
 						[
-							SNew(SBorder)
-							.Visibility(EVisibility::SelfHitTestInvisible)
-							.Padding(FMargin(5.f, 5.f, 5.f, 5.f))
-							.HAlign(HAlign_Center)
-							.VAlign(VAlign_Fill)
-							[
-								SNew(STextBlock)
-								.Visibility_Lambda([this]()
-								{
-									return ReasonForNoCandidates.IsEmpty() ? EVisibility::Collapsed : EVisibility::HitTestInvisible;
-								})
-								.Margin(FMargin(5.f, 5.f, 5.f, 5.f))
-								.Text_Lambda([this]()
-								{
-									return ReasonForNoCandidates;
-								})
-							]
+							FilteredDatabaseView.ListView.ToSharedRef()
 						]
 					]
 				
diff --git a/Engine/Plugins/Animation/PoseSearch/Source/Editor/Private/PoseSearchDebuggerDatabaseView.h b/Engine/Plugins/Animation/PoseSearch/Source/Editor/Private/PoseSearchDebuggerDatabaseView.h
--- a/Engine/Plugins/Animation/PoseSearch/Source/Editor/Private/PoseSearchDebuggerDatabaseView.h
+++ b/Engine/Plugins/Animation/PoseSearch/Source/Editor/Private/PoseSearchDebuggerDatabaseView.h
@@ -148,6 +148,7 @@ private:
 
 	/** All database poses */
 	TArray<TSharedRef<FDebuggerDatabaseRowData>> UnfilteredDatabaseRows;
+	TArray<TSharedRef<FDebuggerDatabaseRowData>> InvalidDatabaseRows;
 
 	/** Database listing for filtered poses */
 	FTable FilteredDatabaseView;
@@ -157,10 +158,6 @@ private:
 
 	/** Text used to filter DatabaseView */
 	FText FilterText;
-
-	FText ReasonForNoActivePose;
-	FText ReasonForNoContinuingPose;
-	FText ReasonForNoCandidates;
 };
 
 } // namespace UE::PoseSearch
