PoseSearch - updated FAnimationAssetSampler::ExtractAnimNotifyStates UE_LOG description, previously referencing a deprecated API name

diff --git a/Engine/Plugins/Animation/PoseSearch/Source/Runtime/Private/PoseSearchAssetSampler.cpp b/Engine/Plugins/Animation/PoseSearch/Source/Runtime/Private/PoseSearchAssetSampler.cpp
--- a/Engine/Plugins/Animation/PoseSearch/Source/Runtime/Private/PoseSearchAssetSampler.cpp
+++ b/Engine/Plugins/Animation/PoseSearch/Source/Runtime/Private/PoseSearchAssetSampler.cpp
@@ -672,7 +672,7 @@ void FAnimationAssetSampler::ExtractAnimNotifyStates(float Time, FAnimNotifyCont
 		}
 		else
 		{
-			UE_LOG(LogPoseSearch, Error, TEXT("FAnimationAssetSampler::ExtractPoseSearchNotifyStates: Unsupported BlendSpace NotifyTriggerMode for '%s'"), *BlendSpace->GetName());
+			UE_LOG(LogPoseSearch, Error, TEXT("FAnimationAssetSampler::ExtractAnimNotifyStates: Unsupported BlendSpace NotifyTriggerMode for '%s'"), *BlendSpace->GetName());
 		}
 	}
 	else if (const UAnimSequenceBase* SequenceBase = Cast<UAnimSequenceBase>(AnimationAssetPtr.Get()))
