PoseSearch - fix for deprecated FSearchContext::FSearchContext not propagating correctly its input InEventToSearch

#rnx

diff --git a/Engine/Plugins/Animation/PoseSearch/Source/Runtime/Private/PoseSearchContext.cpp b/Engine/Plugins/Animation/PoseSearch/Source/Runtime/Private/PoseSearchContext.cpp
--- a/Engine/Plugins/Animation/PoseSearch/Source/Runtime/Private/PoseSearchContext.cpp
+++ b/Engine/Plugins/Animation/PoseSearch/Source/Runtime/Private/PoseSearchContext.cpp
@@ -849,7 +849,7 @@ FSearchContext::FPoseCandidate FSearchContext::FBestPoseCandidates::GetUnsortedC
 // FSearchContext
 FSearchContext::FSearchContext(float InDesiredPermutationTimeOffset, const FPoseIndicesHistory* InPoseIndicesHistory,
 	const FSearchResult& InCurrentResult, const FFloatInterval& InPoseJumpThresholdTime, const FPoseSearchEvent& InEventToSearch)
-: FSearchContext(InDesiredPermutationTimeOffset, InPoseJumpThresholdTime, FPoseSearchEvent())
+: FSearchContext(InDesiredPermutationTimeOffset, InPoseJumpThresholdTime, InEventToSearch)
 {
 	UpdateContinuingPoseSearchResult(InCurrentResult, InCurrentResult);
 }
