PoseSearch - fix for missing InvalidDatabaseRows.Reset() in SDebuggerDatabaseView::Update

#rnx

diff --git a/Engine/Plugins/Animation/PoseSearch/Source/Editor/Private/PoseSearchDebuggerDatabaseView.cpp b/Engine/Plugins/Animation/PoseSearch/Source/Editor/Private/PoseSearchDebuggerDatabaseView.cpp
--- a/Engine/Plugins/Animation/PoseSearch/Source/Editor/Private/PoseSearchDebuggerDatabaseView.cpp
+++ b/Engine/Plugins/Animation/PoseSearch/Source/Editor/Private/PoseSearchDebuggerDatabaseView.cpp
@@ -273,8 +273,8 @@ void SDebuggerDatabaseView::Update(const FTraceMotionMatchingStateMessage& State
 	TArray<uint32> PoseToPCAValuesVectorIndexes;
 	TAlignedArray<float> DynamicWeightsSqrtBuffer;
 
-	TSharedRef<FDebuggerDatabaseSharedData> InvalidDatabaseRowSharedData = MakeShared<FDebuggerDatabaseSharedData>();
-	TSharedRef<FDebuggerDatabaseRowData>& InvalidDatabaseRow = InvalidDatabaseRows.Add_GetRef(MakeShared<FDebuggerDatabaseRowData>(InvalidDatabaseRowSharedData));
+	InvalidDatabaseRows.Reset();
+	InvalidDatabaseRows.Add(MakeShared<FDebuggerDatabaseRowData>(MakeShared<FDebuggerDatabaseSharedData>()));
 
 	FAsyncPoseSearchDatabasesManagementCache DatabasesManagementCache;
 	for (const FTraceMotionMatchingStateDatabaseEntry& DbEntry : State.DatabaseEntries)
