PoseSearch - adding support for UPoseSearchInteractionAsset::MinimalTranslationWeight in UPoseSearchInteractionAsset::CalculateWarpTransforms to calculate the FullAlignedActorRootBoneTransforms containing the minimal possible translation. This solution can potentially rotate all the interacting actors and can be turned off by setting MinimalTranslationWeight to 0 in case one of the involved actors transform must stay unchanged

diff --git a/Engine/Plugins/Animation/PoseSearch/Source/Editor/Private/PoseSearchInteractionAssetEditor.cpp b/Engine/Plugins/Animation/PoseSearch/Source/Editor/Private/PoseSearchInteractionAssetEditor.cpp
--- a/Engine/Plugins/Animation/PoseSearch/Source/Editor/Private/PoseSearchInteractionAssetEditor.cpp
+++ b/Engine/Plugins/Animation/PoseSearch/Source/Editor/Private/PoseSearchInteractionAssetEditor.cpp
@@ -279,8 +279,9 @@ void FInteractionAssetViewModel::Tick(float DeltaSeconds)
 
 	PlayTime += DeltaSeconds * DeltaTimeMultiplier;
 
+	const int32 NumRoles = InteractionAsset->GetNumRoles();
 	FRoleToIndex InteractionAssetRoleToIndex;
-	for (int32 RoleIndex = 0; RoleIndex < InteractionAsset->GetNumRoles(); ++RoleIndex)
+	for (int32 RoleIndex = 0; RoleIndex < NumRoles; ++RoleIndex)
 	{
 		InteractionAssetRoleToIndex.Add(InteractionAsset->GetRole(RoleIndex)) = RoleIndex;
 	}
@@ -303,7 +304,7 @@ void FInteractionAssetViewModel::Tick(float DeltaSeconds)
 
 	if (PreviewActors.Num() != InteractionAssetRoleToIndex.Num())
 	{
-		for (int32 RoleIndex = 0; RoleIndex < InteractionAsset->GetNumRoles(); ++RoleIndex)
+		for (int32 RoleIndex = 0; RoleIndex < NumRoles; ++RoleIndex)
 		{
 			if (!PreviewActorsRoleToIndex.Find(InteractionAsset->GetRole(RoleIndex)))
 			{
@@ -330,11 +331,13 @@ void FInteractionAssetViewModel::Tick(float DeltaSeconds)
 
 	// testing CalculateWarpTransforms
 #if WITH_EDITORONLY_DATA
-	if (InteractionAsset->bEnableDebugWarp && PreviewActors.Num() == InteractionAsset->GetNumRoles())
+	if (InteractionAsset->bEnableDebugWarp && PreviewActors.Num() == NumRoles)
 	{
 		TArray<FTransform, TInlineAllocator<PreallocatedRolesNum>> ActorTransforms;
-		ActorTransforms.SetNum(InteractionAsset->GetNumRoles());
-		for (int32 InteractionAssetRoleIndex = 0; InteractionAssetRoleIndex < InteractionAsset->GetNumRoles(); ++InteractionAssetRoleIndex)
+		TArray<TObjectPtr<const UObject>, TInlineAllocator<PreallocatedRolesNum>> DebugAnimContexts;
+		ActorTransforms.SetNum(NumRoles);
+		DebugAnimContexts.SetNum(NumRoles);
+		for (int32 InteractionAssetRoleIndex = 0; InteractionAssetRoleIndex < NumRoles; ++InteractionAssetRoleIndex)
 		{
 			const FRole& InteractionAssetRole = InteractionAsset->GetRole(InteractionAssetRoleIndex);
 			const int32 PreviewActorIndex = PreviewActorRoleToIndex[InteractionAssetRole];
@@ -343,13 +346,16 @@ void FInteractionAssetViewModel::Tick(float DeltaSeconds)
 			check(PreviewActor.GetRole() == InteractionAssetRole);
 
 			ActorTransforms[InteractionAssetRoleIndex] = PreviewActor.GetDebugActorTransformFromSampler();
+			DebugAnimContexts[InteractionAssetRoleIndex] = PreviewActor.GetAnimPreviewInstance();
 		}
 
 		TArray<FTransform, TInlineAllocator<PreallocatedRolesNum>> FullAlignedActorTransforms;
-		FullAlignedActorTransforms.SetNum(InteractionAsset->GetNumRoles());
-		InteractionAsset->CalculateWarpTransforms(PlayTime, ActorTransforms, FullAlignedActorTransforms);
+		FullAlignedActorTransforms.SetNum(NumRoles);
+		const TConstArrayView<const UMirrorDataTable*> MirrorDataTables;
+		const TConstArrayView<bool> RelevantRoleIndexes;
+		InteractionAsset->CalculateWarpTransforms(PlayTime, ActorTransforms, FullAlignedActorTransforms, MirrorDataTables, RelevantRoleIndexes, DebugAnimContexts);
 
-		for (int32 InteractionAssetRoleIndex = 0; InteractionAssetRoleIndex < InteractionAsset->GetNumRoles(); ++InteractionAssetRoleIndex)
+		for (int32 InteractionAssetRoleIndex = 0; InteractionAssetRoleIndex < NumRoles; ++InteractionAssetRoleIndex)
 		{
 			const FRole& InteractionAssetRole = InteractionAsset->GetRole(InteractionAssetRoleIndex);
 			const int32 PreviewActorIndex = PreviewActorRoleToIndex[InteractionAssetRole];
diff --git a/Engine/Plugins/Animation/PoseSearch/Source/Runtime/Private/PoseSearchInteractionAsset.cpp b/Engine/Plugins/Animation/PoseSearch/Source/Runtime/Private/PoseSearchInteractionAsset.cpp
--- a/Engine/Plugins/Animation/PoseSearch/Source/Runtime/Private/PoseSearchInteractionAsset.cpp
+++ b/Engine/Plugins/Animation/PoseSearch/Source/Runtime/Private/PoseSearchInteractionAsset.cpp
@@ -6,6 +6,10 @@
 #include "PoseSearch/PoseSearchAssetSampler.h"
 #include "PoseSearch/PoseSearchMirrorDataCache.h"
 
+#if ENABLE_ANIM_DEBUG && WITH_EDITOR
+#include "DrawDebugHelpers.h"
+#endif // ENABLE_ANIM_DEBUG && WITH_EDITOR
+
 #include UE_INLINE_GENERATED_CPP_BY_NAME(PoseSearchInteractionAsset)
 
 bool UPoseSearchInteractionAsset::IsLooping() const
@@ -249,7 +253,8 @@ USkeletalMesh* UPoseSearchInteractionAsset::GetPreviewMesh(const UE::PoseSearch:
 
 #endif // WITH_EDITOR
 
-void UPoseSearchInteractionAsset::CalculateWarpTransforms(float Time, const TConstArrayView<const FTransform> ActorRootBoneTransforms, TArrayView<FTransform> FullAlignedActorRootBoneTransforms, const TConstArrayView<const UMirrorDataTable*> MirrorDataTables, const TConstArrayView<bool> RelevantRoleIndexes) const
+void UPoseSearchInteractionAsset::CalculateWarpTransforms(float Time, const TConstArrayView<const FTransform> ActorRootBoneTransforms, TArrayView<FTransform> FullAlignedActorRootBoneTransforms,
+	const TConstArrayView<const UMirrorDataTable*> MirrorDataTables, const TConstArrayView<bool> RelevantRoleIndexes, const TConstArrayView<TObjectPtr<const UObject>> DebugAnimContexts) const
 {
 	using namespace UE::PoseSearch;
 
@@ -285,8 +290,7 @@ void UPoseSearchInteractionAsset::CalculateWarpTransforms(float Time, const TCon
 	}
 	else
 	{
-		// we have at least one relevant item!
-
+		// we have at least two relevant items!
 		TArray<FTransform, TInlineAllocator<PreallocatedRolesNum>> AssetRootBoneTransforms;
 		AssetRootBoneTransforms.SetNum(ItemsNum);
 
@@ -467,46 +471,136 @@ void UPoseSearchInteractionAsset::CalculateWarpTransforms(float Time, const TCon
 				}
 			}
 		}
+		const FVector AssetReferencePosition = FindReferencePosition(AssetRootBoneTransforms, NormalizedWarpingWeightTranslation);
+		const FVector ActorsReferencePosition = FindReferencePosition(ActorRootBoneTransforms, NormalizedWarpingWeightTranslation);
 
-		const FQuat AssetReferenceOrientation = FindReferenceOrientation(AssetRootBoneTransforms, SortedByWarpingWeightRotationItemIndex, NormalizedWarpingWeightRotation);
-		const FQuat ActorsReferenceOrientation = FindReferenceOrientation(ActorRootBoneTransforms, SortedByWarpingWeightRotationItemIndex, NormalizedWarpingWeightRotation);
+#if ENABLE_ANIM_DEBUG && WITH_EDITOR
+		const UWorld* DebugWorld = nullptr;
+		for (const TObjectPtr<const UObject>& DebugAnimContextPtr : DebugAnimContexts)
+		{
+			if (const UObject* Outer = DebugAnimContextPtr->GetOuter())
+			{
+				DebugWorld = Outer->GetWorld();
+				if (DebugWorld)
+				{
+					break;
+				}
+			}
+		}
+#endif // ENABLE_ANIM_DEBUG && WITH_EDITOR
 
-		FQuat WeightedActorsReferenceOrientation = ActorsReferenceOrientation;
-		if (WarpingWeightRotationSum > UE_KINDA_SMALL_NUMBER)
+		if (MinimalTranslationWeight > 0.f)
 		{
-			// ItemIndex are in order of WarpingWeightRotation. the last one is the one with the highest most WarpingWeightRotation, the most "important"
-			for (int32 ItemIndex : SortedByWarpingWeightRotationItemIndex)
+			// finding FullAlignedActorRootBoneTransforms with the minimal possible translations
+			for (int32 ItemIndex = 0; ItemIndex < ItemsNum; ++ItemIndex)
 			{
-				const FPoseSearchInteractionAssetItem& Item = Items[ItemIndex];
-				if (NormalizedWarpingWeightRotation[ItemIndex] > NormalizedHomogeneousWeight)
+				const FVector ActorRootBonePosition = ActorRootBoneTransforms[ItemIndex].GetLocation();
+				const FVector AssetRootBonePosition = AssetRootBoneTransforms[ItemIndex].GetLocation();
+				
+				const FVector AssetReferencePositionLocal = AssetRootBoneTransforms[ItemIndex].InverseTransformPosition(AssetReferencePosition);
+				const FVector ActorsReferencePositionFromAsset = ActorRootBoneTransforms[ItemIndex].TransformPosition(AssetReferencePositionLocal);
+
+				const FVector ActorsReferenceDisplacementFromAsset = ActorsReferencePositionFromAsset - ActorRootBonePosition;
+				const FVector ActorsReferenceDisplacement = ActorsReferencePosition - ActorRootBonePosition;
+
+				const float ActorsReferenceDisplacementFromAssetLength = ActorsReferenceDisplacementFromAsset.Length();
+				const float ActorsReferenceDisplacementLength = ActorsReferenceDisplacement.Length();
+
+				if (ActorsReferenceDisplacementFromAssetLength > UE_SMALL_NUMBER && ActorsReferenceDisplacementLength > UE_SMALL_NUMBER)
 				{
-					// NormalizedHomogeneousWeight is one only if ItemsNum is one, 
-					// BUT NormalizedWarpingWeightRotation[ItemIndex] > NormalizedHomogeneousWeight should always be false
-					check(!FMath::IsNearlyEqual(NormalizedHomogeneousWeight, 1.f));
+					const FVector ActorsReferenceDisplacementWithReducedBanking(ActorsReferenceDisplacement.X, ActorsReferenceDisplacement.Y, 
+						FMath::Lerp(ActorsReferenceDisplacementFromAsset.Z, ActorsReferenceDisplacement.Z, WarpingBankingWeight));
+					const FQuat DeltaRotation = FQuat::FindBetweenVectors(ActorsReferenceDisplacementFromAsset, ActorsReferenceDisplacementWithReducedBanking);
 
-					// how much this item wants to reorient the ReferenceOrientation from the homogeneous "fair" value
-					const float SlerpParam = (NormalizedWarpingWeightRotation[ItemIndex] - NormalizedHomogeneousWeight) / (1.f - NormalizedHomogeneousWeight);
+					FullAlignedActorRootBoneTransforms[ItemIndex].SetLocation(ActorsReferencePosition - ActorsReferenceDisplacement * (ActorsReferenceDisplacementFromAssetLength / ActorsReferenceDisplacementLength));
+					FullAlignedActorRootBoneTransforms[ItemIndex].SetRotation(ActorRootBoneTransforms[ItemIndex].GetRotation() * DeltaRotation);
+					FullAlignedActorRootBoneTransforms[ItemIndex].SetScale3D(ActorRootBoneTransforms[ItemIndex].GetScale3D());
 
-					// calculating the reference orientation relative to the character
-					// AssetReferenceOrientation in actor world orientation
-					const FQuat ActorAssetReferenceOrientation = ActorRootBoneTransforms[ItemIndex].GetRotation() * (AssetRootBoneTransforms[ItemIndex].GetRotation().Inverse() * AssetReferenceOrientation);
+#if ENABLE_ANIM_DEBUG && WITH_EDITOR
+					DrawDebugLine(DebugWorld, ActorsReferencePosition, ActorRootBonePosition, FColorList::Blue, false, 0.f, SDPG_Foreground, 0.f);
+					DrawDebugLine(DebugWorld, ActorsReferencePositionFromAsset, ActorRootBonePosition, FColorList::Red, false, 0.f, SDPG_Foreground, 0.f);
+					const FVector ReconstructedForDebug = DeltaRotation.RotateVector(ActorsReferenceDisplacementFromAsset) + ActorRootBonePosition;
+					DrawDebugCircle(DebugWorld, FMatrix(FVector::ZAxisVector, FVector::XAxisVector, FVector::YAxisVector, ReconstructedForDebug), 5.f, 16, FColorList::Red, false, 0.f, SDPG_Foreground, 0.f, false);
+
+					// asset actor drawing 
+					DrawDebugCircle(DebugWorld, FMatrix(FVector::ZAxisVector, FVector::XAxisVector, FVector::YAxisVector, AssetRootBonePosition), 5.f, 16, FColorList::Green, false, 0.f, SDPG_Foreground, 0.f, false);
+					DrawDebugLine(DebugWorld, AssetRootBonePosition, AssetRootBonePosition + AssetRootBoneTransforms[ItemIndex].GetRotation().RotateVector(FVector::XAxisVector) * 10, FColorList::Green, false, 0.f, SDPG_Foreground, 0.f);
 
-					WeightedActorsReferenceOrientation = FQuat::Slerp(WeightedActorsReferenceOrientation, ActorAssetReferenceOrientation, SlerpParam);
+					DrawDebugCircle(DebugWorld, FMatrix(FVector::ZAxisVector, FVector::XAxisVector, FVector::YAxisVector, ActorRootBonePosition), 5.f, 16, FColorList::Green, false, 0.f, SDPG_Foreground, 0.f, false);
+					DrawDebugLine(DebugWorld, ActorRootBonePosition, ActorRootBonePosition + ActorRootBoneTransforms[ItemIndex].GetRotation().RotateVector(FVector::XAxisVector) * 10, FColorList::Green, false, 0.f, SDPG_Foreground, 0.f);
+#endif // ENABLE_ANIM_DEBUG && WITH_EDITOR
+				}
+				else
+				{
+					UE_LOG(LogPoseSearch, Error, TEXT("UPoseSearchInteractionAsset::CalculateWarpTransforms actors are too close to each other to be able to calculate warp transforms: distance in the world %f, distance in the asset %f ['%s' at time %f]"), ActorsReferenceDisplacementLength, ActorsReferenceDisplacementFromAssetLength, *GetName(), Time);
+					FullAlignedActorRootBoneTransforms[ItemIndex] = ActorRootBoneTransforms[ItemIndex];
 				}
 			}
+
+#if ENABLE_ANIM_DEBUG && WITH_EDITOR
+			DrawDebugCircle(DebugWorld, FMatrix(FVector::ZAxisVector, FVector::XAxisVector, FVector::YAxisVector, AssetReferencePosition), 5.f, 16, FColorList::Blue, false, 0.f, SDPG_Foreground, 0.f, false);
+			DrawDebugCircle(DebugWorld, FMatrix(FVector::ZAxisVector, FVector::XAxisVector, FVector::YAxisVector, ActorsReferencePosition), 5.f, 16, FColorList::Blue, false, 0.f, SDPG_Foreground, 0.f, false);
+#endif // ENABLE_ANIM_DEBUG && WITH_EDITOR
 		}
+		
+		if (MinimalTranslationWeight < 1.f)
+		{
+			const FQuat AssetReferenceOrientation = FindReferenceOrientation(AssetRootBoneTransforms, SortedByWarpingWeightRotationItemIndex, NormalizedWarpingWeightRotation);
+			const FQuat ActorsReferenceOrientation = FindReferenceOrientation(ActorRootBoneTransforms, SortedByWarpingWeightRotationItemIndex, NormalizedWarpingWeightRotation);
 
-		const FVector AssetReferencePosition = FindReferencePosition(AssetRootBoneTransforms, NormalizedWarpingWeightTranslation);
-		const FVector ActorsReferencePosition = FindReferencePosition(ActorRootBoneTransforms, NormalizedWarpingWeightTranslation);
+			FQuat WeightedActorsReferenceOrientation = ActorsReferenceOrientation;
+			if (WarpingWeightRotationSum > UE_KINDA_SMALL_NUMBER)
+			{
+				// ItemIndex are in order of WarpingWeightRotation. the last one is the one with the highest most WarpingWeightRotation, the most "important"
+				for (int32 ItemIndex : SortedByWarpingWeightRotationItemIndex)
+				{
+					const FPoseSearchInteractionAssetItem& Item = Items[ItemIndex];
+					if (NormalizedWarpingWeightRotation[ItemIndex] > NormalizedHomogeneousWeight)
+					{
+						// NormalizedHomogeneousWeight is one only if ItemsNum is one, 
+						// BUT NormalizedWarpingWeightRotation[ItemIndex] > NormalizedHomogeneousWeight should always be false
+						check(!FMath::IsNearlyEqual(NormalizedHomogeneousWeight, 1.f));
 
-		// aligning all the actors to ActorsReferencePosition, WeightedActorsReferenceOrientation
-		const FTransform AssetReferenceTransform(AssetReferenceOrientation, AssetReferencePosition);
-		const FTransform ActorsReferenceTransform(WeightedActorsReferenceOrientation, ActorsReferencePosition);
-		const FTransform AssetReferenceInverseTransform = AssetReferenceTransform.Inverse();
+						// how much this item wants to reorient the ReferenceOrientation from the homogeneous "fair" value
+						const float SlerpParam = (NormalizedWarpingWeightRotation[ItemIndex] - NormalizedHomogeneousWeight) / (1.f - NormalizedHomogeneousWeight);
 
-		for (int32 ItemIndex = 0; ItemIndex < ItemsNum; ++ItemIndex)
-		{
-			FullAlignedActorRootBoneTransforms[ItemIndex] = (AssetRootBoneTransforms[ItemIndex] * AssetReferenceInverseTransform) * ActorsReferenceTransform;
-		}
+						// calculating the reference orientation relative to the character
+						// AssetReferenceOrientation in actor world orientation
+						const FQuat ActorAssetReferenceOrientation = ActorRootBoneTransforms[ItemIndex].GetRotation() * (AssetRootBoneTransforms[ItemIndex].GetRotation().Inverse() * AssetReferenceOrientation);
+
+						WeightedActorsReferenceOrientation = FQuat::Slerp(WeightedActorsReferenceOrientation, ActorAssetReferenceOrientation, SlerpParam);
+					}
+				}
+			}
+
+			// aligning all the actors to ActorsReferencePosition, WeightedActorsReferenceOrientation
+			const FTransform AssetReferenceTransform(AssetReferenceOrientation, AssetReferencePosition);
+			const FTransform ActorsReferenceTransform(WeightedActorsReferenceOrientation, ActorsReferencePosition);
+			const FTransform AssetReferenceInverseTransform = AssetReferenceTransform.Inverse();
+
+			// blend FullAlignedActorRootBoneTransforms with what has been calcualted by finding FullAlignedActorRootBoneTransforms with the minimal possible translations
+			for (int32 ItemIndex = 0; ItemIndex < ItemsNum; ++ItemIndex)
+			{
+				const FTransform WantedFullAlignedActorRootBoneTransforms = (AssetRootBoneTransforms[ItemIndex] * AssetReferenceInverseTransform) * ActorsReferenceTransform;
+				if (MinimalTranslationWeight > 0.f)
+				{
+					FullAlignedActorRootBoneTransforms[ItemIndex].SetLocation(FMath::Lerp<FVector>(WantedFullAlignedActorRootBoneTransforms.GetLocation(), FullAlignedActorRootBoneTransforms[ItemIndex].GetLocation(), MinimalTranslationWeight));
+					FullAlignedActorRootBoneTransforms[ItemIndex].SetRotation(FQuat::Slerp(WantedFullAlignedActorRootBoneTransforms.GetRotation(), FullAlignedActorRootBoneTransforms[ItemIndex].GetRotation(), MinimalTranslationWeight));
+					FullAlignedActorRootBoneTransforms[ItemIndex].SetScale3D(FMath::Lerp<FVector>(WantedFullAlignedActorRootBoneTransforms.GetScale3D(), FullAlignedActorRootBoneTransforms[ItemIndex].GetScale3D(), MinimalTranslationWeight));
+				}
+				else
+				{
+					FullAlignedActorRootBoneTransforms[ItemIndex] = WantedFullAlignedActorRootBoneTransforms;
+				}
+			}
+
+#if ENABLE_ANIM_DEBUG && WITH_EDITOR
+			DrawDebugCircle(DebugWorld, FMatrix(FVector::ZAxisVector, FVector::XAxisVector, FVector::YAxisVector, AssetReferencePosition), 5.f, 16, FColorList::Blue, false, 0.f, SDPG_Foreground, 0.f, false);
+			DrawDebugCircle(DebugWorld, FMatrix(FVector::ZAxisVector, FVector::XAxisVector, FVector::YAxisVector, ActorsReferencePosition), 5.f, 16, FColorList::Blue, false, 0.f, SDPG_Foreground, 0.f, false);
+
+			DrawDebugLine(DebugWorld, AssetReferencePosition, AssetReferencePosition + ActorsReferenceOrientation.RotateVector(FVector::XAxisVector) * 10.f, FColorList::Blue, false, 0.f, SDPG_Foreground, 0.f);
+			DrawDebugLine(DebugWorld, ActorsReferencePosition, ActorsReferencePosition + ActorsReferenceOrientation.RotateVector(FVector::XAxisVector) * 10.f, FColorList::Blue, false, 0.f, SDPG_Foreground, 0.f);
+#endif // ENABLE_ANIM_DEBUG && WITH_EDITOR
+		}						
 	}
 }
diff --git a/Engine/Plugins/Animation/PoseSearch/Source/Runtime/Public/PoseSearch/MultiAnimAsset.h b/Engine/Plugins/Animation/PoseSearch/Source/Runtime/Public/PoseSearch/MultiAnimAsset.h
--- a/Engine/Plugins/Animation/PoseSearch/Source/Runtime/Public/PoseSearch/MultiAnimAsset.h
+++ b/Engine/Plugins/Animation/PoseSearch/Source/Runtime/Public/PoseSearch/MultiAnimAsset.h
@@ -32,7 +32,9 @@ public:
 	[[nodiscard]] virtual UAnimationAsset* GetAnimationAsset(const FName& Role) const PURE_VIRTUAL(UMultiAnimAsset::GetAnimationAsset, return nullptr;);
 	[[nodiscard]] virtual FTransform GetOrigin(const FName& Role) const PURE_VIRTUAL(UMultiAnimAsset::GetOrigin, return FTransform::Identity;);
 
-	virtual void CalculateWarpTransforms(float Time, const TConstArrayView<const FTransform> ActorRootBoneTransforms, TArrayView<FTransform> FullAlignedActorRootBoneTransforms, const TConstArrayView<const UMirrorDataTable*> MirrorDataTables = TConstArrayView<const UMirrorDataTable*>(), const TConstArrayView<bool> RelevantRoleIndexes = TArrayView<bool>()) const PURE_VIRTUAL(UMultiAnimAsset::CalculateWarpTransforms, );
+	virtual void CalculateWarpTransforms(float Time, const TConstArrayView<const FTransform> ActorRootBoneTransforms, TArrayView<FTransform> FullAlignedActorRootBoneTransforms,
+		const TConstArrayView<const UMirrorDataTable*> MirrorDataTables = TConstArrayView<const UMirrorDataTable*>(), const TConstArrayView<bool> RelevantRoleIndexes = TConstArrayView<bool>(),
+		const TConstArrayView<TObjectPtr<const UObject>> DebugAnimContexts = TConstArrayView<TObjectPtr<const UObject>>()) const PURE_VIRTUAL(UMultiAnimAsset::CalculateWarpTransforms, );
 
 	UFUNCTION(BlueprintPure, Category = "Animation", meta=(BlueprintThreadSafe, DisplayName = "Get Animation Asset"))
 	UAnimationAsset* BP_GetAnimationAsset(const FName& Role) const { return GetAnimationAsset(Role); }
diff --git a/Engine/Plugins/Animation/PoseSearch/Source/Runtime/Public/PoseSearch/PoseSearchInteractionAsset.h b/Engine/Plugins/Animation/PoseSearch/Source/Runtime/Public/PoseSearch/PoseSearchInteractionAsset.h
--- a/Engine/Plugins/Animation/PoseSearch/Source/Runtime/Public/PoseSearch/PoseSearchInteractionAsset.h
+++ b/Engine/Plugins/Animation/PoseSearch/Source/Runtime/Public/PoseSearch/PoseSearchInteractionAsset.h
@@ -56,6 +56,9 @@ class UPoseSearchInteractionAsset : public UMultiAnimAsset
 	UPROPERTY(EditAnywhere, Category = "Settings", meta = (ClampMin = "0", ClampMax = "1"))
 	float WarpingBankingWeight = 0.f;
 
+	UPROPERTY(EditAnywhere, Category = "Settings", meta = (ClampMin = "0", ClampMax = "1"))
+	float MinimalTranslationWeight = 0.f;
+
 public:
 #if WITH_EDITORONLY_DATA
 
@@ -90,12 +93,15 @@ public:
 	UE_API virtual USkeletalMesh* GetPreviewMesh(const UE::PoseSearch::FRole& Role) const override;
 #endif // WITH_EDITOR
 
-	UE_API virtual void CalculateWarpTransforms(float Time, const TConstArrayView<const FTransform> ActorRootBoneTransforms, TArrayView<FTransform> FullAlignedActorRootBoneTransforms, const TConstArrayView<const UMirrorDataTable*> MirrorDataTables = TConstArrayView<const UMirrorDataTable*>(), const TConstArrayView<bool> RelevantRoleIndexes = TConstArrayView<bool>()) const override;
+	UE_API virtual void CalculateWarpTransforms(float Time, const TConstArrayView<const FTransform> ActorRootBoneTransforms, TArrayView<FTransform> FullAlignedActorRootBoneTransforms,
+		const TConstArrayView<const UMirrorDataTable*> MirrorDataTables = TConstArrayView<const UMirrorDataTable*>(), const TConstArrayView<bool> RelevantRoleIndexes = TConstArrayView<bool>(),
+		const TConstArrayView<TObjectPtr<const UObject>> DebugAnimContexts = TConstArrayView<TObjectPtr<const UObject>>()) const override;
 	
-	UE_API FVector FindReferencePosition(const TArrayView<const FTransform> Transforms, const TArrayView<float> NormalizedWarpingWeightTranslation) const;
-	UE_API FQuat FindReferenceOrientation(const TArrayView<const FTransform> Transforms, const TArrayView<int32> SortedByWarpingWeightRotationItemIndex, const TArrayView<float> NormalizedWarpingWeightRotation) const;
-	UE_API FQuat FindReferenceOrientationNoBanking(const TArrayView<const FTransform> Transforms, const TArrayView<int32> SortedByWarpingWeightRotationItemIndex, const TArrayView<float> NormalizedWarpingWeightRotation) const;
-	UE_API FQuat FindReferenceOrientationFullBanking(const TArrayView<const FTransform> Transforms, const TArrayView<int32> SortedByWarpingWeightRotationItemIndex, const TArrayView<float> NormalizedWarpingWeightRotation) const;
+private:
+	FVector FindReferencePosition(const TArrayView<const FTransform> Transforms, const TArrayView<float> NormalizedWarpingWeightTranslation) const;
+	FQuat FindReferenceOrientation(const TArrayView<const FTransform> Transforms, const TArrayView<int32> SortedByWarpingWeightRotationItemIndex, const TArrayView<float> NormalizedWarpingWeightRotation) const;
+	FQuat FindReferenceOrientationNoBanking(const TArrayView<const FTransform> Transforms, const TArrayView<int32> SortedByWarpingWeightRotationItemIndex, const TArrayView<float> NormalizedWarpingWeightRotation) const;
+	FQuat FindReferenceOrientationFullBanking(const TArrayView<const FTransform> Transforms, const TArrayView<int32> SortedByWarpingWeightRotationItemIndex, const TArrayView<float> NormalizedWarpingWeightRotation) const;
 };
 
 #undef UE_API
