PoseSearch - adding missing Role to UPoseSearchFeatureChannel_Trajectory::GetLabel

diff --git a/Engine/Plugins/Animation/PoseSearch/Source/Runtime/Private/PoseSearchFeatureChannel_Trajectory.cpp b/Engine/Plugins/Animation/PoseSearch/Source/Runtime/Private/PoseSearchFeatureChannel_Trajectory.cpp
--- a/Engine/Plugins/Animation/PoseSearch/Source/Runtime/Private/PoseSearchFeatureChannel_Trajectory.cpp
+++ b/Engine/Plugins/Animation/PoseSearch/Source/Runtime/Private/PoseSearchFeatureChannel_Trajectory.cpp
@@ -197,13 +197,24 @@ void UPoseSearchFeatureChannel_Trajectory::DebugDraw(const UE::PoseSearch::FDebu
 #if WITH_EDITOR
 UE::PoseSearch::TLabelBuilder& UPoseSearchFeatureChannel_Trajectory::GetLabel(UE::PoseSearch::TLabelBuilder& LabelBuilder, UE::PoseSearch::ELabelFormat LabelFormat) const
 {
+	using namespace UE::PoseSearch;
+
 	GetOuterLabel(LabelBuilder, LabelFormat);
 	AppendLabelSeparator(LabelBuilder, UE::PoseSearch::ELabelFormat::Full_Horizontal);
 	LabelBuilder.Append(TEXT("Traj"));
+	
+	if (SampleRole != DefaultRole)
+	{
+		LabelBuilder.Append(TEXT("["));
+		LabelBuilder.Append(SampleRole.ToString());
+		LabelBuilder.Append(TEXT("]"));
+	}
+
 	if (DebugWeightGroupID != INDEX_NONE)
 	{
 		LabelBuilder.Appendf(TEXT("_%d"), DebugWeightGroupID);
 	}
+	
 	return LabelBuilder;
 }
 #endif // WITH_EDITOR
