PoseSearch - adding the best search cost to the search selection widget

diff --git a/Engine/Plugins/Animation/PoseSearch/Source/Editor/Private/PoseSearchDebuggerView.cpp b/Engine/Plugins/Animation/PoseSearch/Source/Editor/Private/PoseSearchDebuggerView.cpp
--- a/Engine/Plugins/Animation/PoseSearch/Source/Editor/Private/PoseSearchDebuggerView.cpp
+++ b/Engine/Plugins/Animation/PoseSearch/Source/Editor/Private/PoseSearchDebuggerView.cpp
@@ -30,9 +30,9 @@
 namespace UE::PoseSearch
 {
 
-FText GenerateSearchName(const FTraceMotionMatchingStateMessage& MotionMatchingState, const IGameplayProvider* GameplayProvider)
+typedef TStringBuilder<256> FStringBuilder256;
+void GenerateSearchName(const FTraceMotionMatchingStateMessage& MotionMatchingState, const IGameplayProvider* GameplayProvider, FStringBuilder256& StringBuilder)
 {
-	TStringBuilder<256> StringBuilder;
 	if (GameplayProvider)
 	{
 		bool bAddDatabasesSeparator = false;
@@ -87,7 +87,30 @@ FText GenerateSearchName(const FTraceMotionMatchingStateMessage& MotionMatchingS
 	{
 		StringBuilder.Append("Error!");
 	}
+}
+
+FText GenerateSearchName(const FTraceMotionMatchingStateMessage& MotionMatchingState, const IGameplayProvider* GameplayProvider)
+{
+	FStringBuilder256 StringBuilder;
+	GenerateSearchName(MotionMatchingState, GameplayProvider, StringBuilder);
+	return FText::FromString(StringBuilder.ToString());
+}
+
+FText GenerateSearchNameWithBestScore(const FTraceMotionMatchingStateMessage& MotionMatchingState, const IGameplayProvider* GameplayProvider)
+{
+	FStringBuilder256 StringBuilder;
+	GenerateSearchName(MotionMatchingState, GameplayProvider, StringBuilder);
+
+	float BestCost = UE_MAX_FLT;
+	for (const FTraceMotionMatchingStateDatabaseEntry& DbEntry : MotionMatchingState.DatabaseEntries)
+	{
+		for (const FTraceMotionMatchingStatePoseEntry& PoseEntry : DbEntry.PoseEntries)
+		{
+			BestCost = FMath::Min(PoseEntry.Cost, BestCost);
+		}
+	}
 
+	StringBuilder.Appendf(TEXT(" [%.3f]"), BestCost);
 	return FText::FromString(StringBuilder.ToString());
 }
 
@@ -499,7 +522,7 @@ bool SDebuggerView::UpdateNodeSelection()
 				.Padding(10.0f)
 				[
 					SNew(SButton)
-						.Text(GenerateSearchName(MotionMatchingState, GameplayProvider))
+						.Text(GenerateSearchNameWithBestScore(MotionMatchingState, GameplayProvider))
 						.HAlign(HAlign_Center)
 						.VAlign(VAlign_Center)
 						.ContentPadding(10.0f)
