PoseSearch - adding error tolerance on missing or wrongly named pose history on MM interaction availabilities to support CBP setups

#rnx

diff --git a/Engine/Plugins/Animation/PoseSearch/Source/Runtime/Private/PoseSearchInteractionSubsystem.cpp b/Engine/Plugins/Animation/PoseSearch/Source/Runtime/Private/PoseSearchInteractionSubsystem.cpp
--- a/Engine/Plugins/Animation/PoseSearch/Source/Runtime/Private/PoseSearchInteractionSubsystem.cpp
+++ b/Engine/Plugins/Animation/PoseSearch/Source/Runtime/Private/PoseSearchInteractionSubsystem.cpp
@@ -151,7 +151,7 @@ const UE::PoseSearch::IPoseHistory* FInteractionAvailability::GetPoseHistory(con
 		}
 	}
 
-	unimplemented();
+	UE_LOG(LogPoseSearch, Error, TEXT("FInteractionAvailability::GetPoseHistory couldn't find a pose history for AnimContext %s with name %s"), *GetNameSafe(AnimContext), *PoseHistoryName.ToString());
 	return nullptr;
 }
 
@@ -1097,37 +1097,37 @@ void UPoseSearchInteractionSubsystem::GenerateSearchContexts(float DeltaSeconds,
 			// inserted in a per database sorted data structure (DatabaseToRoledAnimContextInfos)
 			for (const FInteractionAvailability& Availability : AnimContextInfo.GetAvailabilities())
 			{
-				const IPoseHistory* PoseHistory = Availability.GetPoseHistory(AnimContextInfo.GetAnimContext());
-				check(PoseHistory);
-				
-				if (const UPoseSearchDatabase* Database = Availability.Database.Get())
+				if (const IPoseHistory* PoseHistory = Availability.GetPoseHistory(AnimContextInfo.GetAnimContext()))
 				{
-					check(Database->Schema);
-					DatabaseToRoledAnimContextInfos.AddRoledAnimContextInfos(Availability, AnimContextInfo, PoseHistory, Database);
-				}
-				else if (Availability.IsTagValid())
-				{
-					// since Database is null, but this availability has a valid Tag, we're looking for valid databases by Availability.Tag
-					if (const FDatabasesPerTag* DatabasesPerTag = TagToDatabases.Find(Availability.Tag))
+					if (const UPoseSearchDatabase* Database = Availability.Database.Get())
 					{
-						check(!DatabasesPerTag->IsEmpty());
-						for (const UPoseSearchDatabase* DatabaseFromTag : *DatabasesPerTag)
+						check(Database->Schema);
+						DatabaseToRoledAnimContextInfos.AddRoledAnimContextInfos(Availability, AnimContextInfo, PoseHistory, Database);
+					}
+					else if (Availability.IsTagValid())
+					{
+						// since Database is null, but this availability has a valid Tag, we're looking for valid databases by Availability.Tag
+						if (const FDatabasesPerTag* DatabasesPerTag = TagToDatabases.Find(Availability.Tag))
+						{
+							check(!DatabasesPerTag->IsEmpty());
+							for (const UPoseSearchDatabase* DatabaseFromTag : *DatabasesPerTag)
+							{
+								DatabaseToRoledAnimContextInfos.AddRoledAnimContextInfos(Availability, AnimContextInfo, PoseHistory, DatabaseFromTag);
+							}
+						}
+						else
 						{
-							DatabaseToRoledAnimContextInfos.AddRoledAnimContextInfos(Availability, AnimContextInfo, PoseHistory, DatabaseFromTag);
+							//@todo: should we add a verbose LOG here? not sure since it'd be very spammy...
+
+							// this is a valid condition we shouldn't log: for example when the "main character" is loaded and publishing availabilities with a valid tag and null database,
+							// looking for other NPC / seconday characters to interact with, but they are not present of didn't publish any availability
 						}
 					}
 					else
 					{
-						//@todo: should we add a verbose LOG here? not sure since it'd be very spammy...
-							
-						// this is a valid condition we shouldn't log: for example when the "main character" is loaded and publishing availabilities with a valid tag and null database,
-						// looking for other NPC / seconday characters to interact with, but they are not present of didn't publish any availability
+						UE_LOG(LogPoseSearch, Log, TEXT("UPoseSearchInteractionSubsystem::GenerateSearchContexts null Availability.Database (with invalid Availability.Tag)"));
 					}
 				}
-				else
-				{
-					UE_LOG(LogPoseSearch, Log, TEXT("UPoseSearchInteractionSubsystem::GenerateSearchContexts null Availability.Database (with invalid Availability.Tag)"));
-				}
 			}
 		},
 		// OnDoneGroupingAnimContexts: called when the FAnimContextInfoVisitor reaches the end of a group of possibly interacting AnimContext(s)
