PoseSearch - adding experimental UPoseSearchLibrary::ClearPoseSearchDatabasesManagement to clear the PSD DDC data

#rnx

diff --git a/Engine/Plugins/Animation/PoseSearch/Source/Runtime/Private/PoseSearchDerivedData.cpp b/Engine/Plugins/Animation/PoseSearch/Source/Runtime/Private/PoseSearchDerivedData.cpp
--- a/Engine/Plugins/Animation/PoseSearch/Source/Runtime/Private/PoseSearchDerivedData.cpp
+++ b/Engine/Plugins/Animation/PoseSearch/Source/Runtime/Private/PoseSearchDerivedData.cpp
@@ -372,6 +372,27 @@ static void FindValidRanges(const FPoseSearchDatabaseAnimationAssetBase* Databas
 	}
 }
 
+// returns false in case of errors
+static bool IsSchemaCompatible(const FPoseSearchDatabaseAnimationAssetBase* DatabaseAsset, const UPoseSearchSchema* Schema)
+{
+	const int32 DatabaseAssetNumRoles = DatabaseAsset->GetNumRoles();
+	if (DatabaseAssetNumRoles != Schema->GetRoledSkeletons().Num())
+	{
+		return false;
+	}
+
+	for (int32 DatabaseAssetRoleIndex = 0; DatabaseAssetRoleIndex < DatabaseAssetNumRoles; ++DatabaseAssetRoleIndex)
+	{
+		const FRole DatabaseAssetRole = DatabaseAsset->GetRole(DatabaseAssetRoleIndex);
+		if (!Schema->GetRoledSkeleton(DatabaseAssetRole))
+		{
+			return false;
+		}
+	}
+	
+	return true;
+}
+
 // returns false in case of errors
 static bool InitSearchIndexAssets(FSearchIndexBase& SearchIndex, const UPoseSearchDatabase* DatabaseToLookForAssets, const UPoseSearchSchema* Schema, const FFloatInterval& ExcludeFromDatabaseParameters)
 {
@@ -1895,11 +1916,10 @@ void FPoseSearchDatabaseAsyncCacheTask::OnGetComplete(UE::DerivedData::FCacheGet
 								ResetSearchIndex();
 								return;
 							}
-							// @todo: Find a way to prevent accessing UObjects while GC since this func is ran async. Commenting out for now to prevent crashes.
-							/*else if (!DatabaseAsset->IsSkeletonCompatible(Database->Schema))
+							else if (!IsSchemaCompatible(DatabaseAsset, Database->Schema))
 							{
-								UE_LOG(LogPoseSearch, Warning, TEXT("OnGetComplete - %s's skeleton is not compatible with the schema's skeleton(s)."), *DatabaseAsset->GetName());
-							}*/
+								UE_LOG(LogPoseSearch, Warning, TEXT("OnGetComplete - %s is not compatible with the schema %s"), *DatabaseAsset->GetName(), *GetNameSafe(Database->Schema));
+							}
 							else if (const UBlendSpace* BlendSpace = Cast<UBlendSpace>(DatabaseAsset->GetAnimationAsset()))
 							{
 								if (!BlendSpace->bShouldMatchSyncPhases)
@@ -2832,6 +2852,19 @@ EAsyncBuildIndexResult FAsyncPoseSearchDatabasesManagement::RequestAsyncBuildInd
 	return RequestAsyncBuildIndexInternal(Database, Flag);
 }
 
+void FAsyncPoseSearchDatabasesManagement::Clear()
+{
+	FScopeLock Lock(&Mutex);
+
+	FAsyncPoseSearchDatabasesManagement& This = FAsyncPoseSearchDatabasesManagement::Get();
+
+	for (TUniquePtr<FPoseSearchDatabaseAsyncCacheTask>& TaskPtr : This.Tasks)
+	{
+		UE_LOG(LogPoseSearch, Log, TEXT("%s - %s Cancelled because of Clear"), *LexToString(TaskPtr->GetDerivedDataKey()), *TaskPtr->GetDatabase()->GetName());
+	}
+	This.Tasks.Reset();
+}
+
 //////////////////////////////////////////////////////////////////////////
 // FAsyncPoseSearchDatabasesManagementCache
 EAsyncBuildIndexResult FAsyncPoseSearchDatabasesManagementCache::operator()(const UPoseSearchDatabase* Database, ERequestAsyncBuildFlag RequestAsyncBuildFlag)
diff --git a/Engine/Plugins/Animation/PoseSearch/Source/Runtime/Private/PoseSearchLibrary.cpp b/Engine/Plugins/Animation/PoseSearch/Source/Runtime/Private/PoseSearchLibrary.cpp
--- a/Engine/Plugins/Animation/PoseSearch/Source/Runtime/Private/PoseSearchLibrary.cpp
+++ b/Engine/Plugins/Animation/PoseSearch/Source/Runtime/Private/PoseSearchLibrary.cpp
@@ -1452,4 +1452,13 @@ const FAnimNode_PoseSearchHistoryCollector_Base* UPoseSearchLibrary::FindPoseHis
 	return nullptr;
 }
 
+void UPoseSearchLibrary::ClearPoseSearchDatabasesManagement()
+{
+#if WITH_EDITOR
+	using namespace UE::PoseSearch;
+
+	FAsyncPoseSearchDatabasesManagement::Clear();
+#endif // WITH_EDITOR
+}
+
 #undef LOCTEXT_NAMESPACE
diff --git a/Engine/Plugins/Animation/PoseSearch/Source/Runtime/Public/PoseSearch/PoseSearchDerivedData.h b/Engine/Plugins/Animation/PoseSearch/Source/Runtime/Public/PoseSearch/PoseSearchDerivedData.h
--- a/Engine/Plugins/Animation/PoseSearch/Source/Runtime/Public/PoseSearch/PoseSearchDerivedData.h
+++ b/Engine/Plugins/Animation/PoseSearch/Source/Runtime/Public/PoseSearch/PoseSearchDerivedData.h
@@ -37,6 +37,8 @@ namespace UE::PoseSearch
 	{
 	public:
 		static UE_API EAsyncBuildIndexResult RequestAsyncBuildIndex(const UPoseSearchDatabase* Database, ERequestAsyncBuildFlag Flag);
+		
+		static void Clear();
 
 	private:
 		FAsyncPoseSearchDatabasesManagement();
diff --git a/Engine/Plugins/Animation/PoseSearch/Source/Runtime/Public/PoseSearch/PoseSearchLibrary.h b/Engine/Plugins/Animation/PoseSearch/Source/Runtime/Public/PoseSearch/PoseSearchLibrary.h
--- a/Engine/Plugins/Animation/PoseSearch/Source/Runtime/Public/PoseSearch/PoseSearchLibrary.h
+++ b/Engine/Plugins/Animation/PoseSearch/Source/Runtime/Public/PoseSearch/PoseSearchLibrary.h
@@ -285,6 +285,9 @@ public:
 
 	UFUNCTION(BlueprintPure, Category = "Animation|Pose Search|Experimental", meta = (BlueprintThreadSafe))
 	static UE_API void GetDatabaseTags(const UPoseSearchDatabase* Database, TArray<FName>& Tags);
+
+	UFUNCTION(BlueprintCallable, Category = "Animation|Pose Search|Experimental", meta = (BlueprintThreadSafe))
+	static UE_API void ClearPoseSearchDatabasesManagement();
 };
 
 #undef UE_API
