PoseSearch - adding RoleIndex to FPoseSearchBlueprintResult

#rnx

diff --git a/Engine/Plugins/Animation/PoseSearch/Source/Runtime/Private/PoseSearchInteractionIsland.cpp b/Engine/Plugins/Animation/PoseSearch/Source/Runtime/Private/PoseSearchInteractionIsland.cpp
--- a/Engine/Plugins/Animation/PoseSearch/Source/Runtime/Private/PoseSearchInteractionIsland.cpp
+++ b/Engine/Plugins/Animation/PoseSearch/Source/Runtime/Private/PoseSearchInteractionIsland.cpp
@@ -1299,7 +1299,13 @@ bool FInteractionIsland::GetResult_AnyThread(const UObject* AnimContext, FPoseSe
 
 					for (int32 MultiAnimAssetRoleIndex = 0; MultiAnimAssetRoleIndex < NumRoles; ++MultiAnimAssetRoleIndex)
 					{
-						if (const int32* InteractionSearchContextRoleIndex = InteractionSearchContextRoleToIndex.Find(MultiAnimAsset->GetRole(MultiAnimAssetRoleIndex)))
+						const FRole MultiAnimAssetRole = MultiAnimAsset->GetRole(MultiAnimAssetRoleIndex);
+						if (MultiAnimAssetRole == Result.Role)
+						{
+							Result.RoleIndex = MultiAnimAssetRoleIndex;
+						}
+
+						if (const int32* InteractionSearchContextRoleIndex = InteractionSearchContextRoleToIndex.Find(MultiAnimAssetRole))
 						{
 							Result.ActorRootTransforms[MultiAnimAssetRoleIndex] = SearchResult.ActorRootTransforms[*InteractionSearchContextRoleIndex];
 							Result.ActorRootBoneTransforms[MultiAnimAssetRoleIndex] = SearchResult.ActorRootBoneTransforms[*InteractionSearchContextRoleIndex];
@@ -1323,6 +1329,8 @@ bool FInteractionIsland::GetResult_AnyThread(const UObject* AnimContext, FPoseSe
 					
 					Result.AnimContexts.SetNum(1);
 					Result.AnimContexts[0] = SearchContext.GetAnimContext(0);
+
+					Result.RoleIndex = 0;
 				}
 
 				// we found our AnimContext: we can stop searching
diff --git a/Engine/Plugins/Animation/PoseSearch/Source/Runtime/Public/PoseSearch/PoseSearchResult.h b/Engine/Plugins/Animation/PoseSearch/Source/Runtime/Public/PoseSearch/PoseSearchResult.h
--- a/Engine/Plugins/Animation/PoseSearch/Source/Runtime/Public/PoseSearch/PoseSearchResult.h
+++ b/Engine/Plugins/Animation/PoseSearch/Source/Runtime/Public/PoseSearch/PoseSearchResult.h
@@ -238,6 +238,11 @@ struct FPoseSearchBlueprintResult
 	UPROPERTY(Transient, VisibleAnywhere, BlueprintReadOnly, Category="Experimental|Interaction", meta = (Experimental))
 	FName Role = UE::PoseSearch::DefaultRole;
 
+	// assigned role index to this character. If this result is from a motion matched interaction, every interacting Actors will be assigned a different Role index.
+	// It can be used to find this actor properies from ActorRootTransforms, ActorRootBoneTransforms, and AnimContexts
+	UPROPERTY(Transient, VisibleAnywhere, BlueprintReadOnly, Category="Experimental|Interaction", meta = (Experimental))
+	int32 RoleIndex = INDEX_NONE;
+
 	// cached actors root transforms for all the roles in SelectedAnim (as UMultiAnimAsset). For all the optional missing roles (where the associated AnimContexts[i] is null), the ActorRootTransform will be identity
 	UPROPERTY(Transient, VisibleAnywhere, BlueprintReadOnly, Category="Experimental|Interaction", meta = (Experimental))
 	TArray<FTransform> ActorRootTransforms;
