PoseSearch - exposing CalculateFullAlignedTransform(s) to BP

#rnx

diff --git a/Engine/Plugins/Animation/PoseSearch/Source/Runtime/Private/PoseSearchInteractionLibrary.cpp b/Engine/Plugins/Animation/PoseSearch/Source/Runtime/Private/PoseSearchInteractionLibrary.cpp
--- a/Engine/Plugins/Animation/PoseSearch/Source/Runtime/Private/PoseSearchInteractionLibrary.cpp
+++ b/Engine/Plugins/Animation/PoseSearch/Source/Runtime/Private/PoseSearchInteractionLibrary.cpp
@@ -4,8 +4,10 @@
 #include "Animation/AnimMontage.h"
 #include "Engine/World.h"
 #include "GameFramework/Pawn.h"
+#include "PoseSearch/MultiAnimAsset.h"
 #include "PoseSearch/PoseSearchDatabase.h"
 #include "PoseSearch/PoseSearchInteractionSubsystem.h"
+#include "PoseSearch/PoseSearchInteractionUtils.h"
 
 #include UE_INLINE_GENERATED_CPP_BY_NAME(PoseSearchInteractionLibrary)
 
@@ -96,6 +98,36 @@ bool UPoseSearchInteractionLibrary::MotionMatchInteraction(FPoseSearchBlueprintR
 	return bIsInteraction;
 }
 
+void UPoseSearchInteractionLibrary::CalculateFullAlignedTransforms(const FPoseSearchBlueprintResult& Result, bool bWarpUsingRootBone, TArray<FTransform>& OutFullAlignedTransforms)
+{
+	if (!Result.AnimContexts.IsEmpty() && Result.SelectedAnim && Result.SelectedAnim->IsA<UMultiAnimAsset>())
+	{
+		OutFullAlignedTransforms.SetNum(Result.AnimContexts.Num());
+		UE::PoseSearch::CalculateFullAlignedTransforms(Result, bWarpUsingRootBone, OutFullAlignedTransforms);
+	}
+	else
+	{
+		UE_LOG(LogPoseSearch, Error, TEXT("UPoseSearchInteractionLibrary::CalculateFullAlignedTransforms failed!"));
+		OutFullAlignedTransforms.Reset();
+	}
+}
+
+void UPoseSearchInteractionLibrary::CalculateFullAlignedTransform(const FPoseSearchBlueprintResult& Result, bool bWarpUsingRootBone, FTransform& OutFullAlignedTransform)
+{
+	if (Result.AnimContexts.IsValidIndex(Result.RoleIndex) && Result.SelectedAnim && Result.SelectedAnim->IsA<UMultiAnimAsset>())
+	{
+		TArray<FTransform, TInlineAllocator<UE::PoseSearch::PreallocatedRolesNum>> FullAlignedTransforms;
+		FullAlignedTransforms.SetNum(Result.AnimContexts.Num());
+		UE::PoseSearch::CalculateFullAlignedTransforms(Result, bWarpUsingRootBone, FullAlignedTransforms);
+		OutFullAlignedTransform = FullAlignedTransforms[Result.RoleIndex];
+	}
+	else
+	{
+		UE_LOG(LogPoseSearch, Error, TEXT("UPoseSearchInteractionLibrary::CalculateFullAlignedTransform failed!"));
+		OutFullAlignedTransform.SetIdentity();
+	}
+}
+
 FPoseSearchBlueprintResult UPoseSearchInteractionLibrary::GetCachedInteraction(const UObject* AnimContext, bool bCompareOwningActors)
 {
 	FPoseSearchBlueprintResult Result;
diff --git a/Engine/Plugins/Animation/PoseSearch/Source/Runtime/Public/PoseSearch/PoseSearchInteractionLibrary.h b/Engine/Plugins/Animation/PoseSearch/Source/Runtime/Public/PoseSearch/PoseSearchInteractionLibrary.h
--- a/Engine/Plugins/Animation/PoseSearch/Source/Runtime/Public/PoseSearch/PoseSearchInteractionLibrary.h
+++ b/Engine/Plugins/Animation/PoseSearch/Source/Runtime/Public/PoseSearch/PoseSearchInteractionLibrary.h
@@ -33,6 +33,12 @@ public:
 	UFUNCTION(Experimental, BlueprintCallable, Category = "Animation|Pose Search", meta = (BlueprintThreadSafe))
 	static UE_API FPoseSearchBlueprintResult MotionMatchInteraction(TArray<FPoseSearchInteractionAvailability> Availabilities, const UObject* AnimContext, FName PoseHistoryName, bool bValidateResultAgainstAvailabilities = true);
 
+	UFUNCTION(Experimental, BlueprintCallable, Category = "Animation|Pose Search", meta = (BlueprintThreadSafe))
+	static UE_API void CalculateFullAlignedTransforms(const FPoseSearchBlueprintResult& CurrentResult, bool bWarpUsingRootBone, TArray<FTransform>& OutFullAlignedTransforms);
+
+	UFUNCTION(Experimental, BlueprintCallable, Category = "Animation|Pose Search", meta = (BlueprintThreadSafe))
+	static UE_API void CalculateFullAlignedTransform(const FPoseSearchBlueprintResult& CurrentResult, bool bWarpUsingRootBone, FTransform& OutFullAlignedTransform);
+
 	static UE_API void MotionMatchInteraction(FPoseSearchBlueprintResult& Result, const TArrayView<const FPoseSearchInteractionAvailability> Availabilities, const UObject* AnimContext, FName PoseHistoryName, const UE::PoseSearch::IPoseHistory* PoseHistory, bool bValidateResultAgainstAvailabilities = true);
 	static UE_API bool MotionMatchInteraction(FPoseSearchBlueprintResult& InOutResult, const TArrayView<const FPoseSearchInteractionAvailability> Availabilities, const UObject* AnimContext, FName PoseHistoryName, const UE::PoseSearch::IPoseHistory* PoseHistory, bool bValidateResultAgainstAvailabilities, bool bKeepInteractionAlive, float BlendTime);
 
