PoseSearch - support for CBP MMI searches: automatically adding required tick dependencies between PoseSearchInteractionIsland(s) and actors

#rnx

diff --git a/Engine/Plugins/Animation/PoseSearch/Source/Runtime/Private/PoseSearchHistory.cpp b/Engine/Plugins/Animation/PoseSearch/Source/Runtime/Private/PoseSearchHistory.cpp
--- a/Engine/Plugins/Animation/PoseSearch/Source/Runtime/Private/PoseSearchHistory.cpp
+++ b/Engine/Plugins/Animation/PoseSearch/Source/Runtime/Private/PoseSearchHistory.cpp
@@ -894,6 +894,18 @@ const FPoseIndicesHistory* FPoseHistory::GetPoseIndicesHistory() const
 	return &PoseIndicesHistory;
 }
 
+int32 FPoseHistory::GetMaxNumPoses() const
+{
+	UE_MT_SCOPED_READ_ACCESS(PoseDataThreadSafeCounter);
+	return MaxNumPoses;
+}
+
+float FPoseHistory::GetSamplingInterval() const
+{
+	UE_MT_SCOPED_READ_ACCESS(PoseDataThreadSafeCounter);
+	return SamplingInterval;
+}
+
 void FPoseHistory::GenerateTrajectory(const UObject* AnimContext, float DeltaTime, const FPoseSearchTrajectoryData& TrajectoryData, const FPoseSearchTrajectoryData::FSampling& TrajectoryDataSampling)
 {
 	// @todo: Synchronize the FPoseSearchQueryTrajectorySample::AccumulatedSeconds of the generated trajectory with the FPoseHistoryEntry::AccumulatedSeconds of the captured poses
diff --git a/Engine/Plugins/Animation/PoseSearch/Source/Runtime/Private/PoseSearchInteractionIsland.cpp b/Engine/Plugins/Animation/PoseSearch/Source/Runtime/Private/PoseSearchInteractionIsland.cpp
--- a/Engine/Plugins/Animation/PoseSearch/Source/Runtime/Private/PoseSearchInteractionIsland.cpp
+++ b/Engine/Plugins/Animation/PoseSearch/Source/Runtime/Private/PoseSearchInteractionIsland.cpp
@@ -766,11 +766,14 @@ void FInteractionIsland::AddTickDependencies(UActorComponent* TickActorComponent
 {
 	check(TickActorComponent);
 
+	AActor* TickActor = TickActorComponent->GetOwner();
+	check(TickActor && TickActor->CanEverTick());
+			
 	if (IInteractionIslandDependency* IslandDependency = FindCustomDependency(TickActorComponent))
 	{
 		if (const FTickFunction* TickActorComponentTickFunction = IslandDependency->FindTickFunction(TickActorComponent))
 		{
-			if(bInIsMainActor)
+			if (bInIsMainActor)
 			{
 				// PostTickFunction prerequisites should be empty since we haven't add the main actor tick function yet
 				check(PostTickFunction.GetPrerequisites().IsEmpty());
@@ -781,24 +784,30 @@ void FInteractionIsland::AddTickDependencies(UActorComponent* TickActorComponent
 					AddPrerequisite(PreTickFunction, TickActorComponentPrerequisite.PrerequisiteObject.Get(), *TickActorComponentPrerequisite.PrerequisiteTickFunction);
 				}
 
-				// Add post-tick function dependencies
+				// adding to TickActorComponent and owner actor TickActor dependencies to the PreTickFunction
 				IslandDependency->AddSubsequent(InteractionSubsystem, PreTickFunction, TickActorComponent);
+				AddPrerequisite(TickActor->PrimaryActorTick, InteractionSubsystem, PreTickFunction);
+
+				// add post-tick function dependencies to TickActorComponent and owner actor
 				IslandDependency->AddPrerequisite(InteractionSubsystem, PostTickFunction, TickActorComponent);
+				AddPrerequisite(PostTickFunction, TickActor, TickActor->PrimaryActorTick);
 
 				check(!bHasTickDependencies);
 				bHasTickDependencies = true;
 			}
 			else
 			{
-				// PostTickFunction should contain only the tick function to the main actor's one
-				check(PostTickFunction.GetPrerequisites().Num());
+				// PostTickFunction should contain only the tick function to the main actor's ones (component and owner)
+				check(PostTickFunction.GetPrerequisites().Num() == 2);
 
 				// adding to PreTickFunction all the tick dependencies TickActorComponent has, so it runs after all the tick dependencies of ALL the TickActorComponents in this FInteractionIsland
-				// BUT excluding the main actor tick fuction that is PostTickFunction.GetPrerequisites()[0]
-				const FTickFunction* MainActorTickFunction = PostTickFunction.GetPrerequisites()[0].Get();
+				// BUT excluding the main actor tick fuctions that are in PostTickFunction.GetPrerequisites()
 				for (const FTickPrerequisite& TickActorComponentPrerequisite : TickActorComponentTickFunction->GetPrerequisites())
 				{
-					if (TickActorComponentPrerequisite.PrerequisiteTickFunction != MainActorTickFunction)
+					if (!PostTickFunction.GetPrerequisites().ContainsByPredicate([&TickActorComponentPrerequisite](const FTickPrerequisite& MainActorTickPrerequisite)
+						{
+							return TickActorComponentPrerequisite.PrerequisiteTickFunction == MainActorTickPrerequisite.Get();
+						}))
 					{
 						AddPrerequisite(PreTickFunction, TickActorComponentPrerequisite.PrerequisiteObject.Get(), *TickActorComponentPrerequisite.PrerequisiteTickFunction);
 					}
@@ -806,8 +815,12 @@ void FInteractionIsland::AddTickDependencies(UActorComponent* TickActorComponent
 
 				// Add post-tick function dependencies
 				IslandDependency->AddSubsequent(InteractionSubsystem, PostTickFunction, TickActorComponent);
+				AddPrerequisite(TickActor->PrimaryActorTick, InteractionSubsystem, PostTickFunction);
 				check(bHasTickDependencies);
 			}
+
+			// adding dependency between actor and UAF TickActorComponent so CBP runs before UAF parallel execution jobs
+			IslandDependency->AddSubsequent(InteractionSubsystem, TickActor->PrimaryActorTick, TickActorComponent);
 		}
 		else
 		{
@@ -827,35 +840,43 @@ void FInteractionIsland::AddTickDependencies(UActorComponent* TickActorComponent
 				AddPrerequisite(PreTickFunction, TickActorComponentPrerequisite.PrerequisiteObject.Get(), *TickActorComponentPrerequisite.PrerequisiteTickFunction);
 			}
 
-			// Add post-tick function dependencies
-			// it should be this island, but it's not a UObject, so we use the InteractionSubsystem
+			// adding to TickActorComponent and owner actor TickActor dependencies to the PreTickFunction
 			AddPrerequisite(TickActorComponent->PrimaryComponentTick, InteractionSubsystem, PreTickFunction);
+			AddPrerequisite(TickActor->PrimaryActorTick, InteractionSubsystem, PreTickFunction);
+
+			// add post-tick function dependencies to TickActorComponent and owner actor
 			AddPrerequisite(PostTickFunction, TickActorComponent, TickActorComponent->PrimaryComponentTick);
+			AddPrerequisite(PostTickFunction, TickActor, TickActor->PrimaryActorTick);
 
 			check(!bHasTickDependencies);
 			bHasTickDependencies = true;
 		}
 		else
 		{
-			// PostTickFunction should contain only the tick function to the main actor's one
-			check(PostTickFunction.GetPrerequisites().Num());
+			// PostTickFunction should contain only the tick function to the main actor's ones (component and owner)
+			check(PostTickFunction.GetPrerequisites().Num() == 2);
 
 			// adding to PreTickFunction all the tick dependencies TickActorComponent has, so it runs after all the tick dependencies of ALL the TickActorComponents in this FInteractionIsland
-			// BUT excluding the main actor tick fuction that is PostTickFunction.GetPrerequisites()[0]
-			const FTickFunction* MainActorTickFunction = PostTickFunction.GetPrerequisites()[0].Get();
+			// BUT excluding the main actor tick fuctions that are in PostTickFunction.GetPrerequisites()
 			for (FTickPrerequisite& TickActorComponentPrerequisite : TickActorComponent->PrimaryComponentTick.GetPrerequisites())
 			{
-				if (TickActorComponentPrerequisite.PrerequisiteTickFunction != MainActorTickFunction)
+				if (!PostTickFunction.GetPrerequisites().ContainsByPredicate([&TickActorComponentPrerequisite](const FTickPrerequisite& MainActorTickPrerequisite)
+					{
+						return TickActorComponentPrerequisite.PrerequisiteTickFunction == MainActorTickPrerequisite.Get();
+					}))
 				{
 					AddPrerequisite(PreTickFunction, TickActorComponentPrerequisite.PrerequisiteObject.Get(), *TickActorComponentPrerequisite.PrerequisiteTickFunction);
 				}
 			}
 
 			// Add post-tick function dependencies
-			// it should be this island, but it's not a UObject, so we use the InteractionSubsystem
 			AddPrerequisite(TickActorComponent->PrimaryComponentTick, InteractionSubsystem, PostTickFunction);
+			AddPrerequisite(TickActor->PrimaryActorTick, InteractionSubsystem, PostTickFunction);
 			check(bHasTickDependencies);
 		}
+
+		// adding prerequisite between actor and skeletal mesh component so CBP runs before ABP parallel execution jobs
+		AddPrerequisite(TickActor->PrimaryActorTick, InteractionSubsystem, TickActorComponent->PrimaryComponentTick);
 	}
 
 #if ENABLE_ANIM_DEBUG
@@ -931,17 +952,26 @@ void FInteractionIsland::RemoveTickDependencies(bool bValidateTickDependencies)
 		{
 			if (UActorComponent* TickActorComponent = TickActorComponentPtr.Get())
 			{
+				AActor* TickActor = TickActorComponent->GetOwner();
+				check(TickActor && TickActor->CanEverTick());
+
 				if(IInteractionIslandDependency* IslandDependency = FindCustomDependency(TickActorComponent))
 				{
 					if (bMainActor)
 					{
 						IslandDependency->RemoveSubsequent(InteractionSubsystem, PreTickFunction, TickActorComponent);
 						IslandDependency->RemovePrerequisite(InteractionSubsystem, PostTickFunction, TickActorComponent);
+
+						TickActor->PrimaryActorTick.RemovePrerequisite(InteractionSubsystem, PreTickFunction);
+						PostTickFunction.RemovePrerequisite(TickActor, TickActor->PrimaryActorTick);
 					}
 					else
 					{
 						IslandDependency->RemoveSubsequent(InteractionSubsystem, PostTickFunction, TickActorComponent);
+						TickActor->PrimaryActorTick.RemovePrerequisite(InteractionSubsystem, PostTickFunction);
 					}
+
+					IslandDependency->RemoveSubsequent(InteractionSubsystem, TickActor->PrimaryActorTick, TickActorComponent);
 				}
 				else
 				{
@@ -949,11 +979,17 @@ void FInteractionIsland::RemoveTickDependencies(bool bValidateTickDependencies)
 					{
 						TickActorComponent->PrimaryComponentTick.RemovePrerequisite(InteractionSubsystem, PreTickFunction);
 						PostTickFunction.RemovePrerequisite(TickActorComponent, TickActorComponent->PrimaryComponentTick);
+
+						TickActor->PrimaryActorTick.RemovePrerequisite(InteractionSubsystem, PreTickFunction);
+						PostTickFunction.RemovePrerequisite(TickActor, TickActor->PrimaryActorTick);
 					}
 					else
 					{
 						TickActorComponent->PrimaryComponentTick.RemovePrerequisite(InteractionSubsystem, PostTickFunction);
+						TickActor->PrimaryActorTick.RemovePrerequisite(InteractionSubsystem, PostTickFunction);
 					}
+
+					TickActor->PrimaryActorTick.RemovePrerequisite(InteractionSubsystem, TickActorComponent->PrimaryComponentTick);
 				}
 			}
 			bMainActor = false;
@@ -1001,9 +1037,9 @@ void FInteractionIsland::InjectToActor(const UObject* AnimContext, bool bAddTick
 			//	tick order: 
 			//		ALL TickActorComponents prerequisites (ultimately we're looking to have UCharacterMovementComponent or UCharacterMoverComponent ticked) ->
 			//			Island.PreTickFunction ->
-			//				first injected TickActorComponent (USkeletalMeshComponent, or UAnimNextComponent) ->
+			//				first injected TickActor AND TickActorComponent (USkeletalMeshComponent, or UAnimNextComponent) to have CBP running before and ABP / UAF on the same actor ->
 			//					Island.PostTickFunction ->
-			//						other TickActorComponent(s)
+			//						other TickActor(s) before TickActorComponent(s)
 			TickActorComponents.AddUnique(TickActorComponent);
 			IslandAnimContexts.AddUnique(AnimContext);
 
diff --git a/Engine/Plugins/Animation/PoseSearch/Source/Runtime/Public/PoseSearch/PoseSearchHistory.h b/Engine/Plugins/Animation/PoseSearch/Source/Runtime/Public/PoseSearch/PoseSearchHistory.h
--- a/Engine/Plugins/Animation/PoseSearch/Source/Runtime/Public/PoseSearch/PoseSearchHistory.h
+++ b/Engine/Plugins/Animation/PoseSearch/Source/Runtime/Public/PoseSearch/PoseSearchHistory.h
@@ -179,10 +179,8 @@ FArchive& operator<<(FArchive& Ar, FArchivedPoseHistory& Entry);
 
 struct FPoseHistory : public IPoseHistory
 {
-	PRAGMA_DISABLE_DEPRECATION_WARNINGS
 	FPoseHistory() = default;
 	virtual ~FPoseHistory() override = default;
-	UE_API PRAGMA_ENABLE_DEPRECATION_WARNINGS
 
 	FPoseHistory(const FPoseHistory& Other);
 	UE_API FPoseHistory(FPoseHistory&& Other);
@@ -220,12 +218,11 @@ struct FPoseHistory : public IPoseHistory
 	UE_API virtual void DebugDraw(FAnimInstanceProxy& AnimInstanceProxy, FColor Color) const override;
 #endif
 	// End of IPoseHistory interface
-	
-	int32 GetMaxNumPoses() const { return MaxNumPoses; }
-	float GetSamplingInterval() const { return SamplingInterval; }
 
-private:
+	UE_API int32 GetMaxNumPoses() const;
+	UE_API float GetSamplingInterval() const;
 
+private:
 	// caching MaxNumPoses, since FData::Entries.Max() is a padded number
 	int32 MaxNumPoses = 0;
 
diff --git a/Engine/Plugins/Animation/PoseSearch/Source/Runtime/Public/PoseSearch/PoseSearchInteractionSubsystem.h b/Engine/Plugins/Animation/PoseSearch/Source/Runtime/Public/PoseSearch/PoseSearchInteractionSubsystem.h
--- a/Engine/Plugins/Animation/PoseSearch/Source/Runtime/Public/PoseSearch/PoseSearchInteractionSubsystem.h
+++ b/Engine/Plugins/Animation/PoseSearch/Source/Runtime/Public/PoseSearch/PoseSearchInteractionSubsystem.h
@@ -44,6 +44,7 @@ struct FTagToDatabases;
 // 
 //				Tick Island[k].PreTickFunction (that eventually generates the trajectories with all the updated CMCs or Mover)
 //
+//				Tick Actor[k][0] to have CBP running before and ABP or UAF on the same actor 0
 //				Tick Actor[k][0].SkeletalMeshComponent (or AnimNextComponent, that performs the MotionMatchInteraction queries for all the involved actors via DoSearch_AnyThread)
 // 
 //				Tick Island[k].PostTickFunction (currently just a threading fence for the execution of all the other SkeletalMeshComponent(s) or AnimNextComponent(s))
@@ -51,7 +52,10 @@ struct FTagToDatabases;
 //				for each Actor[k][i]
 //				{
 //					if (i != 0)
-//						Tick SkeletalMeshComponent[k][i] (or AnimNextComponent(s) that DoSearch_AnyThread get the cached result calculated by Tick Actor[k][0].SkeletalMeshComponent) in parallel
+//					{
+//						Tick Actor[k][i] to have CBP running before and ABP or UAF on the same actor i
+// 						Tick SkeletalMeshComponent[k][i] (or AnimNextComponent(s) that DoSearch_AnyThread get the cached result calculated by Tick Actor[k][0].SkeletalMeshComponent) in parallel
+//					}
 //				}
 //			}
 // - next frame UPoseSearchInteractionLibrary::MotionMatchInteraction_Pure(TArray<FPoseSearchInteractionAvailability> Availabilities, UObject* AnimInstance), with the context of all the published 
