PoseSearch - adding support for calculating relative transforms of bones and desired reach during a MM interaction, that can be queried via GetMotionMatchInteractionConstraint to feed an IK system

#rnx

diff --git a/Engine/Plugins/Animation/PoseSearch/Source/Runtime/Private/PoseSearchAssetSampler.cpp b/Engine/Plugins/Animation/PoseSearch/Source/Runtime/Private/PoseSearchAssetSampler.cpp
--- a/Engine/Plugins/Animation/PoseSearch/Source/Runtime/Private/PoseSearchAssetSampler.cpp
+++ b/Engine/Plugins/Animation/PoseSearch/Source/Runtime/Private/PoseSearchAssetSampler.cpp
@@ -636,7 +636,7 @@ void FAnimationAssetSampler::Process()
 	}
 }
 
-void FAnimationAssetSampler::ExtractAnimNotifyStates(float Time, FAnimNotifyContext& PreAllocatedNotifyContext, const TFunction<bool(UAnimNotifyState*)>& ProcessAnimNotifyState) const
+void FAnimationAssetSampler::ExtractAnimNotifyContext(float Time, FAnimNotifyContext& PreAllocatedNotifyContext) const
 {
 	float SampleTime = Time;
 
@@ -685,6 +685,11 @@ void FAnimationAssetSampler::ExtractAnimNotifyStates(float Time, FAnimNotifyCont
 	{
 		checkNoEntry();
 	}
+}
+
+void FAnimationAssetSampler::ExtractAnimNotifyStates(float Time, FAnimNotifyContext& PreAllocatedNotifyContext, const TFunction<bool(UAnimNotifyState*)>& ProcessAnimNotifyState) const
+{
+	ExtractAnimNotifyContext(Time, PreAllocatedNotifyContext);
 
 	// check which notifies actually overlap Time and are of the right base type
 	for (const FAnimNotifyEventReference& EventReference : PreAllocatedNotifyContext.ActiveNotifies)
@@ -699,6 +704,23 @@ void FAnimationAssetSampler::ExtractAnimNotifyStates(float Time, FAnimNotifyCont
 	}
 }
 
+void FAnimationAssetSampler::ExtractAnimNotifyEvents(float Time, FAnimNotifyContext& PreAllocatedNotifyContext, const TFunction<bool(const FAnimNotifyEvent*)>& ProcessAnimNotifyEvent) const
+{
+	ExtractAnimNotifyContext(Time, PreAllocatedNotifyContext);
+
+	// check which notifies actually overlap Time and are of the right base type
+	for (const FAnimNotifyEventReference& EventReference : PreAllocatedNotifyContext.ActiveNotifies)
+	{
+		if (const FAnimNotifyEvent* NotifyEvent = EventReference.GetNotify())
+		{
+			if (!ProcessAnimNotifyEvent(NotifyEvent))
+			{
+				break;
+			}
+		}
+	}
+}
+
 TConstArrayView<FAnimNotifyEvent> FAnimationAssetSampler::GetAllAnimNotifyEvents() const
 {
 	if (const UBlendSpace* BlendSpace = Cast<UBlendSpace>(AnimationAssetPtr.Get()))
diff --git a/Engine/Plugins/Animation/PoseSearch/Source/Runtime/Private/PoseSearchHistory.cpp b/Engine/Plugins/Animation/PoseSearch/Source/Runtime/Private/PoseSearchHistory.cpp
--- a/Engine/Plugins/Animation/PoseSearch/Source/Runtime/Private/PoseSearchHistory.cpp
+++ b/Engine/Plugins/Animation/PoseSearch/Source/Runtime/Private/PoseSearchHistory.cpp
@@ -679,7 +679,7 @@ FPoseHistory& FPoseHistory::operator=(const FPoseHistory& Other)
 	if (this != &Other)
 	{
 #if ENABLE_ANIM_DEBUG
-		UE_MT_SCOPED_WRITE_ACCESS(Other.PoseDataThreadSafeCounter);
+		UE_MT_SCOPED_READ_ACCESS(Other.PoseDataThreadSafeCounter);
 		UE_MT_SCOPED_WRITE_ACCESS(PoseDataThreadSafeCounter);
 #endif // ENABLE_ANIM_DEBUG
 
diff --git a/Engine/Plugins/Animation/PoseSearch/Source/Runtime/Private/PoseSearchInteractionIsland.cpp b/Engine/Plugins/Animation/PoseSearch/Source/Runtime/Private/PoseSearchInteractionIsland.cpp
--- a/Engine/Plugins/Animation/PoseSearch/Source/Runtime/Private/PoseSearchInteractionIsland.cpp
+++ b/Engine/Plugins/Animation/PoseSearch/Source/Runtime/Private/PoseSearchInteractionIsland.cpp
@@ -9,6 +9,7 @@
 #include "Features/IModularFeatures.h"
 #include "PoseSearch/AnimNode_PoseSearchHistoryCollector.h"
 #include "PoseSearch/MultiAnimAsset.h"
+#include "PoseSearch/PoseSearchAnimNotifies.h"
 #include "PoseSearch/PoseSearchContext.h"
 #include "PoseSearch/PoseSearchDatabase.h"
 #include "PoseSearch/PoseSearchInteractionLibrary.h"
@@ -313,7 +314,14 @@ const UE::PoseSearch::IPoseHistory* FInteractionSearchContextBase::GetPoseHistor
 	return nullptr;
 }
 
-uint32 FInteractionSearchContextBase::GetHashForEquivalence() const
+void FInteractionSearchContextBase::CacheHashForEquivalence()
+{
+	// trying to inefficiently call CacheHash multiple times?
+	check(CachedHashForEquivalence == InvalidCachedHashForEquivalence);
+	CachedHashForEquivalence = CalculateHashForEquivalence();
+}
+
+uint32 FInteractionSearchContextBase::CalculateHashForEquivalence() const
 {
 	uint32 Hash = GetTypeHash(Database);
 
@@ -328,15 +336,17 @@ uint32 FInteractionSearchContextBase::GetHashForEquivalence() const
 		Hash = HashCombineFast(Hash, GetTypeHash(Roles[Index]));
 	}
 
+	// only an invalid hash is zero!
+	check(Hash != InvalidCachedHashForEquivalence);
 	return Hash;
 }
 
-bool FInteractionSearchContextBase::IsEquivalent(const FInteractionSearchContextBase& Other) const
+uint32 FInteractionSearchContextBase::GetHashForEquivalence() const
 {
-	// skipping bDisableCollisions and PoseHistories for equality
-	return Database == Other.Database &&
-		AnimContexts == Other.AnimContexts &&
-		Roles == Other.Roles;
+	check(CachedHashForEquivalence != InvalidCachedHashForEquivalence);
+	// making sure properties didn't change without updating CachedHashForEquivalence
+	check(CachedHashForEquivalence == CalculateHashForEquivalence());
+	return CachedHashForEquivalence;
 }
 
 void FInteractionSearchContextBase::AddReferencedObjects(FReferenceCollector& Collector)
@@ -349,7 +359,7 @@ void FInteractionSearchContextBase::AddReferencedObjects(FReferenceCollector& Co
 	}
 }
 
-void FInteractionSearchContextBase::Add(const UObject* AnimContext, const UE::PoseSearch::IPoseHistory* PoseHistory, const UE::PoseSearch::FRole Role, bool bWantsDisableCollisions)
+void FInteractionSearchContextBase::Add(const UObject* AnimContext, const IPoseHistory* PoseHistory, const FRole Role, bool bWantsDisableCollisions)
 {
 	check(AnimContext && PoseHistory);
 
@@ -445,15 +455,402 @@ bool FInteractionSearchContextBase::CheckForConsistency() const
 
 #endif // DO_CHECK
 
+// FConstraint
+///////////////////////////////////////////////////////////
+FConstraint::FConstraint(const UAnimNotifyState_PoseSearchConstraint* ConstraintNotifyState)
+{
+	FromBoneName = ConstraintNotifyState->FromBoneName;
+	FromBoneRole = ConstraintNotifyState->FromBoneRole;
+	ToBoneName = ConstraintNotifyState->ToBoneName;
+	ToBoneRole = ConstraintNotifyState->ToBoneRole;
+	TranslationWeight = ConstraintNotifyState->TranslationWeight;
+	RotationWeight = ConstraintNotifyState->RotationWeight;
+	RampUpTime = ConstraintNotifyState->RampUpTime;
+	CoolDownTime = ConstraintNotifyState->CoolDownTime;
+}
+
+const FTransform* FConstraint::GetBoneTransform(FRole Role, FName BoneName) const
+{
+	if (FromBoneRole == Role && FromBoneName == BoneName)
+	{
+		return &FromBoneTransform;
+	}
+
+	if (ToBoneRole == Role && ToBoneName == BoneName)
+	{
+		return &ToBoneTransform;
+	}
+
+	return nullptr;
+}
+
+bool FConstraint::operator==(const FConstraint& Other) const
+{
+	// same order of From and To properties
+	if (FromBoneRole == Other.FromBoneRole && FromBoneName == Other.FromBoneName && ToBoneRole == Other.ToBoneRole && ToBoneName == Other.ToBoneName)
+	{
+		return true;
+	}
+
+	// opposite order of From and To properties
+	if (FromBoneRole == Other.ToBoneRole && FromBoneName == Other.ToBoneName && ToBoneRole == Other.FromBoneRole && ToBoneName == Other.FromBoneName)
+	{
+		return true;
+	}
+
+	return false;
+}
+
+// FConstraints
+///////////////////////////////////////////////////////////
+bool FConstraints::GetConstraint(FRole Role, FName BoneName, float& OutDesiredReach, FTransform& OutTransform) const
+{
+	for (const FConstraint& Constraint : Constraints)
+	{
+		if (const FTransform* BoneTransform = Constraint.GetBoneTransform(Role, BoneName))
+		{
+			OutDesiredReach = Constraint.DesiredReach;
+			OutTransform = *BoneTransform;
+			return true;
+		}
+	}
+
+	OutDesiredReach = 0.f;
+	OutTransform = FTransform::Identity;
+	return false;
+}
+
+int32 FConstraints::Num() const
+{
+	return Constraints.Num();
+}
+
+const FConstraint* FConstraints::Find(const FConstraint& Constraint) const
+{
+	const int32 ConstraintIndex = Constraints.Find(Constraint);
+	if (ConstraintIndex != INDEX_NONE)
+	{
+		return &Constraints[ConstraintIndex];
+	}
+	return nullptr;
+}
+
+FConstraint* FConstraints::Find(const FConstraint& Constraint)
+{
+	const int32 ConstraintIndex = Constraints.Find(Constraint);
+	if (ConstraintIndex != INDEX_NONE)
+	{
+		return &Constraints[ConstraintIndex];
+	}
+	return nullptr;
+}
+
+void FConstraints::ProcessAll(float DeltaTime, const FSearchResult& SearchResult, const FInteractionSearchContext& SearchContext, const FConstraints* SearchContextPreviousConstraints)
+{
+	// making sure nobody already called UpdateConstraints or we didn't reset the constraints...
+	check(Constraints.IsEmpty());
+
+	const UPoseSearchDatabase* Database = SearchResult.Database.Get();
+	check(Database);
+
+	const FSearchIndexAsset* SearchIndexAsset = SearchResult.GetSearchIndexAsset();
+	check(SearchIndexAsset);
+
+	const FPoseSearchDatabaseAnimationAssetBase* DatabaseAnimationAssetBase = Database->GetDatabaseAnimationAsset(*SearchIndexAsset);
+	check(DatabaseAnimationAssetBase);
+
+	if (const UMultiAnimAsset* MultiAnimAsset = Cast<UMultiAnimAsset>(DatabaseAnimationAssetBase->GetAnimationAsset()))
+	{
+		const int32 NumAnimContexts = SearchContext.Num();
+		const float AssetTime = SearchResult.GetAssetTime();
+		const bool bIsMirrored = SearchIndexAsset->IsMirrored();
+		const FVector BlendParameters = SearchIndexAsset->GetBlendParameters();
+
+		FMemMark Mark(FMemStack::Get());
+
+		TConstArrayView<FConstraint> PreviousConstraints;
+		if (SearchContextPreviousConstraints)
+		{
+			PreviousConstraints = SearchContextPreviousConstraints->Constraints;
+		}
+
+		TArray<bool, TInlineAllocator<NumPreallocatedContraints, TMemStackAllocator<>>> VisitedPreviousConstraints;
+		VisitedPreviousConstraints.SetNumZeroed(PreviousConstraints.Num());
+		
+		// extracting the contraints for this SearchContext from MultiAnimAsset
+		FAnimNotifyContext PreAllocatedNotifyContext;
+		for (int32 AnimContextIndex = 0; AnimContextIndex < NumAnimContexts; ++AnimContextIndex)
+		{
+			const FRole Role = SearchContext.GetRole(AnimContextIndex);
+			const UAnimationAsset* Animation = MultiAnimAsset->GetAnimationAsset(Role);
+			check(Animation);
+
+			// initializing FAnimationAssetSampler only for extracting AnimNotifyStates
+			FAnimationAssetSampler Sampler(Animation, FTransform::Identity, FVector::ZeroVector, FAnimationAssetSampler::DefaultRootTransformSamplingRate, false, false);
+			Sampler.ExtractAnimNotifyEvents(AssetTime, PreAllocatedNotifyContext, [AssetTime, DeltaTime, this, &PreviousConstraints, &VisitedPreviousConstraints](const FAnimNotifyEvent* AnimNotifyEvent)
+				{
+					if (const UAnimNotifyState_PoseSearchConstraint* ConstraintNotifyState = Cast<const UAnimNotifyState_PoseSearchConstraint>(AnimNotifyEvent->NotifyStateClass))
+					{
+						// subtracting ConstraintNotifyState->CoolDownTime from the validity of this UAnimNotifyState_PoseSearchConstraint
+						const float NotifyEndTime = AnimNotifyEvent->GetEndTriggerTime();
+						if (AssetTime <= (NotifyEndTime - ConstraintNotifyState->CoolDownTime))
+						{
+							FConstraint NewConstraint(ConstraintNotifyState);
+							const int32 PreviousConstraintIndex = PreviousConstraints.Find(NewConstraint);
+							if (PreviousConstraintIndex != INDEX_NONE)
+							{
+								VisitedPreviousConstraints[PreviousConstraintIndex] = true;
+								const FConstraint& PreviousConstraint = PreviousConstraints[PreviousConstraintIndex];
+
+								// @todo: for now we just ignore conflicting constraints, but we should probably blend PreviousConstraint into NewConstraint :/
+								//        perhaps keeping the min or RampUpTime, to make the constraint as following the fastest one to get into full DesiredReach, and the max CoolDownTime to fade off with the slowest rate
+								//        not really sure what to do about TranslationWeight and RotationWeight.. maybe slowly blending them to the newer one?
+								if (PreviousConstraint.RampUpTime > UE_KINDA_SMALL_NUMBER)
+								{
+									NewConstraint.DesiredReach = FMath::Min(1.f, PreviousConstraint.DesiredReach + (DeltaTime / PreviousConstraint.RampUpTime));
+								}
+								else
+								{
+									NewConstraint.DesiredReach = 1.f;
+								}
+							}
+							else if (NewConstraint.RampUpTime <= UE_KINDA_SMALL_NUMBER)
+							{
+								NewConstraint.DesiredReach = 1.f;
+							}
+
+							// @todo: for now we also ignore duplicate constraints, and perhaps we should use a similar strategy of blending the constraints
+							Constraints.AddUnique(NewConstraint);
+						}
+					}
+					return true;
+				});
+		}
+
+		// adding remaining unprocessed constraints from PreviousConstraints
+		for (int32 PreviousConstraintIndex = 0; PreviousConstraintIndex < PreviousConstraints.Num(); ++PreviousConstraintIndex)
+		{
+			if (!VisitedPreviousConstraints[PreviousConstraintIndex])
+			{
+				const FConstraint& PreviousConstraint = PreviousConstraints[PreviousConstraintIndex];
+				if (PreviousConstraint.CoolDownTime > UE_KINDA_SMALL_NUMBER)
+				{
+					const float DesiredReach = PreviousConstraint.DesiredReach - (DeltaTime / PreviousConstraint.CoolDownTime);
+					if (DesiredReach > UE_KINDA_SMALL_NUMBER)
+					{
+						FConstraint& NewConstraint = Constraints.Add_GetRef(PreviousConstraint);
+						NewConstraint.DesiredReach = DesiredReach;
+					}
+				}
+			}
+		}
+
+		// process all the collected constraints
+		if (!Constraints.IsEmpty())
+		{
+			TArray<FTransform, TInlineAllocator<PreallocatedRolesNum, TMemStackAllocator<>>> ComponentToWorldTransforms;
+			TArray<TConstArrayView<FTransform>, TInlineAllocator<PreallocatedRolesNum, TMemStackAllocator<>>> AnimContextsComponentSpaceTransforms;
+			TArray<FBoneContainer, TInlineAllocator<PreallocatedRolesNum, TMemStackAllocator<>>> BoneContainers;
+			TArray<FCompactPose, TInlineAllocator<PreallocatedRolesNum, TMemStackAllocator<>>> Poses;
+			TArray<FCSPose<FCompactPose>, TInlineAllocator<PreallocatedRolesNum, TMemStackAllocator<>>> ComponentSpacePoses;
+			TArray<const USkeleton*, TInlineAllocator<PreallocatedRolesNum, TMemStackAllocator<>>> Skeletons;
+
+			ComponentToWorldTransforms.SetNum(NumAnimContexts);
+			AnimContextsComponentSpaceTransforms.SetNum(NumAnimContexts);
+			BoneContainers.SetNum(NumAnimContexts);
+			Poses.SetNum(NumAnimContexts);
+			ComponentSpacePoses.SetNum(NumAnimContexts);
+			Skeletons.SetNum(NumAnimContexts);
+
+			FRoleToIndex RoleToIndex;
+			for (int32 AnimContextIndex = 0; AnimContextIndex < NumAnimContexts; ++AnimContextIndex)
+			{
+				const UObject* AnimContext = SearchContext.GetAnimContext(AnimContextIndex);
+				check(AnimContext);
+				const USkeletalMeshComponent* SkeletalMeshComponent = GetContextSkeletalMeshComponent(AnimContext, false);
+				check(SkeletalMeshComponent);
+
+				// @todo: do we need this logic about URO etc? if we're executing this is because we update the mesh, so we probably don't need any interpolation...
+				const bool bUROInSync = SkeletalMeshComponent->ShouldUseUpdateRateOptimizations() && SkeletalMeshComponent->AnimUpdateRateParams != nullptr;
+				const bool bUsingExternalInterpolation = SkeletalMeshComponent->IsUsingExternalInterpolation();
+
+				const TArray<FTransform>& CachedComponentSpaceTransforms = SkeletalMeshComponent->GetCachedComponentSpaceTransforms();
+				const bool bArraySizesMatch = CachedComponentSpaceTransforms.Num() == SkeletalMeshComponent->GetComponentSpaceTransforms().Num();
+
+				ComponentToWorldTransforms[AnimContextIndex] = SkeletalMeshComponent->GetComponentTransform();
+				AnimContextsComponentSpaceTransforms[AnimContextIndex] = (bUROInSync || bUsingExternalInterpolation) && bArraySizesMatch ? CachedComponentSpaceTransforms : SkeletalMeshComponent->GetComponentSpaceTransforms();
+
+				// @todo: check FAnimNode_CopyPoseFromMesh::PreUpdate to see how to GetAnimCurves and GetCustomAttributes
+				const FRole Role = SearchContext.GetRole(AnimContextIndex);
+				RoleToIndex.Add(Role) = AnimContextIndex;
+
+				const UAnimationAsset* Animation = MultiAnimAsset->GetAnimationAsset(Role);
+				check(Animation);
+				const FTransform Origin = MultiAnimAsset->GetOrigin(Role);
+
+				FAnimationAssetSampler Sampler(Animation, Origin, BlendParameters);
+				const FTransform RootTransform = Sampler.ExtractRootTransform(AssetTime);
+
+				// @todo add support for mirroring if needed
+				//if (MirrorDataTables.IsValidIndex(ItemIndex) && MirrorDataTables[ItemIndex])
+				//{
+				//	const FMirrorDataCache MirrorDataCache(MirrorDataTables[ItemIndex]);
+				//	AssetRootBoneTransforms[ItemIndex] = MirrorDataCache.MirrorTransform(AssetRootBoneTransforms[ItemIndex]);
+				//}
+
+				// array containing the bone index of the root bone (0)
+				Skeletons[AnimContextIndex] = Animation->GetSkeleton();
+				check(Skeletons[AnimContextIndex]);
+				check(SkeletalMeshComponent->GetSkeletalMeshAsset()->GetSkeleton() == Skeletons[AnimContextIndex]);
+
+				const FReferenceSkeleton& ReferenceSkeleton = Skeletons[AnimContextIndex]->GetReferenceSkeleton();
+
+				// @todo: refine this BoneIndices / cache bone containers
+				const int32 NumBones = ReferenceSkeleton.GetNum();
+				TArray<uint16, TMemStackAllocator<>> BoneIndices;
+				BoneIndices.SetNumUninitialized(NumBones);
+				for (int32 BoneIndex = 0; BoneIndex < NumBones; ++BoneIndex)
+				{
+					BoneIndices[BoneIndex] = BoneIndex;
+				}
+
+				// extracting the pose, containing only the root bone from the Sampler 
+				BoneContainers[AnimContextIndex].InitializeTo(BoneIndices, UE::Anim::FCurveFilterSettings(UE::Anim::ECurveFilterMode::DisallowAll), *Skeletons[AnimContextIndex]);
+				Poses[AnimContextIndex].SetBoneContainer(&BoneContainers[AnimContextIndex]);
+				Sampler.ExtractPose(AssetTime, Poses[AnimContextIndex]);
+
+
+				// making sure the animation root bone transform is Identity, so we can confuse the root with the root BONE transform and preserve performances!
+				FTransform& RootBoneTransform = Poses[AnimContextIndex][FCompactPoseBoneIndex(0)];
+				if (!RootBoneTransform.Equals(FTransform::Identity))
+				{
+					const FVector Pos = RootBoneTransform.GetTranslation();
+					const FRotator Rot(RootBoneTransform.GetRotation());
+					UE_LOG(LogPoseSearch, Error, TEXT("FInteractionIsland::UpdateConstraints unsupported non identity root bone in %s at time %f Pos(%f, %f, %f), Rot(%f, %f, %f)"), *Animation->GetName(), AssetTime, Pos.X, Pos.Y, Pos.Z, Rot.Pitch, Rot.Yaw, Rot.Roll);
+				}
+				RootBoneTransform = Sampler.ExtractRootTransform(AssetTime);
+				ComponentSpacePoses[AnimContextIndex].InitPose(Poses[AnimContextIndex]);
+
+#if ENABLE_VISUAL_LOG
+				if (FVisualLogger::IsRecording())
+				{
+					static const TCHAR* LogName = TEXT("PoseSearchInteraction_Constraints");
+
+					for (int32 BoneIndex = 0; BoneIndex < ReferenceSkeleton.GetNum(); ++BoneIndex)
+					{
+						const int32 ParentBoneIndex = ReferenceSkeleton.GetParentIndex(BoneIndex);
+						if (ParentBoneIndex != INDEX_NONE)
+						{
+							const FVector CurrentBoneWorldLocation = ComponentToWorldTransforms[AnimContextIndex].TransformPosition(AnimContextsComponentSpaceTransforms[AnimContextIndex][BoneIndex].GetTranslation());
+							const FVector CurrentParentBoneWorldLocation = ComponentToWorldTransforms[AnimContextIndex].TransformPosition(AnimContextsComponentSpaceTransforms[AnimContextIndex][ParentBoneIndex].GetTranslation());
+
+							const FVector AnimBoneWorldLocation = ComponentSpacePoses[AnimContextIndex].GetComponentSpaceTransform(FCompactPoseBoneIndex(BoneIndex)).GetTranslation();
+							const FVector AnimParentBoneWorldLocation = ComponentSpacePoses[AnimContextIndex].GetComponentSpaceTransform(FCompactPoseBoneIndex(ParentBoneIndex)).GetTranslation();
+
+							for (int32 DrawAnimContextIndex = 0; DrawAnimContextIndex < NumAnimContexts; ++DrawAnimContextIndex)
+							{
+								const UObject* DrawAnimContext = SearchContext.GetAnimContext(DrawAnimContextIndex);
+								const FColor Color = DrawAnimContext == AnimContext ? FColor::Red : FColor::Blue;
+
+								// current characters poses
+								UE_VLOG_SEGMENT(DrawAnimContext, LogName, Display, CurrentBoneWorldLocation, CurrentParentBoneWorldLocation, Color, TEXT(""));
+
+								// sampled poses from UAnimationAsset
+								UE_VLOG_SEGMENT(DrawAnimContext, LogName, Display, AnimBoneWorldLocation, AnimParentBoneWorldLocation, Color, TEXT(""));
+							}
+						}
+					}
+				}
+#endif //ENABLE_VISUAL_LOG
+			}
+
+			for (FConstraint& Constraint : Constraints)
+			{
+				const int32* FromBoneRoleIndex = RoleToIndex.Find(Constraint.FromBoneRole);
+				const int32* ToBoneRoleIndex = RoleToIndex.Find(Constraint.ToBoneRole);
+				if (FromBoneRoleIndex && ToBoneRoleIndex)
+				{
+					FBoneReference FromBoneReference(Constraint.FromBoneName);
+					FBoneReference ToBoneReference(Constraint.ToBoneName);
+
+					if (FromBoneReference.Initialize(Skeletons[*FromBoneRoleIndex]) &&
+						ToBoneReference.Initialize(Skeletons[*ToBoneRoleIndex]))
+					{
+						const FTransform CurrentFromTransform = AnimContextsComponentSpaceTransforms[*FromBoneRoleIndex][FromBoneReference.BoneIndex] * ComponentToWorldTransforms[*FromBoneRoleIndex];
+						const FTransform CurrentToTransform = AnimContextsComponentSpaceTransforms[*ToBoneRoleIndex][ToBoneReference.BoneIndex] * ComponentToWorldTransforms[*ToBoneRoleIndex];
+
+						const FTransform& AnimFromTransform = ComponentSpacePoses[*FromBoneRoleIndex].GetComponentSpaceTransform(FCompactPoseBoneIndex(FromBoneReference.BoneIndex));
+						const FTransform& AnimToTransform = ComponentSpacePoses[*ToBoneRoleIndex].GetComponentSpaceTransform(FCompactPoseBoneIndex(ToBoneReference.BoneIndex));
+
+						const FTransform WantedCurrentToTransform = AnimToTransform.GetRelativeTransform(AnimFromTransform) * CurrentFromTransform;
+						const FTransform WantedCurrentFromTransform = AnimFromTransform.GetRelativeTransform(AnimToTransform) * CurrentToTransform;
+
+						Constraint.FromBoneTransform.SetTranslation(FMath::Lerp(WantedCurrentFromTransform.GetTranslation(), CurrentFromTransform.GetTranslation(), Constraint.TranslationWeight));
+						Constraint.FromBoneTransform.SetRotation(FMath::Lerp(WantedCurrentFromTransform.GetRotation(), CurrentFromTransform.GetRotation(), Constraint.RotationWeight));
+						Constraint.FromBoneTransform.SetScale3D(CurrentFromTransform.GetScale3D());
+
+						Constraint.ToBoneTransform.SetTranslation(FMath::Lerp(WantedCurrentToTransform.GetTranslation(), CurrentToTransform.GetTranslation(), 1.f - Constraint.TranslationWeight));
+						Constraint.ToBoneTransform.SetRotation(FMath::Lerp(WantedCurrentToTransform.GetRotation(), CurrentToTransform.GetRotation(), 1.f - Constraint.RotationWeight));
+						Constraint.ToBoneTransform.SetScale3D(CurrentToTransform.GetScale3D());
+
+#if WITH_EDITORONLY_DATA && ENABLE_ANIM_DEBUG && ENABLE_VISUAL_LOG
+						if (FVisualLogger::IsRecording())
+						{
+							static const TCHAR* LogName = TEXT("PoseSearchInteraction_Constraints");
+
+							for (int32 DrawAnimContextIndex = 0; DrawAnimContextIndex < NumAnimContexts; ++DrawAnimContextIndex)
+							{
+								const UObject* DrawAnimContext = SearchContext.GetAnimContext(DrawAnimContextIndex);
+
+								// current characters poses
+								UE_VLOG_SEGMENT(DrawAnimContext, LogName, Display, CurrentFromTransform.GetTranslation(), CurrentToTransform.GetTranslation(), FColor::Yellow, TEXT(""));
+
+								// sampled poses from UAnimationAsset
+								UE_VLOG_SEGMENT(DrawAnimContext, LogName, Display, AnimFromTransform.GetTranslation(), AnimToTransform.GetTranslation(), FColor::Yellow, TEXT(""));
+
+								// wanted transforms
+								UE_VLOG_SEGMENT(DrawAnimContext, LogName, Display, CurrentFromTransform.GetTranslation(), WantedCurrentToTransform.GetTranslation(), FColor::Green, TEXT(""));
+								UE_VLOG_SEGMENT(DrawAnimContext, LogName, Display, CurrentToTransform.GetTranslation(), WantedCurrentFromTransform.GetTranslation(), FColor::Green, TEXT(""));
+
+								UE_VLOG_SEGMENT(DrawAnimContext, LogName, Display, CurrentFromTransform.GetTranslation(), Constraint.FromBoneTransform.GetTranslation(), FColor::Orange, TEXT(""));
+								UE_VLOG_SEGMENT(DrawAnimContext, LogName, Display, CurrentToTransform.GetTranslation(), Constraint.ToBoneTransform.GetTranslation(), FColor::Orange, TEXT(""));
+							}
+						}
+#endif // WITH_EDITORONLY_DATA && ENABLE_ANIM_DEBUG && ENABLE_VISUAL_LOG
+					}
+				}
+			}
+		}
+	}
+}
+
 // FValidInteractionSearch
 ///////////////////////////////////////////////////////////
-void FValidInteractionSearch::Init(const FInteractionSearchContext& Other)
+FValidInteractionSearch& FValidInteractionSearch::Init(const FInteractionSearchContextBase& SearchContextBase)
 {
-	static_cast<FInteractionSearchContextBase&>(*this) = static_cast<const FInteractionSearchContextBase&>(Other);
+#if ENABLE_ANIM_DEBUG
+	UE_MT_SCOPED_WRITE_ACCESS(ThreadSafeCounter);
+#endif // ENABLE_ANIM_DEBUG
+
+	static_cast<FInteractionSearchContextBase&>(*this) = SearchContextBase;
+	return *this;
+}
+
+void FValidInteractionSearch::Update(const FInteractionSearchResult& InSearchResult)
+{
+#if ENABLE_ANIM_DEBUG
+	UE_MT_SCOPED_WRITE_ACCESS(ThreadSafeCounter);
+#endif // ENABLE_ANIM_DEBUG
+
+	SearchResult = InSearchResult;
 }
 
 void FValidInteractionSearch::DoDisableCollisions()
 {
+#if ENABLE_ANIM_DEBUG
+	UE_MT_SCOPED_WRITE_ACCESS(ThreadSafeCounter);
+#endif // ENABLE_ANIM_DEBUG
+
 	check(DisabledCollisions.IsEmpty());
 	if (IsDisableCollisions())
 	{
@@ -503,6 +900,10 @@ void FValidInteractionSearch::DoDisableCollisions()
 
 void FValidInteractionSearch::UndoDisableCollisions()
 {
+#if ENABLE_ANIM_DEBUG
+	UE_MT_SCOPED_WRITE_ACCESS(ThreadSafeCounter);
+#endif // ENABLE_ANIM_DEBUG
+
 	for (const FDisabledCollisions::ElementType& DisabledCollision : DisabledCollisions)
 	{
 		if (AActor* ActorA = DisabledCollision.Key)
@@ -521,7 +922,9 @@ void FValidInteractionSearch::UndoDisableCollisions()
 #if ENABLE_VISUAL_LOG
 void FValidInteractionSearch::VLogContext(const FColor& Color) const
 {
-	using namespace UE::PoseSearch;
+#if ENABLE_ANIM_DEBUG
+	UE_MT_SCOPED_READ_ACCESS(ThreadSafeCounter);
+#endif // ENABLE_ANIM_DEBUG
 
 	if (FVisualLogger::IsRecording())
 	{
@@ -538,7 +941,7 @@ void FValidInteractionSearch::VLogContext(const FColor& Color) const
 				check(ContextOwningActor);
 				// using ContextOwningActor location instead of skeletal mesh location since it's faster for UAF
 				//Locations[Index] = GetContextLocation(AnimContext);
-				Locations[Index] = ContextOwningActor->GetTransform().GetLocation();
+				Locations[Index] = ContextOwningActor->GetTransform().GetTranslation();
 			}
 		}
 
@@ -561,6 +964,10 @@ void FValidInteractionSearch::VLogContext(const FColor& Color) const
 
 void FValidInteractionSearch::AddReferencedObjects(FReferenceCollector& Collector)
 {
+#if ENABLE_ANIM_DEBUG
+	UE_MT_SCOPED_READ_ACCESS(ThreadSafeCounter);
+#endif // ENABLE_ANIM_DEBUG
+
 	FInteractionSearchContextBase::AddReferencedObjects(Collector);
 
 	for (FDisabledCollision& DisabledCollision : DisabledCollisions)
@@ -572,8 +979,113 @@ void FValidInteractionSearch::AddReferencedObjects(FReferenceCollector& Collecto
 
 // FInteractionSearchContext
 ///////////////////////////////////////////////////////////
-void FInteractionSearchContext::Add(const UObject* AnimContext, const UE::PoseSearch::IPoseHistory* PoseHistory
-	, const UE::PoseSearch::FRole Role, bool bWantsDisableCollisions, int32 TickPriority
+void FValidInteractionSearches::Update(TConstArrayView<FInteractionIsland*> Islands)
+{
+#if ENABLE_ANIM_DEBUG
+	UE_MT_SCOPED_WRITE_ACCESS(ThreadSafeCounter);
+#endif //ENABLE_ANIM_DEBUG
+
+	FMemMark Mark(FMemStack::Get());
+	TMap<uint32, bool, TInlineSetAllocator<32, TMemStackSetAllocator<>>> ContinuingOrNewValidSearchContexts;
+	for (const FInteractionIsland* Island : Islands)
+	{
+		if (Island->IsInitialized())
+		{
+			// analyzing ALL current tick interaction results
+			for (const FInteractionSearchResult& SearchResult : Island->GetSearchResults())
+			{
+				const FInteractionSearchContext& SearchContext = Island->GetSearchContexts()[SearchResult.SearchIndex];
+				const uint32 ValidSearchContextHash = SearchContext.GetHashForEquivalence();
+				if (FValidInteractionSearch* ContinuingValidInteractionSearch = Searches.Find(ValidSearchContextHash))
+				{
+					ContinuingOrNewValidSearchContexts.Add(ValidSearchContextHash) = true;
+					ContinuingValidInteractionSearch->Update(SearchResult);
+				}
+				else
+				{
+					ContinuingOrNewValidSearchContexts.Add(ValidSearchContextHash) = false;
+					Searches.Add(ValidSearchContextHash).Init(SearchContext).Update(SearchResult);
+				}
+			}
+		}
+	}
+
+	// calling all the OnInteractionEnd first
+	for (TMap<uint32, FValidInteractionSearch>::TIterator It = Searches.CreateIterator(); It; ++It)
+	{
+		if (!ContinuingOrNewValidSearchContexts.Find(It->Key))
+		{
+			OnInteractionEnd(It->Value);
+			It.RemoveCurrent();
+		}
+	}
+
+	// calling all the OnInteractionContinuing and OnInteractionStart after
+	for (TMap<uint32, FValidInteractionSearch>::TIterator It = Searches.CreateIterator(); It; ++It)
+	{
+		const bool bIsContinuingValidSearchContexts = ContinuingOrNewValidSearchContexts.FindChecked(It->Key);
+		if (bIsContinuingValidSearchContexts)
+		{
+			OnInteractionContinuing(It->Value);
+		}
+		else
+		{
+			OnInteractionStart(It->Value);
+		}
+	}
+}
+
+void FValidInteractionSearches::OnInteractionStart(FValidInteractionSearch& ValidInteractionSearch)
+{
+#if ENABLE_VISUAL_LOG
+	ValidInteractionSearch.VLogContext(FColor::Blue);
+#endif
+
+	ValidInteractionSearch.DoDisableCollisions();
+}
+
+void FValidInteractionSearches::OnInteractionContinuing(FValidInteractionSearch& ValidInteractionSearch)
+{
+#if ENABLE_VISUAL_LOG
+	ValidInteractionSearch.VLogContext(FColor::Green);
+#endif
+}
+
+void FValidInteractionSearches::OnInteractionEnd(FValidInteractionSearch& ValidInteractionSearch)
+{
+#if ENABLE_VISUAL_LOG
+	ValidInteractionSearch.VLogContext(FColor::Black);
+#endif // ENABLE_VISUAL_LOG
+
+	ValidInteractionSearch.UndoDisableCollisions();
+}
+
+const FValidInteractionSearch* FValidInteractionSearches::Get(const FInteractionSearchContext& SearchContext) const
+{
+#if ENABLE_ANIM_DEBUG
+	UE_MT_SCOPED_READ_ACCESS(ThreadSafeCounter);
+#endif // ENABLE_ANIM_DEBUG
+
+	const uint32 SearchContextHash = SearchContext.GetHashForEquivalence();
+	return Searches.Find(SearchContextHash);
+}
+
+void FValidInteractionSearches::AddReferencedObjects(FReferenceCollector& Collector)
+{
+#if ENABLE_ANIM_DEBUG
+	UE_MT_SCOPED_READ_ACCESS(ThreadSafeCounter);
+#endif // ENABLE_ANIM_DEBUG
+
+	for (TPair<uint32, FValidInteractionSearch>& SearchPair : Searches)
+	{
+		SearchPair.Value.AddReferencedObjects(Collector);
+	}
+}
+
+// FInteractionSearchContext
+///////////////////////////////////////////////////////////
+void FInteractionSearchContext::Add(const UObject* AnimContext, const IPoseHistory* PoseHistory
+	, const FRole Role, bool bWantsDisableCollisions, int32 TickPriority
 #if ENABLE_ANIM_DEBUG
 	, float InBroadPhaseRadius, float InBroadPhaseRadiusIncrementOnInteraction
 #endif // ENABLE_ANIM_DEBUG
@@ -679,6 +1191,12 @@ void FInteractionIsland::FPreTickFunction::ExecuteTick(float DeltaTime, enum ELe
 				Island->LogTickDependencies();
 			}
 
+			if (Island->bClosingTickFunctionExecuted)
+			{
+				UE_LOG(LogPoseSearch, Error, TEXT("FInteractionIsland::FPreTickFunction::ExecuteTick, FClosingTickFunction::ExecuteTick alreay run?!"));
+				Island->LogTickDependencies();
+			}
+
 			Island->bPreTickFunctionExecuted = true;
 		}
 #endif // ENABLE_ANIM_DEBUG
@@ -710,12 +1228,56 @@ void FInteractionIsland::FPostTickFunction::ExecuteTick(float DeltaTime, enum EL
 				Island->LogTickDependencies();
 			}
 
+			if (Island->bClosingTickFunctionExecuted)
+			{
+				UE_LOG(LogPoseSearch, Error, TEXT("FInteractionIsland::FPostTickFunction::ExecuteTick, FClosingTickFunction::ExecuteTick alreay run?!"));
+				Island->LogTickDependencies();
+			}
+
 			Island->bPostTickFunctionExecuted = true;
 		}
 	}
 #endif // ENABLE_ANIM_DEBUG
 }
 
+// FClosingTickFunction
+///////////////////////////////////////////////////////////
+void FInteractionIsland::FClosingTickFunction::ExecuteTick(float DeltaTime, enum ELevelTick TickType, ENamedThreads::Type CurrentThread, const FGraphEventRef& MyCompletionGraphEvent)
+{
+	check(Island);
+
+	if (Island->HasTickDependencies())
+	{
+		CheckInteractionThreadSafety(Island);
+		Island->UpdateConstraints(DeltaTime);
+
+#if ENABLE_ANIM_DEBUG
+		if (Island->bClosingTickFunctionExecuted)
+		{
+			// @todo: need to figure out why when creating a new island FClosingTickFunction gets called twice (it's not a real issue rather than a performance hit)
+			//		  use GVarPoseSearchInteractionCacheIslands = false to debug the issue (it destroys the islands every frame)
+			UE_LOG(LogPoseSearch, Warning, TEXT("FInteractionIsland::FClosingTickFunction::ExecuteTick, called twice before UPoseSearchInteractionSubsystem::Tick!"));
+		}
+		else
+		{
+			if (!Island->bPreTickFunctionExecuted)
+			{
+				UE_LOG(LogPoseSearch, Error, TEXT("FInteractionIsland::FClosingTickFunction::ExecuteTick, FPreTickFunction::ExecuteTick didn't run!"));
+				Island->LogTickDependencies();
+			}
+
+			if (!Island->bPostTickFunctionExecuted)
+			{
+				UE_LOG(LogPoseSearch, Error, TEXT("FInteractionIsland::FClosingTickFunction::ExecuteTick, FPostTickFunction::ExecuteTick didn't run!"));
+				Island->LogTickDependencies();
+			}
+
+			Island->bClosingTickFunctionExecuted = true;
+		}
+#endif // ENABLE_ANIM_DEBUG
+	}
+}
+
 // FInteractionIsland
 ///////////////////////////////////////////////////////////
 FInteractionIsland::FInteractionIsland(ULevel* Level, UPoseSearchInteractionSubsystem* Subsystem)
@@ -730,6 +1292,11 @@ FInteractionIsland::FInteractionIsland(ULevel* Level, UPoseSearchInteractionSubs
 	PostTickFunction.Island = this;
 	PostTickFunction.RegisterTickFunction(Level);
 
+	ClosingTickFunction.bAllowTickBatching = true;
+	ClosingTickFunction.bRunOnAnyThread = true;
+	ClosingTickFunction.Island = this;
+	ClosingTickFunction.RegisterTickFunction(Level);
+
 	InteractionSubsystem = Subsystem;
 }
 
@@ -739,6 +1306,7 @@ FInteractionIsland::~FInteractionIsland()
 
 	PreTickFunction.UnRegisterTickFunction();
 	PostTickFunction.UnRegisterTickFunction();
+	ClosingTickFunction.UnRegisterTickFunction();
 
 	InteractionSubsystem = nullptr;
 }
@@ -781,6 +1349,7 @@ void FInteractionIsland::AddTickDependencies(UActorComponent* TickActorComponent
 				// adding to PreTickFunction all the tick dependencies TickActorComponent has, so it runs after all the tick dependencies of ALL the TickActorComponents in this FInteractionIsland
 				for (const FTickPrerequisite& TickActorComponentPrerequisite : TickActorComponentTickFunction->GetPrerequisites())
 				{
+					check(TickActorComponentPrerequisite.PrerequisiteTickFunction);
 					AddPrerequisite(PreTickFunction, TickActorComponentPrerequisite.PrerequisiteObject.Get(), *TickActorComponentPrerequisite.PrerequisiteTickFunction);
 				}
 
@@ -788,7 +1357,7 @@ void FInteractionIsland::AddTickDependencies(UActorComponent* TickActorComponent
 				IslandDependency->AddSubsequent(InteractionSubsystem, PreTickFunction, TickActorComponent);
 				AddPrerequisite(TickActor->PrimaryActorTick, InteractionSubsystem, PreTickFunction);
 
-				// add post-tick function dependencies to TickActorComponent and owner actor
+				// add PostTickFunction dependencies to TickActorComponent and owner actor
 				IslandDependency->AddPrerequisite(InteractionSubsystem, PostTickFunction, TickActorComponent);
 				AddPrerequisite(PostTickFunction, TickActor, TickActor->PrimaryActorTick);
 
@@ -809,13 +1378,19 @@ void FInteractionIsland::AddTickDependencies(UActorComponent* TickActorComponent
 							return TickActorComponentPrerequisite.PrerequisiteTickFunction == MainActorTickPrerequisite.Get();
 						}))
 					{
+						check(TickActorComponentPrerequisite.PrerequisiteTickFunction);
 						AddPrerequisite(PreTickFunction, TickActorComponentPrerequisite.PrerequisiteObject.Get(), *TickActorComponentPrerequisite.PrerequisiteTickFunction);
 					}
 				}
 
-				// Add post-tick function dependencies
+				// add PostTickFunction dependencies
 				IslandDependency->AddSubsequent(InteractionSubsystem, PostTickFunction, TickActorComponent);
 				AddPrerequisite(TickActor->PrimaryActorTick, InteractionSubsystem, PostTickFunction);
+
+				// add ClosingTickFunction dependencies to TickActorComponent and owner actor
+				IslandDependency->AddPrerequisite(InteractionSubsystem, ClosingTickFunction, TickActorComponent);
+				AddPrerequisite(ClosingTickFunction, TickActor, TickActor->PrimaryActorTick);
+
 				check(bHasTickDependencies);
 			}
 
@@ -837,6 +1412,7 @@ void FInteractionIsland::AddTickDependencies(UActorComponent* TickActorComponent
 			// adding to PreTickFunction all the tick dependencies TickActorComponent has, so it runs after all the tick dependencies of ALL the TickActorComponents in this FInteractionIsland
 			for (FTickPrerequisite& TickActorComponentPrerequisite : TickActorComponent->PrimaryComponentTick.GetPrerequisites())
 			{
+				check(TickActorComponentPrerequisite.PrerequisiteTickFunction);
 				AddPrerequisite(PreTickFunction, TickActorComponentPrerequisite.PrerequisiteObject.Get(), *TickActorComponentPrerequisite.PrerequisiteTickFunction);
 			}
 
@@ -844,7 +1420,7 @@ void FInteractionIsland::AddTickDependencies(UActorComponent* TickActorComponent
 			AddPrerequisite(TickActorComponent->PrimaryComponentTick, InteractionSubsystem, PreTickFunction);
 			AddPrerequisite(TickActor->PrimaryActorTick, InteractionSubsystem, PreTickFunction);
 
-			// add post-tick function dependencies to TickActorComponent and owner actor
+			// add PostTickFunction function dependencies to TickActorComponent and owner actor
 			AddPrerequisite(PostTickFunction, TickActorComponent, TickActorComponent->PrimaryComponentTick);
 			AddPrerequisite(PostTickFunction, TickActor, TickActor->PrimaryActorTick);
 
@@ -869,14 +1445,33 @@ void FInteractionIsland::AddTickDependencies(UActorComponent* TickActorComponent
 				}
 			}
 
-			// Add post-tick function dependencies
+			// add PostTickFunction dependencies
 			AddPrerequisite(TickActorComponent->PrimaryComponentTick, InteractionSubsystem, PostTickFunction);
 			AddPrerequisite(TickActor->PrimaryActorTick, InteractionSubsystem, PostTickFunction);
+
+			// add ClosingTickFunction dependencies to TickActorComponent and owner actor
+			AddPrerequisite(ClosingTickFunction, TickActorComponent, TickActorComponent->PrimaryComponentTick);
+			AddPrerequisite(ClosingTickFunction, TickActor, TickActor->PrimaryActorTick);
+
 			check(bHasTickDependencies);
 		}
 
 		// adding prerequisite between actor and skeletal mesh component so CBP runs before ABP parallel execution jobs
 		AddPrerequisite(TickActor->PrimaryActorTick, InteractionSubsystem, TickActorComponent->PrimaryComponentTick);
+
+		// @todo: find a better way to do this!
+		// adding prerequisite to the ClosingTickFunction to components having a tick dependency on TickActorComponent, so they can run AFTER the ClosingTickFunction
+		for (UActorComponent* DependantTickActorComponent : TickActor->GetComponents())
+		{
+			if (DependantTickActorComponent)
+			{
+				if (DependantTickActorComponent->PrimaryComponentTick.GetPrerequisites().Contains(FTickPrerequisite(TickActorComponent, TickActorComponent->PrimaryComponentTick)))
+				{
+					PostClosingTickActorComponents.AddUnique(DependantTickActorComponent);
+					AddPrerequisite(DependantTickActorComponent->PrimaryComponentTick, InteractionSubsystem, ClosingTickFunction);
+				}
+			}
+		}
 	}
 
 #if ENABLE_ANIM_DEBUG
@@ -887,7 +1482,9 @@ void FInteractionIsland::AddTickDependencies(UActorComponent* TickActorComponent
 		check(VisitedTickFunctions.IsEmpty());
 		bool bIsPostTickFunctionValid = ValidateTickDependenciesCycles(PostTickFunction, VisitedTickFunctions);
 		check(VisitedTickFunctions.IsEmpty());
-		if (!bIsPreTickFunctionValid || !bIsPostTickFunctionValid)
+		bool bIsClosingTickFunctionValid = ValidateTickDependenciesCycles(ClosingTickFunction, VisitedTickFunctions);
+		check(VisitedTickFunctions.IsEmpty());
+		if (!bIsPreTickFunctionValid || !bIsPostTickFunctionValid || !bIsClosingTickFunctionValid)
 		{
 			// if this validation triggers here, FInteractionIsland is not respectint the already present dependencies, creating cycles
 			UE_LOG(LogPoseSearch, Error, TEXT("============== FInteractionIsland::AddTickDependencies ValidateTickDependencies failed! Analyze the log and tune the FPoseSearchInteractionAvailability::TickPriority =============="));
@@ -907,16 +1504,20 @@ void FInteractionIsland::RemoveTickDependencies(bool bValidateTickDependencies)
 	if (!bHasTickDependencies)
 	{
 #if ENABLE_ANIM_DEBUG
-		if (bValidateTickDependencies && (bPreTickFunctionExecuted || bPostTickFunctionExecuted))
+		if (bValidateTickDependencies && (bPreTickFunctionExecuted || bPostTickFunctionExecuted || bClosingTickFunctionExecuted))
 		{
 			if (bPreTickFunctionExecuted)
 			{
 				UE_LOG(LogPoseSearch, Error, TEXT("FInteractionIsland::RemoveTickDependencies, unexpected FPreTickFunction::ExecuteTick run!"));
 			}
-			else
+			else if (bPostTickFunctionExecuted)
 			{
 				UE_LOG(LogPoseSearch, Error, TEXT("FInteractionIsland::RemoveTickDependencies, unexpected FPostTickFunction::ExecuteTick run!"));
 			}
+			else
+			{
+				UE_LOG(LogPoseSearch, Error, TEXT("FInteractionIsland::RemoveTickDependencies, unexpected FClosingTickFunction::ExecuteTick run!"));
+			}
 				
 			LogTickDependencies();
 		}
@@ -925,22 +1526,26 @@ void FInteractionIsland::RemoveTickDependencies(bool bValidateTickDependencies)
 	else
 	{
 #if ENABLE_ANIM_DEBUG
-		if (bValidateTickDependencies && (!bPreTickFunctionExecuted || !bPostTickFunctionExecuted))
+		if (bValidateTickDependencies && (!bPreTickFunctionExecuted || !bPostTickFunctionExecuted || !bClosingTickFunctionExecuted))
 		{
 			if (!bPreTickFunctionExecuted)
 			{
 				UE_LOG(LogPoseSearch, Error, TEXT("FInteractionIsland::RemoveTickDependencies, expected FPreTickFunction::ExecuteTick didn't run!"));
 			}
-			else
+			else if (!bPostTickFunctionExecuted)
 			{
 				UE_LOG(LogPoseSearch, Error, TEXT("FInteractionIsland::RemoveTickDependencies, expected FPostTickFunction::ExecuteTick didn't run!"));
 			}
+			else
+			{
+				UE_LOG(LogPoseSearch, Error, TEXT("FInteractionIsland::RemoveTickDependencies, expected FClosingTickFunction::ExecuteTick didn't run!"));
+			}
 			
 			LogTickDependencies();
 		}
 #endif // ENABLE_ANIM_DEBUG
 
-		// removing ALL the prerequisites form the PreTickFunction
+		// [undo] adding to PreTickFunction all the tick dependencies TickActorComponent has, so it runs after all the tick dependencies of ALL the TickActorComponents in this FInteractionIsland
 		while (!PreTickFunction.GetPrerequisites().IsEmpty())
 		{
 			const FTickPrerequisite& PreTickFunctionPrerequisite = PreTickFunction.GetPrerequisites().Last();
@@ -959,44 +1564,67 @@ void FInteractionIsland::RemoveTickDependencies(bool bValidateTickDependencies)
 				{
 					if (bMainActor)
 					{
+						// [undo] adding to TickActorComponent and owner actor TickActor dependencies to the PreTickFunction
 						IslandDependency->RemoveSubsequent(InteractionSubsystem, PreTickFunction, TickActorComponent);
-						IslandDependency->RemovePrerequisite(InteractionSubsystem, PostTickFunction, TickActorComponent);
-
 						TickActor->PrimaryActorTick.RemovePrerequisite(InteractionSubsystem, PreTickFunction);
+						
+						// [undo] add PostTickFunction dependencies to TickActorComponent and owner actor
+						IslandDependency->RemovePrerequisite(InteractionSubsystem, PostTickFunction, TickActorComponent);
 						PostTickFunction.RemovePrerequisite(TickActor, TickActor->PrimaryActorTick);
 					}
 					else
 					{
+						// [undo] add PostTickFunction dependencies
 						IslandDependency->RemoveSubsequent(InteractionSubsystem, PostTickFunction, TickActorComponent);
 						TickActor->PrimaryActorTick.RemovePrerequisite(InteractionSubsystem, PostTickFunction);
+						
+						// [undo] add ClosingTickFunction dependencies to TickActorComponent and owner actor
+						IslandDependency->RemovePrerequisite(InteractionSubsystem, ClosingTickFunction, TickActorComponent);
+						ClosingTickFunction.RemovePrerequisite(TickActor, TickActor->PrimaryActorTick);
 					}
 
+					// [undo] adding dependency between actor and UAF TickActorComponent so CBP runs before UAF parallel execution jobs
 					IslandDependency->RemoveSubsequent(InteractionSubsystem, TickActor->PrimaryActorTick, TickActorComponent);
 				}
 				else
 				{
 					if (bMainActor)
 					{
+						// [undo] adding to TickActorComponent and owner actor TickActor dependencies to the PreTickFunction
 						TickActorComponent->PrimaryComponentTick.RemovePrerequisite(InteractionSubsystem, PreTickFunction);
 						PostTickFunction.RemovePrerequisite(TickActorComponent, TickActorComponent->PrimaryComponentTick);
 
-						TickActor->PrimaryActorTick.RemovePrerequisite(InteractionSubsystem, PreTickFunction);
+						// [undo] add PostTickFunction function dependencies to TickActorComponent and owner actor
+						PostTickFunction.RemovePrerequisite(TickActorComponent, TickActorComponent->PrimaryComponentTick);
 						PostTickFunction.RemovePrerequisite(TickActor, TickActor->PrimaryActorTick);
 					}
 					else
 					{
+						// [undo] add PostTickFunction dependencies
 						TickActorComponent->PrimaryComponentTick.RemovePrerequisite(InteractionSubsystem, PostTickFunction);
 						TickActor->PrimaryActorTick.RemovePrerequisite(InteractionSubsystem, PostTickFunction);
+						
+						// [undo] add ClosingTickFunction dependencies to TickActorComponent and owner actor
+						ClosingTickFunction.RemovePrerequisite(TickActorComponent, TickActorComponent->PrimaryComponentTick);
+						ClosingTickFunction.RemovePrerequisite(TickActor, TickActor->PrimaryActorTick);
 					}
 
+					// [undo] adding prerequisite between actor and skeletal mesh component so CBP runs before ABP parallel execution jobs
 					TickActor->PrimaryActorTick.RemovePrerequisite(InteractionSubsystem, TickActorComponent->PrimaryComponentTick);
 				}
 			}
 			bMainActor = false;
 		}
 
+		// [undo] adding prerequisite to the ClosingTickFunction to components having a tick dependency on TickActorComponent, so they can run AFTER the ClosingTickFunction
+		for (TObjectPtr<UActorComponent>& PostClosingTickActorComponentPtr : PostClosingTickActorComponents)
+		{
+			PostClosingTickActorComponentPtr->PrimaryComponentTick.RemovePrerequisite(InteractionSubsystem, ClosingTickFunction);
+		}
+
 		check(PreTickFunction.GetPrerequisites().IsEmpty());
 		check(PostTickFunction.GetPrerequisites().IsEmpty());
+		check(ClosingTickFunction.GetPrerequisites().IsEmpty());
 
 		bHasTickDependencies = false;
 	}
@@ -1004,6 +1632,7 @@ void FInteractionIsland::RemoveTickDependencies(bool bValidateTickDependencies)
 #if ENABLE_ANIM_DEBUG
 	bPreTickFunctionExecuted = false;
 	bPostTickFunctionExecuted = false;
+	bClosingTickFunctionExecuted = false;
 #endif // ENABLE_ANIM_DEBUG
 }
 
@@ -1015,16 +1644,20 @@ void FInteractionIsland::InjectToActor(const UObject* AnimContext, bool bAddTick
 	if (AnimContext)
 	{
 #if ENABLE_ANIM_DEBUG
-		if (bPreTickFunctionExecuted || bPostTickFunctionExecuted)
+		if (bPreTickFunctionExecuted || bPostTickFunctionExecuted || bClosingTickFunctionExecuted)
 		{
 			if (bPreTickFunctionExecuted)
 			{
 				UE_LOG(LogPoseSearch, Error, TEXT("FInteractionIsland::InjectToActor, unexpected FPreTickFunction::ExecuteTick run!"));
 			}
-			else
+			else if (bPostTickFunctionExecuted)
 			{
 				UE_LOG(LogPoseSearch, Error, TEXT("FInteractionIsland::InjectToActor, unexpected FPostTickFunction::ExecuteTick run!"));
 			}
+			else
+			{
+				UE_LOG(LogPoseSearch, Error, TEXT("FInteractionIsland::InjectToActor, unexpected FClosingTickFunction::ExecuteTick run!"));
+			}
 				
 			LogTickDependencies();
 		}
@@ -1078,7 +1711,9 @@ void FInteractionIsland::Uninitialize(bool bValidateTickDependencies)
 		check(VisitedTickFunctions.IsEmpty());
 		bool bIsPostTickFunctionValid = ValidateTickDependenciesCycles(PostTickFunction, VisitedTickFunctions);
 		check(VisitedTickFunctions.IsEmpty());
-		if (!bIsPreTickFunctionValid || !bIsPostTickFunctionValid)
+		bool bIsClosingTickFunctionValid = ValidateTickDependenciesCycles(ClosingTickFunction, VisitedTickFunctions);
+		check(VisitedTickFunctions.IsEmpty());
+		if (!bIsPreTickFunctionValid || !bIsPostTickFunctionValid || !bIsClosingTickFunctionValid)
 		{
 			// if this validation triggers here, some additional tick dependency outside FInteractionIsland has been injected witout respecting the already present dependencies, creating cycles
 			UE_LOG(LogPoseSearch, Error, TEXT("============== FInteractionIsland::Uninitialize ValidateTickDependencies failed! =============="));
@@ -1092,6 +1727,7 @@ void FInteractionIsland::Uninitialize(bool bValidateTickDependencies)
 	if (IsInitialized())
 	{
 		TickActorComponents.Reset();
+		PostClosingTickActorComponents.Reset();
 		IslandAnimContexts.Reset();
 
 		SearchContexts.Reset();
@@ -1100,7 +1736,7 @@ void FInteractionIsland::Uninitialize(bool bValidateTickDependencies)
 	}
 	else
 	{
-		check(TickActorComponents.IsEmpty() && IslandAnimContexts.IsEmpty() && SearchContexts.IsEmpty() && SearchResults.IsEmpty() && !bSearchPerfomed);
+		check(TickActorComponents.IsEmpty() && PostClosingTickActorComponents.IsEmpty() && IslandAnimContexts.IsEmpty() && SearchContexts.IsEmpty() && SearchResults.IsEmpty() && !bSearchPerfomed);
 	}
 }
 
@@ -1170,13 +1806,16 @@ void FInteractionIsland::LogTickDependencies() const
 	if (IsInGameThread())
 	{
 		LogTickDependencies(TickActorComponents, InteractionIslandIndex);
+		LogTickDependencies(PostClosingTickActorComponents, InteractionIslandIndex);
 	}
 	else
 	{
 		TArray<UActorComponent*> TickActorComponentsCopy = TickActorComponents;
-		FFunctionGraphTask::CreateAndDispatchWhenReady([TickActorComponentsCopy, InteractionIslandIndex]()
+		TArray<UActorComponent*> PostClosingTickActorComponentsCopy = PostClosingTickActorComponents;
+		FFunctionGraphTask::CreateAndDispatchWhenReady([TickActorComponentsCopy, PostClosingTickActorComponentsCopy, InteractionIslandIndex]()
 		{
 			LogTickDependencies(TickActorComponentsCopy, InteractionIslandIndex);
+			LogTickDependencies(PostClosingTickActorComponentsCopy, InteractionIslandIndex);
 		}, TStatId(), nullptr, ENamedThreads::GameThread);
 	}
 }
@@ -1189,6 +1828,11 @@ void FInteractionIsland::AddReferencedObjects(FReferenceCollector& Collector)
 		Collector.AddReferencedObject(TickActorComponent);
 	}
 
+	for (TObjectPtr<UActorComponent>& PostClosingTickActorComponent : PostClosingTickActorComponents)
+	{
+		Collector.AddReferencedObject(PostClosingTickActorComponent);
+	}
+	
 	for (TObjectPtr<const UObject>& IslandAnimContext : IslandAnimContexts)
 	{
 		Collector.AddReferencedObject(IslandAnimContext);
@@ -1212,7 +1856,7 @@ bool FInteractionIsland::DoSearch_AnyThread(const UObject* AnimContext, FPoseSea
 		FMemMark Mark(FMemStack::Get());
 
 		TArray<FSearchResult, TMemStackAllocator<>> PoseSearchResults;
-		TArray<FChooserEvaluationContext, TInlineAllocator<UE::PoseSearch::PreallocatedRolesNum>> Contexts;
+		TArray<FChooserEvaluationContext, TInlineAllocator<PreallocatedRolesNum, TMemStackAllocator<>>> Contexts;
 
 		// SearchContexts are modified only by UPoseSearchInteractionSubsystem::Tick and constant otherwise, so it's safe to access them in a threaded environment without locks
 		PoseSearchResults.SetNum(SearchContexts.Num());
@@ -1235,7 +1879,7 @@ bool FInteractionIsland::DoSearch_AnyThread(const UObject* AnimContext, FPoseSea
 
 				Contexts[RoleIndex].AddObjectParam(const_cast<UObject*>(SearchContextAnimContext));
 
-				if (const UE::PoseSearch::IPoseHistory* PoseHistory = InteractionSearchContext.GetPoseHistory(RoleIndex))
+				if (const IPoseHistory* PoseHistory = InteractionSearchContext.GetPoseHistory(RoleIndex))
 				{
 					SearchContext.AddRole(InteractionSearchContext.GetRole(RoleIndex), &Contexts[RoleIndex], PoseHistory);
 				}
@@ -1263,7 +1907,7 @@ bool FInteractionIsland::DoSearch_AnyThread(const UObject* AnimContext, FPoseSea
 	return GetResult_AnyThread(AnimContext, Result);
 }
 
-bool FInteractionIsland::GetResult_AnyThread(const UObject* AnimContext, FPoseSearchBlueprintResult& Result, bool bCompareOwningActors)
+const FInteractionSearchResult* FInteractionIsland::GetInteractionSearchResult_AnyThread(const UObject* AnimContext, bool bCompareOwningActors, FRole& OutRole) const
 {
 	check(AnimContext);
 
@@ -1275,125 +1919,152 @@ bool FInteractionIsland::GetResult_AnyThread(const UObject* AnimContext, FPoseSe
 		const FInteractionSearchContext& SearchContext = SearchContexts[SearchResult.SearchIndex];
 		for (int32 AnimContextIndex = 0; AnimContextIndex < SearchContext.Num(); ++AnimContextIndex)
 		{
-			bool bAnimContextFound = false;
 			if (bCompareOwningActors)
 			{
-				bAnimContextFound = GetContextOwningActor(SearchContext.GetAnimContext(AnimContextIndex), false) == Actor;
+				if (GetContextOwningActor(SearchContext.GetAnimContext(AnimContextIndex), false) == Actor)
+				{
+					// we found our AnimContext: we can stop searching
+					OutRole = SearchContext.GetRole(AnimContextIndex);
+					return &SearchResult;
+				}
 			}
 			else
 			{
-				bAnimContextFound = SearchContext.GetAnimContext(AnimContextIndex) == AnimContext;
-			}		
+				if (SearchContext.GetAnimContext(AnimContextIndex) == AnimContext)
+				{
+					// we found our AnimContext: we can stop searching
+					OutRole = SearchContext.GetRole(AnimContextIndex);
+					return &SearchResult;
+				}
+			}
+		}
+	}
 
-			if (bAnimContextFound)
-			{
-				// @todo: perhaps add a custom Result.InitFrom(SearchResult, 1.f) for MM interactions
-				const UPoseSearchDatabase* Database = SearchResult.Database.Get();
-				check(Database);
-
-				const FSearchIndexAsset* SearchIndexAsset = SearchResult.GetSearchIndexAsset();
-				check(SearchIndexAsset);
-
-				const FPoseSearchDatabaseAnimationAssetBase* DatabaseAnimationAssetBase = Database->GetDatabaseAnimationAsset(*SearchIndexAsset);
-				check(DatabaseAnimationAssetBase);
-
-				Result.SelectedAnim = DatabaseAnimationAssetBase->GetAnimationAsset();
-				Result.SelectedTime = SearchResult.GetAssetTime();
-				Result.bIsContinuingPoseSearch = SearchResult.bIsContinuingPoseSearch;
-				Result.bLoop = SearchIndexAsset->IsLooping();
-				Result.bIsMirrored = SearchIndexAsset->IsMirrored();
-				Result.BlendParameters = SearchIndexAsset->GetBlendParameters();
-				Result.SelectedDatabase = Database;
-				Result.SearchCost = SearchResult.PoseCost;
-				Result.bIsInteraction = true;
-				Result.Role = SearchContext.GetRole(AnimContextIndex);
-
-				// figuring out the WantedPlayRate
-				Result.WantedPlayRate = 1.f;
-				//if (Future.Animation && Future.IntervalTime > 0.f)
-				//{
-				//	if (const UPoseSearchFeatureChannel_PermutationTime* PermutationTimeChannel = Database->Schema->FindFirstChannelOfType<UPoseSearchFeatureChannel_PermutationTime>())
-				//	{
-				//		const FSearchIndex& SearchIndex = Database->GetSearchIndex();
-				//		if (!SearchIndex.IsValuesEmpty())
-				//		{
-				//			TConstArrayView<float> ResultData = Database->GetSearchIndex().GetPoseValues(SearchResult.PoseIdx);
-				//			const float ActualIntervalTime = PermutationTimeChannel->GetPermutationTime(ResultData);
-				//			ProviderResult.WantedPlayRate = ActualIntervalTime / Future.IntervalTime;
-				//		}
-				//	}
-				//}
+	OutRole = DefaultRole;
+	return nullptr;
+}
 
-				if (const UMultiAnimAsset* MultiAnimAsset = Cast<UMultiAnimAsset>(Result.SelectedAnim))
-				{
-					const int32 NumRoles = MultiAnimAsset->GetNumRoles();
-					Result.ActorRootTransforms.SetNum(NumRoles);
-					Result.ActorRootBoneTransforms.SetNum(NumRoles);
-					Result.AnimContexts.SetNum(NumRoles);
+bool FInteractionIsland::GetResult_AnyThread(const UObject* AnimContext, FPoseSearchBlueprintResult& Result, bool bCompareOwningActors)
+{
+	check(AnimContext);
 
-					const FRoleToIndex InteractionSearchContextRoleToIndex = MakeRoleToIndex(SearchContext.GetRoles());
+	FRole AnimContextRole;
+	if (const FInteractionSearchResult* SearchResult = GetInteractionSearchResult_AnyThread(AnimContext, bCompareOwningActors, AnimContextRole))
+	{
+		// @todo: perhaps add a custom Result.InitFrom(SearchResult, 1.f) for MM interactions
+		const UPoseSearchDatabase* Database = SearchResult->Database.Get();
+		check(Database);
+
+		const FSearchIndexAsset* SearchIndexAsset = SearchResult->GetSearchIndexAsset();
+		check(SearchIndexAsset);
+
+		const FPoseSearchDatabaseAnimationAssetBase* DatabaseAnimationAssetBase = Database->GetDatabaseAnimationAsset(*SearchIndexAsset);
+		check(DatabaseAnimationAssetBase);
+
+		Result.SelectedAnim = DatabaseAnimationAssetBase->GetAnimationAsset();
+		Result.SelectedTime = SearchResult->GetAssetTime();
+		Result.bIsContinuingPoseSearch = SearchResult->bIsContinuingPoseSearch;
+		Result.bLoop = SearchIndexAsset->IsLooping();
+		Result.bIsMirrored = SearchIndexAsset->IsMirrored();
+		Result.BlendParameters = SearchIndexAsset->GetBlendParameters();
+		Result.SelectedDatabase = Database;
+		Result.SearchCost = SearchResult->PoseCost;
+		Result.bIsInteraction = true;
+		Result.Role = AnimContextRole;
+
+		// @todo: figuring out the WantedPlayRate
+		Result.WantedPlayRate = 1.f;
+		//if (Future.Animation && Future.IntervalTime > 0.f)
+		//{
+		//	if (const UPoseSearchFeatureChannel_PermutationTime* PermutationTimeChannel = Database->Schema->FindFirstChannelOfType<UPoseSearchFeatureChannel_PermutationTime>())
+		//	{
+		//		const FSearchIndex& SearchIndex = Database->GetSearchIndex();
+		//		if (!SearchIndex.IsValuesEmpty())
+		//		{
+		//			TConstArrayView<float> ResultData = Database->GetSearchIndex().GetPoseValues(SearchResult.PoseIdx);
+		//			const float ActualIntervalTime = PermutationTimeChannel->GetPermutationTime(ResultData);
+		//			ProviderResult.WantedPlayRate = ActualIntervalTime / Future.IntervalTime;
+		//		}
+		//	}
+		//}
+
+		const FInteractionSearchContext& SearchContext = SearchContexts[SearchResult->SearchIndex];
+		if (const UMultiAnimAsset* MultiAnimAsset = Cast<UMultiAnimAsset>(Result.SelectedAnim))
+		{
+			const int32 NumRoles = MultiAnimAsset->GetNumRoles();
+			Result.ActorRootTransforms.SetNum(NumRoles);
+			Result.ActorRootBoneTransforms.SetNum(NumRoles);
+			Result.AnimContexts.SetNum(NumRoles);
 
-					for (int32 MultiAnimAssetRoleIndex = 0; MultiAnimAssetRoleIndex < NumRoles; ++MultiAnimAssetRoleIndex)
-					{
-						const FRole MultiAnimAssetRole = MultiAnimAsset->GetRole(MultiAnimAssetRoleIndex);
-						if (MultiAnimAssetRole == Result.Role)
-						{
-							Result.RoleIndex = MultiAnimAssetRoleIndex;
-						}
+			const FRoleToIndex InteractionSearchContextRoleToIndex = MakeRoleToIndex(SearchContext.GetRoles());
 
-						if (const int32* InteractionSearchContextRoleIndex = InteractionSearchContextRoleToIndex.Find(MultiAnimAssetRole))
-						{
-							Result.ActorRootTransforms[MultiAnimAssetRoleIndex] = SearchResult.ActorRootTransforms[*InteractionSearchContextRoleIndex];
-							Result.ActorRootBoneTransforms[MultiAnimAssetRoleIndex] = SearchResult.ActorRootBoneTransforms[*InteractionSearchContextRoleIndex];
-							Result.AnimContexts[MultiAnimAssetRoleIndex] = SearchContext.GetAnimContext(*InteractionSearchContextRoleIndex);
-						}
-						else
-						{
-							Result.ActorRootTransforms[MultiAnimAssetRoleIndex] = FTransform::Identity;
-							Result.ActorRootBoneTransforms[MultiAnimAssetRoleIndex] = FTransform::Identity;
-							Result.AnimContexts[MultiAnimAssetRoleIndex] = nullptr;
-						}
-					}
+			for (int32 MultiAnimAssetRoleIndex = 0; MultiAnimAssetRoleIndex < NumRoles; ++MultiAnimAssetRoleIndex)
+			{
+				const FRole MultiAnimAssetRole = MultiAnimAsset->GetRole(MultiAnimAssetRoleIndex);
+				if (MultiAnimAssetRole == Result.Role)
+				{
+					Result.RoleIndex = MultiAnimAssetRoleIndex;
+				}
+
+				if (const int32* InteractionSearchContextRoleIndex = InteractionSearchContextRoleToIndex.Find(MultiAnimAssetRole))
+				{
+					Result.ActorRootTransforms[MultiAnimAssetRoleIndex] = SearchResult->ActorRootTransforms[*InteractionSearchContextRoleIndex];
+					Result.ActorRootBoneTransforms[MultiAnimAssetRoleIndex] = SearchResult->ActorRootBoneTransforms[*InteractionSearchContextRoleIndex];
+					Result.AnimContexts[MultiAnimAssetRoleIndex] = SearchContext.GetAnimContext(*InteractionSearchContextRoleIndex);
 				}
 				else
 				{
-					// @todo: should we support trivial "interactions" with only a character defined using some other assets rather then UMultiAnimAsset?
-					check(Result.AnimContexts.Num() == 1);
+					Result.ActorRootTransforms[MultiAnimAssetRoleIndex] = FTransform::Identity;
+					Result.ActorRootBoneTransforms[MultiAnimAssetRoleIndex] = FTransform::Identity;
+					Result.AnimContexts[MultiAnimAssetRoleIndex] = nullptr;
+				}
+			}
+		}
+		else
+		{
+			// @todo: should we support trivial "interactions" with only a character defined using some other assets rather then UMultiAnimAsset?
+			check(Result.AnimContexts.Num() == 1);
 
-					Result.ActorRootTransforms = SearchResult.ActorRootTransforms;
-					Result.ActorRootBoneTransforms = SearchResult.ActorRootBoneTransforms;
+			Result.ActorRootTransforms = SearchResult->ActorRootTransforms;
+			Result.ActorRootBoneTransforms = SearchResult->ActorRootBoneTransforms;
 					
-					Result.AnimContexts.SetNum(1);
-					Result.AnimContexts[0] = SearchContext.GetAnimContext(0);
+			Result.AnimContexts.SetNum(1);
+			Result.AnimContexts[0] = SearchContext.GetAnimContext(0);
 
-					Result.RoleIndex = 0;
-				}
-
-				// we found our AnimContext: we can stop searching
-				return true;
-			}
+			Result.RoleIndex = 0;
 		}
+
+		return true;
 	}
 
 	Result = FPoseSearchBlueprintResult();
 	return false;
 }
 
-const FInteractionSearchResult* FInteractionIsland::FindSearchResult(const FInteractionSearchContext& SearchContext) const
+void FInteractionIsland::UpdateConstraints(float DeltaTime)
 {
-	// called only by UPoseSearchInteractionSubsystem::Tick via UPoseSearchInteractionSubsystem::PopulateContinuingProperties so no need to lock SearchResultsMutex to protect the read of SearchResults
-	check(IsInGameThread());
+	check(HasTickDependencies());
+	check(InteractionSubsystem);
+		
+	for (FInteractionSearchResult& SearchResult : SearchResults)
+	{
+		const FInteractionSearchContext& SearchContext = SearchContexts[SearchResult.SearchIndex];
+		const FConstraints* SearchContextPreviousConstraints = InteractionSubsystem->GetConstraints(SearchContext);
+		SearchResult.Constraints.ProcessAll(DeltaTime, SearchResult, SearchContext, SearchContextPreviousConstraints);
+	}
+}
 
-	// searching for InSearchContext in all the SearchContexts referenced by valid active SearchResults
-	for (const FInteractionSearchResult& SearchResult : SearchResults)
+bool FInteractionIsland::GetConstraint(const UObject* AnimContext, FName BoneName, float& OutDesiredReach, FTransform& OutTransform, bool bCompareOwningActors) const
+{
+	FRole AnimContextRole;
+	if (const FInteractionSearchResult* SearchResult = GetInteractionSearchResult_AnyThread(AnimContext, bCompareOwningActors, AnimContextRole))
 	{
-		const FInteractionSearchContext& LocalSearchContext = SearchContexts[SearchResult.SearchIndex];
-		if (LocalSearchContext.IsEquivalent(SearchContext))
-		{
-			return &SearchResult;
-		}
+		return SearchResult->Constraints.GetConstraint(AnimContextRole, BoneName, OutDesiredReach, OutTransform);
 	}
-	return nullptr;
+
+	OutDesiredReach = 0.f;
+	OutTransform = FTransform::Identity;
+	return false;
 }
 
 } // namespace UE::PoseSearch
diff --git a/Engine/Plugins/Animation/PoseSearch/Source/Runtime/Private/PoseSearchInteractionLibrary.cpp b/Engine/Plugins/Animation/PoseSearch/Source/Runtime/Private/PoseSearchInteractionLibrary.cpp
--- a/Engine/Plugins/Animation/PoseSearch/Source/Runtime/Private/PoseSearchInteractionLibrary.cpp
+++ b/Engine/Plugins/Animation/PoseSearch/Source/Runtime/Private/PoseSearchInteractionLibrary.cpp
@@ -128,16 +128,6 @@ void UPoseSearchInteractionLibrary::CalculateFullAlignedTransform(const FPoseSea
 	}
 }
 
-FPoseSearchBlueprintResult UPoseSearchInteractionLibrary::GetCachedInteraction(const UObject* AnimContext, bool bCompareOwningActors)
-{
-	FPoseSearchBlueprintResult Result;
-	if (UPoseSearchInteractionSubsystem* InteractionSubsystem = UPoseSearchInteractionSubsystem::GetSubsystem_AnyThread(AnimContext))
-	{
-		InteractionSubsystem->GetResult_AnyThread(AnimContext, Result, bCompareOwningActors);
-	}
-	return Result;
-}
-
 FPoseSearchContinuingProperties UPoseSearchInteractionLibrary::GetMontageContinuingProperties(UAnimInstance* AnimInstance)
 {
 	FPoseSearchContinuingProperties ContinuingProperties;
@@ -148,3 +138,17 @@ FPoseSearchContinuingProperties UPoseSearchInteractionLibrary::GetMontageContinu
 	}
 	return ContinuingProperties;
 }
+
+void UPoseSearchInteractionLibrary::GetMotionMatchInteractionConstraint(const UObject* AnimContext, FName BoneName, float& OutDesiredReach, FTransform& OutTransform, bool& OutIsValid, bool bCompareOwningActors)
+{
+	if (UPoseSearchInteractionSubsystem* InteractionSubsystem = UPoseSearchInteractionSubsystem::GetSubsystem_AnyThread(AnimContext))
+	{
+		OutIsValid = InteractionSubsystem->GetConstraint(AnimContext, BoneName, OutDesiredReach, OutTransform, bCompareOwningActors);
+	}
+	else
+	{
+		OutDesiredReach = 0.f;
+		OutTransform = FTransform::Identity;
+		OutIsValid = false;
+	}
+}
\ No newline at end of file
diff --git a/Engine/Plugins/Animation/PoseSearch/Source/Runtime/Private/PoseSearchInteractionSubsystem.cpp b/Engine/Plugins/Animation/PoseSearch/Source/Runtime/Private/PoseSearchInteractionSubsystem.cpp
--- a/Engine/Plugins/Animation/PoseSearch/Source/Runtime/Private/PoseSearchInteractionSubsystem.cpp
+++ b/Engine/Plugins/Animation/PoseSearch/Source/Runtime/Private/PoseSearchInteractionSubsystem.cpp
@@ -521,48 +521,6 @@ bool UPoseSearchInteractionSubsystem::ValidateAllIslands() const
 }
 #endif // DO_CHECK
 
-void UPoseSearchInteractionSubsystem::PopulateContinuingProperties(float DeltaSeconds, TArrayView<UE::PoseSearch::FInteractionSearchContext> SearchContexts) const
-{
-	QUICK_SCOPE_CYCLE_COUNTER(STAT_UPoseSearchInteractionSubsystem_PopulateContinuingProperties);
-
-	using namespace UE::PoseSearch;
-
-	check(IsInGameThread());
-
-	for (FInteractionSearchContext& SearchContext : SearchContexts)
-	{
-		// searching this SearchContext in all the islands to initialize its continuing pose
-		for (const FInteractionIsland* Island : Islands)
-		{
-			if (const FSearchResult* SearchResult = Island->FindSearchResult(SearchContext))
-			{
-				// is still valid...
-				if (SearchResult->IsValid())
-				{
-					if (const UE::PoseSearch::FSearchIndexAsset* SearchIndexAsset = SearchResult->GetSearchIndexAsset())
-					{
-						if (const FPoseSearchDatabaseAnimationAssetBase* DatabaseAsset = SearchResult->Database->GetDatabaseAnimationAsset(*SearchIndexAsset))
-						{
-							check(SearchIndexAsset->GetToRealTimeFactor() > UE_KINDA_SMALL_NUMBER);
-							// in case DatabaseAsset->GetAnimationAsset() is a blendspace, SearchResult->AssetTime is a normalized time in the interval [0,1] 
-							// so we need to convert the delta time in seconds to the asset normalized time before integrating SearchResult->AssetTime
-
-							const float NormalizedDeltaTime = DeltaSeconds / SearchIndexAsset->GetToRealTimeFactor();
-							const float PlayingAssetAccumulatedTime = SearchResult->GetAssetTime() + NormalizedDeltaTime;
-
-							// @todo: populate SearchContext.InterruptMode
-							EPoseSearchInterruptMode InterruptMode = EPoseSearchInterruptMode::DoNotInterrupt;
-							SearchContext.SetContinuingProperties(PlayingAssetAccumulatedTime, DatabaseAsset->GetAnimationAsset(),
-								SearchIndexAsset->IsMirrored(), SearchIndexAsset->GetBlendParameters(), InterruptMode);
-						}
-					}
-				}
-				break;
-			}
-		}
-	}
-}
-
 UE::PoseSearch::FInteractionIsland* UPoseSearchInteractionSubsystem::FindIsland(const UObject* AnimContext, bool bCompareOwningActors)
 {
 	using namespace UE::PoseSearch;
@@ -623,7 +581,7 @@ void UPoseSearchInteractionSubsystem::Initialize(FSubsystemCollectionBase& Colle
 
 void UPoseSearchInteractionSubsystem::Deinitialize()
 {
-	UpdateValidInteractionSearches();
+	ValidInteractionSearches.Update(Islands);
 	DestroyAllIslands();
 	Super::Deinitialize();
 }
@@ -1131,7 +1089,7 @@ void UPoseSearchInteractionSubsystem::GenerateSearchContexts(float DeltaSeconds,
 			}
 		},
 		// OnDoneGroupingAnimContexts: called when the FAnimContextInfoVisitor reaches the end of a group of possibly interacting AnimContext(s)
-		[this, &DatabaseToRoledAnimContextInfos, &SearchContexts]()
+		[this, &DatabaseToRoledAnimContextInfos, &SearchContexts, DeltaSeconds]()
 		{
 			// for each database now we try to create all the possible combinations of the roled anim instances
 			// for example, given a database set up with assets for 2 characters interactions with roles RoleA and RoleB
@@ -1167,7 +1125,7 @@ void UPoseSearchInteractionSubsystem::GenerateSearchContexts(float DeltaSeconds,
 
 					GenerateCombinations(RoledAnimContextInfos.Num(), CombinationCardinality,
 						// Combination is an array of indexes in RoledAnimContextInfos: 0 <= Combination[i] < RoledAnimContextInfos.Num()
-						[this, Database, &RoledSkeletons, &RoledAnimContextInfos, &SearchContexts](const TConstArrayView<int32> Combination)
+						[this, Database, &RoledSkeletons, &RoledAnimContextInfos, &SearchContexts, DeltaSeconds](const TConstArrayView<int32> Combination)
 						{
 							// CombinationCardinality represents the number of roles as well as the number interacting AnimContext(s) (ultimately number of Characters involved in the interaction)
 							const int32 CombinationCardinality = Combination.Num();
@@ -1221,8 +1179,14 @@ void UPoseSearchInteractionSubsystem::GenerateSearchContexts(float DeltaSeconds,
 							}
 
 							// looking for a preexisting valid interaction resembling SearchContext
-							const uint32 SearchContextHash = SearchContext.GetHashForEquivalence();
-							SearchContext.SetContinuingInteraction(ValidInteractionSearches.Contains(SearchContextHash));
+							SearchContext.CacheHashForEquivalence();
+
+							// looking for a previous valid interaction with the SearchContext properties in ValidInteractionSearches 
+							// (it's using SearchContext.CacheHashForEquivalence to speed up the search)
+							const FValidInteractionSearch* ValidInteractionSearch = ValidInteractionSearches.Get(SearchContext);
+							bool bInIsContinuingInteraction = ValidInteractionSearch != nullptr;
+
+							SearchContext.SetContinuingInteraction(bInIsContinuingInteraction);
 
 							// checking if this combination is valid for the Database:
 							for (int32 CombinationIndex = 0; CombinationIndex < CombinationCardinality; ++CombinationIndex)
@@ -1240,7 +1204,7 @@ void UPoseSearchInteractionSubsystem::GenerateSearchContexts(float DeltaSeconds,
 									const float DistanceSquared = DeltaLocation.SquaredLength();
 
 									float MaxDistanceSquared;
-									if (SearchContext.IsContinuingInteraction())
+									if (bInIsContinuingInteraction)
 									{
 										MaxDistanceSquared = FMath::Min(RoledAnimContextInfo.IncrementedBroadPhaseRadiusSquared, OtherRoledAnimContextInfo.IncrementedBroadPhaseRadiusSquared);
 									}
@@ -1257,6 +1221,29 @@ void UPoseSearchInteractionSubsystem::GenerateSearchContexts(float DeltaSeconds,
 								}
 							}
 
+							// Populating Continuing properties by using the ValidInteractionSearch->SearchResult
+							if (ValidInteractionSearch)
+							{
+								const FInteractionSearchResult& PreviousSearchResult = ValidInteractionSearch->GetSearchResult();
+								if (const UE::PoseSearch::FSearchIndexAsset* SearchIndexAsset = PreviousSearchResult.GetSearchIndexAsset())
+								{
+									if (const FPoseSearchDatabaseAnimationAssetBase* DatabaseAsset = PreviousSearchResult.Database->GetDatabaseAnimationAsset(*SearchIndexAsset))
+									{
+										check(SearchIndexAsset->GetToRealTimeFactor() > UE_KINDA_SMALL_NUMBER);
+										// in case DatabaseAsset->GetAnimationAsset() is a blendspace, PreviousSearchResult.AssetTime is a normalized time in the interval [0,1] 
+										// so we need to convert the delta time in seconds to the asset normalized time before integrating PreviousSearchResult.AssetTime
+
+										const float NormalizedDeltaTime = DeltaSeconds / SearchIndexAsset->GetToRealTimeFactor();
+										const float PlayingAssetAccumulatedTime = PreviousSearchResult.GetAssetTime() + NormalizedDeltaTime;
+
+										// @todo: populate SearchContext.InterruptMode
+										EPoseSearchInterruptMode InterruptMode = EPoseSearchInterruptMode::DoNotInterrupt;
+										SearchContext.SetContinuingProperties(PlayingAssetAccumulatedTime, DatabaseAsset->GetAnimationAsset(),
+											SearchIndexAsset->IsMirrored(), SearchIndexAsset->GetBlendParameters(), InterruptMode);
+									}
+								}
+							}
+
 							SearchContexts.Emplace(SearchContext);
 							return true;
 						});
@@ -1266,91 +1253,6 @@ void UPoseSearchInteractionSubsystem::GenerateSearchContexts(float DeltaSeconds,
 			// done using DatabaseToRoledAnimContextInfos. clearing up for the next group of AnimContext(s)
 			DatabaseToRoledAnimContextInfos.Reset();
 		});
-
-	// populating the continuing pose properties for the SearchContexts from the current Islands
-	PopulateContinuingProperties(DeltaSeconds, SearchContexts);
-}
-
-void UPoseSearchInteractionSubsystem::OnInteractionStart(UE::PoseSearch::FValidInteractionSearch& ValidInteractionSearch)
-{
-	using namespace UE::PoseSearch;
-
-#if ENABLE_VISUAL_LOG
-	ValidInteractionSearch.VLogContext(FColor::Blue);
-#endif
-
-	ValidInteractionSearch.DoDisableCollisions();
-}
-
-void UPoseSearchInteractionSubsystem::OnInteractionContinuing(UE::PoseSearch::FValidInteractionSearch& ValidInteractionSearch)
-{
-#if ENABLE_VISUAL_LOG
-	ValidInteractionSearch.VLogContext(FColor::Green);
-#endif
-}
-
-void UPoseSearchInteractionSubsystem::OnInteractionEnd(UE::PoseSearch::FValidInteractionSearch& ValidInteractionSearch)
-{
-	using namespace UE::PoseSearch;
-
-#if ENABLE_VISUAL_LOG
-	ValidInteractionSearch.VLogContext(FColor::Black);
-#endif
-
-	ValidInteractionSearch.UndoDisableCollisions();
-}
-
-void UPoseSearchInteractionSubsystem::UpdateValidInteractionSearches()
-{
-	using namespace UE::PoseSearch;
-
-	FMemMark Mark(FMemStack::Get());
-	TMap<uint32, bool, TInlineSetAllocator<32, TMemStackSetAllocator<>>> ContinuingOrNewValidSearchContexts;
-	for (FInteractionIsland* Island : Islands)
-	{
-		if (Island->IsInitialized())
-		{
-			// analyzing ALL current tick interaction results
-			for (const FInteractionSearchResult& SearchResult : Island->GetSearchResults())
-			{
-				const FInteractionSearchContext& SearchContext = Island->GetSearchContexts()[SearchResult.SearchIndex];
-				const uint32 ValidSearchContextHash = SearchContext.GetHashForEquivalence();
-				if (FValidInteractionSearch* ContinuingValidInteractionSearch = ValidInteractionSearches.Find(ValidSearchContextHash))
-				{
-					ContinuingOrNewValidSearchContexts.Add(ValidSearchContextHash) = true;
-				}
-				else
-				{
-					ContinuingOrNewValidSearchContexts.Add(ValidSearchContextHash) = false;
-					ValidInteractionSearches.Add(ValidSearchContextHash).Init(SearchContext);
-				}
-			}
-		}
-	}
-
-	// calling all the OnInteractionEnd first
-	for (TMap<uint32, UE::PoseSearch::FValidInteractionSearch>::TIterator It = ValidInteractionSearches.CreateIterator(); It; ++It)
-	{
-		if (!ContinuingOrNewValidSearchContexts.Find(It->Key))
-		{
-			OnInteractionEnd(It->Value);
-			It.RemoveCurrent();
-		}
-	}
-
-	// calling all the OnInteractionContinuing and OnInteractionStart after
-	for (TMap<uint32, UE::PoseSearch::FValidInteractionSearch>::TIterator It = ValidInteractionSearches.CreateIterator(); It; ++It)
-	{
-		const bool bIsContinuingValidSearchContexts = ContinuingOrNewValidSearchContexts.FindChecked(It->Key);
-		if (bIsContinuingValidSearchContexts)
-		{
-			OnInteractionContinuing(It->Value);
-		}
-		else
-		{
-			OnInteractionStart(It->Value);
-		}
-	}
 }
 
 bool UPoseSearchInteractionSubsystem::ShouldRegenerateAllIslands()
@@ -1399,7 +1301,7 @@ void UPoseSearchInteractionSubsystem::Tick(float DeltaSeconds)
 	Super::Tick(DeltaSeconds);
 
 	FMemMark Mark(FMemStack::Get());
-	UpdateValidInteractionSearches();
+	ValidInteractionSearches.Update(Islands);
 
 	if (ShouldRegenerateAllIslands())
 	{
@@ -1478,6 +1380,31 @@ void UPoseSearchInteractionSubsystem::GetResult_AnyThread(const UObject* AnimCon
 	}
 }
 
+const UE::PoseSearch::FConstraints* UPoseSearchInteractionSubsystem::GetConstraints(const UE::PoseSearch::FInteractionSearchContext& SearchContext) const
+{
+	using namespace UE::PoseSearch;
+
+	if (const FValidInteractionSearch* ValidInteractionSearch = ValidInteractionSearches.Get(SearchContext))
+	{
+		return &ValidInteractionSearch->GetSearchResult().Constraints;
+	}
+	return nullptr;
+}
+
+bool UPoseSearchInteractionSubsystem::GetConstraint(const UObject* AnimContext, FName BoneName, float& OutDesiredReach, FTransform& OutTransform, bool bCompareOwningActors)
+{
+	using namespace UE::PoseSearch;
+
+	if (const FInteractionIsland* Island = FindIsland(AnimContext, bCompareOwningActors))
+	{
+		return Island->GetConstraint(AnimContext, BoneName, OutDesiredReach, OutTransform, bCompareOwningActors);
+	}
+	
+	OutDesiredReach = 0.f;
+	OutTransform = FTransform::Identity;
+	return false;
+}
+
 void UPoseSearchInteractionSubsystem::AddReferencedObjects(UObject* This, FReferenceCollector& Collector)
 {
 	using namespace UE::PoseSearch;
@@ -1500,10 +1427,7 @@ void UPoseSearchInteractionSubsystem::AddReferencedObjects(UObject* This, FRefer
 		Island->AddReferencedObjects(Collector);
 	}
 
-	for (TPair<uint32, FValidInteractionSearch>& ValidInteractionSearchPair : Subsystem->ValidInteractionSearches)
-	{
-		ValidInteractionSearchPair.Value.AddReferencedObjects(Collector);
-	}
+	Subsystem->ValidInteractionSearches.AddReferencedObjects(Collector);
 }
 
 #if ENABLE_ANIM_DEBUG
diff --git a/Engine/Plugins/Animation/PoseSearch/Source/Runtime/Public/PoseSearch/PoseSearchAnimNotifies.h b/Engine/Plugins/Animation/PoseSearch/Source/Runtime/Public/PoseSearch/PoseSearchAnimNotifies.h
--- a/Engine/Plugins/Animation/PoseSearch/Source/Runtime/Public/PoseSearch/PoseSearchAnimNotifies.h
+++ b/Engine/Plugins/Animation/PoseSearch/Source/Runtime/Public/PoseSearch/PoseSearchAnimNotifies.h
@@ -122,45 +122,43 @@ public:
 	UE_API uint32 GetBranchInId() const;
 };
 
-// multi character interaction notifies
-USTRUCT(Experimental, Blueprintable)
-struct FPoseSearchIKWindowConstraint
+// multi character interaction constraint notify state: used to calculate relative bone positions in UPoseSearchInteractionAsset items
+UCLASS(MinimalAPI, Experimental, Blueprintable, meta = (DisplayName = "Pose Search: Constraint"))
+class UAnimNotifyState_PoseSearchConstraint : public UAnimNotifyState_PoseSearchBase
 {
 	GENERATED_BODY()
-	
+
+public:
 	UPROPERTY(EditAnywhere, Category = Config)
-	FName FromGoalName;
+	FName FromBoneName;
 
 	UPROPERTY(EditAnywhere, Category = Config)
-	FName FromGoalRole;
+	FName FromBoneRole;
 	
 	UPROPERTY(EditAnywhere, Category = Config)
-	FName ToGoalName;
+	FName ToBoneName;
 	
 	UPROPERTY(EditAnywhere, Category = Config)
-	FName ToGoalRole;
+	FName ToBoneRole;
 
+	// who should translate to reach full aligment? FromBone or ToBone? If TranslationWeight is 0 FromBone will move 100% and ToBone 0%, if TranslationWeight is 1 FromBone will move 0% and ToBone 100%,
 	UPROPERTY(EditAnywhere, Category = Config, meta = (ClampMin = "0", ClampMax = "1", UIMin = "0", UIMax = "1"))
 	float TranslationWeight = 0.f;
 	
+	// who should rotate to reach full aligment? FromBone or ToBone? If RotationWeight is 0 FromBone will move 100% and ToBone 0%, if RotationWeight is 1 FromBone will move 0% and ToBone 100%,
 	UPROPERTY(EditAnywhere, Category = Config, meta = (ClampMin = "0", ClampMax = "1", UIMin = "0", UIMax = "1"))
 	float RotationWeight = 0.f;
-	
-	UPROPERTY(EditAnywhere, Category = Config, meta = (ClampMin = "0", UIMin = "0"))
-	float ActivationTime = 0.f;
-	
-	UPROPERTY(EditAnywhere, Category = Config, meta = (ClampMin = "0", UIMin = "0"))
-	float DeactivationTime = 0.f;
-};
 
-UCLASS(MinimalAPI, Experimental, Blueprintable, meta = (DisplayName = "Pose Search: IKWindow"))
-class UAnimNotifyState_PoseSearchIKWindow : public UAnimNotifyState_PoseSearchBase
-{
-	GENERATED_BODY()
+	// rump up time used to calculate the desired reach for this constraint 
+	// (it doesn't start at the beginning of this notify state, BUT when this notify state is first seen, since MM can jump into a random frame of the animation containing the notify state)
+	UPROPERTY(EditAnywhere, Category = Config, meta = (ClampMin = "0", UIMin = "0"))
+	float RampUpTime = 0.2f;
 
-public:
-	UPROPERTY(EditAnywhere, Category = Config)
-	TArray<FPoseSearchIKWindowConstraint> Constraints;
+	// cool down time used to calculate the desired reach for this constraint
+	// (it doesn't start at the end of this notify state, BUT when this notify state gets out of scope
+	// (ultimately getting out of scope when reaching the notify state end time minus CoolDownTime, since MM can jump into a random frame of the animation containing the notify state)
+	UPROPERTY(EditAnywhere, Category = Config, meta = (ClampMin = "0", UIMin = "0"))
+	float CoolDownTime = 0.2f;
 };
 
 // Base class for ALL pose search anim notifies
diff --git a/Engine/Plugins/Animation/PoseSearch/Source/Runtime/Public/PoseSearch/PoseSearchAssetSampler.h b/Engine/Plugins/Animation/PoseSearch/Source/Runtime/Public/PoseSearch/PoseSearchAssetSampler.h
--- a/Engine/Plugins/Animation/PoseSearch/Source/Runtime/Public/PoseSearch/PoseSearchAssetSampler.h
+++ b/Engine/Plugins/Animation/PoseSearch/Source/Runtime/Public/PoseSearch/PoseSearchAssetSampler.h
@@ -47,7 +47,9 @@ struct FAnimationAssetSampler
 	POSESEARCH_API FTransform ExtractRootTransform(float Time) const;
 	
 	// Extracts notify states present in the sequence at Time
+	POSESEARCH_API void ExtractAnimNotifyContext(float Time, FAnimNotifyContext& PreAllocatedNotifyContext) const;
 	POSESEARCH_API void ExtractAnimNotifyStates(float Time, FAnimNotifyContext& PreAllocatedNotifyContext, const TFunction<bool(UAnimNotifyState*)>& ProcessAnimNotifyState) const;
+	POSESEARCH_API void ExtractAnimNotifyEvents(float Time, FAnimNotifyContext& PreAllocatedNotifyContext, const TFunction<bool(const FAnimNotifyEvent*)>& ProcessAnimNotifyEvent) const;
 	
 	POSESEARCH_API TConstArrayView<FAnimNotifyEvent> GetAllAnimNotifyEvents() const;
 
diff --git a/Engine/Plugins/Animation/PoseSearch/Source/Runtime/Public/PoseSearch/PoseSearchInteractionIsland.h b/Engine/Plugins/Animation/PoseSearch/Source/Runtime/Public/PoseSearch/PoseSearchInteractionIsland.h
--- a/Engine/Plugins/Animation/PoseSearch/Source/Runtime/Public/PoseSearch/PoseSearchInteractionIsland.h
+++ b/Engine/Plugins/Animation/PoseSearch/Source/Runtime/Public/PoseSearch/PoseSearchInteractionIsland.h
@@ -3,20 +3,21 @@
 #pragma once
 
 #include "Engine/EngineBaseTypes.h"
+#include "PoseSearch/PoseSearchAnimNotifies.h"
 #include "PoseSearch/PoseSearchInteractionAvailability.h"
 #include "PoseSearch/PoseSearchInteractionValidator.h"
 #include "PoseSearch/PoseSearchLibrary.h"
 
-struct FPoseSearchBlueprintResult;
-struct FPoseSearchContinuingProperties;
 class UMultiAnimAsset;
 class UPoseSearchInteractionSubsystem;
+struct FPoseSearchBlueprintResult;
+struct FPoseSearchContinuingProperties;
 
 namespace UE::PoseSearch
 {
 struct FInteractionIsland;
-struct IInteractionIslandDependency;
 struct FInteractionSearchContext;
+struct IInteractionIslandDependency;
 
 // Experimental, this feature might be removed without warning, not for production use
 typedef TPair<TObjectPtr<AActor>, TObjectPtr<AActor>> FDisabledCollision;
@@ -28,41 +29,119 @@ struct FInteractionSearchContextBase
 {
 	int32 Num() const { return AnimContexts.Num(); }
 	const UObject* GetAnimContext(int32 Index) const { return AnimContexts[Index]; }
-	const UE::PoseSearch::IPoseHistory* GetPoseHistory(int32 Index) const;
-	const UE::PoseSearch::FRole GetRole(int32 Index) const { return Roles[Index]; }
+	const IPoseHistory* GetPoseHistory(int32 Index) const;
+	const FRole GetRole(int32 Index) const { return Roles[Index]; }
 	const UPoseSearchDatabase* GetDatabase() const { return Database; }
-	const TConstArrayView<UE::PoseSearch::FRole> GetRoles() const { return Roles; }
+	const TConstArrayView<FRole> GetRoles() const { return Roles; }
 
 	bool IsDisableCollisions() const { return bDisableCollisions; }
 
 	uint32 GetHashForEquivalence() const;
-	bool IsEquivalent(const FInteractionSearchContextBase& Other) const;
+	void CacheHashForEquivalence();
 	
 	void SetDatabase(const UPoseSearchDatabase* InDatabase) { Database = InDatabase; }
 	void AddReferencedObjects(FReferenceCollector& Collector);
-	void Add(const UObject* AnimContext, const UE::PoseSearch::IPoseHistory* PoseHistory, const UE::PoseSearch::FRole Role, bool bWantsDisableCollisions);
+	void Add(const UObject* AnimContext, const IPoseHistory* PoseHistory, const FRole Role, bool bWantsDisableCollisions);
 
 #if DO_CHECK
 	bool CheckForConsistency() const;
 #endif // DO_CHECK
 
 private:
+	uint32 CalculateHashForEquivalence() const;
+
 	TArray<TObjectPtr<const UObject>, TInlineAllocator<PreallocatedRolesNum>> AnimContexts;
 
 	// storing PoseHistories as TWeakPtr since they are from nodes or traits that can be invalidated within the same frame.
 	// for performance reasons are not used to determine equivalence IsEquivalent / GetHashForEquivalence
-	TArray<TWeakPtr<const UE::PoseSearch::IPoseHistory, ESPMode::ThreadSafe>, TInlineAllocator<PreallocatedRolesNum>> PoseHistories;
+	TArray<TWeakPtr<const IPoseHistory, ESPMode::ThreadSafe>, TInlineAllocator<PreallocatedRolesNum>> PoseHistories;
 	
-	TArray<UE::PoseSearch::FRole, TInlineAllocator<PreallocatedRolesNum>> Roles;
+	TArray<FRole, TInlineAllocator<PreallocatedRolesNum>> Roles;
 
 	TObjectPtr<const UPoseSearchDatabase> Database;
 	bool bDisableCollisions = false;
+	enum { InvalidCachedHashForEquivalence = 0 };
+	uint32 CachedHashForEquivalence = InvalidCachedHashForEquivalence;
+};
+
+struct FConstraint
+{
+	FConstraint() = default;
+	FConstraint(const FConstraint& Other) = default;
+	FConstraint(FConstraint&& Other) = default;
+	FConstraint& operator=(const FConstraint& Other) = default;
+	FConstraint& operator=(FConstraint&& Other) = default;
+
+	FConstraint(const UAnimNotifyState_PoseSearchConstraint* ConstraintNotifyState);
+	bool operator==(const FConstraint& Other) const;
+	const FTransform* GetBoneTransform(FRole Role, FName BoneName) const;
+
+	FName FromBoneName;
+	FName FromBoneRole;
+	FName ToBoneName;
+	FName ToBoneRole;
+
+	// who should translate to reach full aligment? FromBone or ToBone? If TranslationWeight is 0 FromBone will move 100% and ToBone 0%, if TranslationWeight is 1 FromBone will move 0% and ToBone 100%,
+	float TranslationWeight = 0.f;
+	
+	// who should rotate to reach full aligment? FromBone or ToBone? If RotationWeight is 0 FromBone will move 100% and ToBone 0%, if RotationWeight is 1 FromBone will move 0% and ToBone 100%,
+	float RotationWeight = 0.f;
+
+	// rump up time used to calculate the desired reach for this constraint 
+	// (it doesn't start at the beginning of this notify state, BUT when this notify state is first seen, since MM can jump into a random frame of the animation containing the notify state)
+	float RampUpTime = 0.2f;
+
+	// cool down time used to calculate the desired reach for this constraint
+	// (it doesn't start at the end of this notify state, BUT when this notify state gets out of scope
+	// (ultimately getting out of scope when reaching the notify state end time minus CoolDownTime, since MM can jump into a random frame of the animation containing the notify state)
+	float CoolDownTime = 0.2f;
+
+	FTransform FromBoneTransform;
+	FTransform ToBoneTransform;
+
+	float DesiredReach = 0.f;
+};
+
+struct FConstraints
+{
+	// @todo: mitigate / budget allocations properly
+	enum { NumPreallocatedContraints = 4 };
+
+	void ProcessAll(float DeltaTime, const FSearchResult& SearchResult, const FInteractionSearchContext& SearchContext, const FConstraints* SearchContextPreviousConstraints);
+
+	bool GetConstraint(FRole Role, FName BoneName, float& OutDesiredReach, FTransform& OutTransform) const;
+
+private:
+	int32 Num() const;
+	const FConstraint* Find(const FConstraint& Constraint) const;
+	FConstraint* Find(const FConstraint& Constraint);
+	TArray<FConstraint, TInlineAllocator<NumPreallocatedContraints>> Constraints;
+};
+
+// Experimental, this feature might be removed without warning, not for production use
+struct FInteractionSearchResult : public FSearchResult
+{
+	int32 SearchIndex = INDEX_NONE;
+
+	// cached actors root transforms for all the roles in SelectedAnimation (as UMultiAnimAsset),
+	// so we don't have to query the pose history to gather it when it's not thread safe to do so
+	TArray<FTransform, TInlineAllocator<PreallocatedRolesNum>> ActorRootTransforms;
+
+	// cached actors root bone transforms for all the roles in SelectedAnimation (as UMultiAnimAsset),
+	// so we don't have to query the pose history to gather it when it's not thread safe to do so
+	TArray<FTransform, TInlineAllocator<PreallocatedRolesNum>> ActorRootBoneTransforms;
+
+	// all the active constraints in this result
+	FConstraints Constraints;
+
+	bool operator==(const FInteractionSearchResult& Other) const;
 };
 
 // Experimental, this feature might be removed without warning, not for production use
 struct FValidInteractionSearch : public FInteractionSearchContextBase
 {
-	void Init(const FInteractionSearchContext& Other);
+	FValidInteractionSearch& Init(const FInteractionSearchContextBase& SearchContextBase);
+	void Update(const FInteractionSearchResult& InSearchResult);
 	void DoDisableCollisions();
 	void UndoDisableCollisions();
 
@@ -71,18 +150,43 @@ struct FValidInteractionSearch : public FInteractionSearchContextBase
 #endif // ENABLE_VISUAL_LOG
 
 	void AddReferencedObjects(FReferenceCollector& Collector);
+	const FInteractionSearchResult& GetSearchResult() const { return SearchResult; }
 
 private:
+	FInteractionSearchResult SearchResult;
 	FDisabledCollisions DisabledCollisions;
+
+#if ENABLE_ANIM_DEBUG
+	UE_MT_DECLARE_TS_RW_ACCESS_DETECTOR(ThreadSafeCounter);
+#endif // ENABLE_ANIM_DEBUG
 };
 
+struct FValidInteractionSearches
+{
+	void Update(TConstArrayView<FInteractionIsland*> Islands);
+	//FValidInteractionSearch* Get(const FInteractionSearchContext& SearchContext);
+	const FValidInteractionSearch* Get(const FInteractionSearchContext& SearchContext) const;
+	void AddReferencedObjects(FReferenceCollector& Collector);
+
+private:
+	static void OnInteractionStart(FValidInteractionSearch& ValidInteractionSearch);
+	static void OnInteractionContinuing(FValidInteractionSearch& ValidInteractionSearch);
+	static void OnInteractionEnd(FValidInteractionSearch& ValidInteractionSearch);
+
+	TMap<uint32, FValidInteractionSearch> Searches;
+#if ENABLE_ANIM_DEBUG
+	UE_MT_DECLARE_TS_RW_ACCESS_DETECTOR(ThreadSafeCounter);
+#endif // ENABLE_ANIM_DEBUG
+};
+
+
 // Experimental, this feature might be removed without warning, not for production use
 // stack allocated temporary struct used to generate interaction contexts, then the FInteractionSearchContextBase portion will be 
 // stored including DisabledCollisions as FValidInteractionSearch
 struct FInteractionSearchContext : public FInteractionSearchContextBase
 {
-	void Add(const UObject* AnimContext, const UE::PoseSearch::IPoseHistory* PoseHistory
-		, const UE::PoseSearch::FRole Role, bool bWantsDisableCollisions, int32 TickPriority
+	void Add(const UObject* AnimContext, const IPoseHistory* PoseHistory
+		, const FRole Role, bool bWantsDisableCollisions, int32 TickPriority
 #if ENABLE_ANIM_DEBUG
 		, float InBroadPhaseRadius, float InBroadPhaseRadiusIncrementOnInteraction
 #endif // ENABLE_ANIM_DEBUG
@@ -119,22 +223,6 @@ private:
 // Experimental, this feature might be removed without warning, not for production use
 typedef TArray<FInteractionSearchContext, TInlineAllocator<PreallocatedSearchesNum, TMemStackAllocator<>>> FInteractionSearchContexts;
 
-// Experimental, this feature might be removed without warning, not for production use
-struct FInteractionSearchResult : public FSearchResult
-{
-	int32 SearchIndex = INDEX_NONE;
-
-	// cached actors root transforms for all the roles in SelectedAnimation (as UMultiAnimAsset),
-	// so we don't have to query the pose history to gather it when it's not thread safe to do so
-	TArray<FTransform, TInlineAllocator<PreallocatedRolesNum>> ActorRootTransforms;
-
-	// cached actors root bone transforms for all the roles in SelectedAnimation (as UMultiAnimAsset),
-	// so we don't have to query the pose history to gather it when it's not thread safe to do so
-	TArray<FTransform, TInlineAllocator<PreallocatedRolesNum>> ActorRootBoneTransforms;
-
-	bool operator==(const FInteractionSearchResult& Other) const;
-};
-
 // Experimental, this feature might be removed without warning, not for production use
 // FInteractionIsland contains ticks functions injected between the interacting actors TickActorComponents (UCharacterMovementComponent, or UCharacterMoverComponent)
 // and PostTickComponent (USkeletalMeshComponent, or UAnimNextComponent)
@@ -155,8 +243,6 @@ struct FInteractionIsland
 	const TArray<FInteractionSearchContext>& GetSearchContexts() const { return SearchContexts; }
 	const TArray<FInteractionSearchResult>& GetSearchResults() const { return SearchResults; }
 
-	const FInteractionSearchResult* FindSearchResult(const FInteractionSearchContext& SearchContext) const;
-
 	bool IsInitialized() const;
 	void AddSearchContext(const FInteractionSearchContext& SearchContext);
 	void Uninitialize(bool bValidateTickDependencies);
@@ -169,6 +255,8 @@ struct FInteractionIsland
 #endif // ENABLE_ANIM_DEBUG
 
 	void AddReferencedObjects(FReferenceCollector& Collector);
+	void UpdateConstraints(float DeltaTime);
+	bool GetConstraint(const UObject* AnimContext, FName BoneName, float& OutDesiredReach, FTransform& OutTransform, bool bCompareOwningActors = false) const;
 
 private:
 #if ENABLE_ANIM_DEBUG
@@ -179,9 +267,12 @@ private:
 	void RemoveTickDependencies(bool bValidateTickDependencies);
 	static IInteractionIslandDependency* FindCustomDependency(UActorComponent* InTickComponent);
 
+	const FInteractionSearchResult* GetInteractionSearchResult_AnyThread(const UObject* AnimContext, bool bCompareOwningActors, FRole& OutRole) const;
+
 	const UObject* GetMainAnimContext() const;
 	const AActor* GetMainActor() const;
 
+	// running before main actor animation evaluation
 	struct FPreTickFunction : public FTickFunction
 	{
 		virtual void ExecuteTick(float DeltaTime, enum ELevelTick TickType, ENamedThreads::Type CurrentThread, const FGraphEventRef& MyCompletionGraphEvent) override;
@@ -190,6 +281,7 @@ private:
 	};
 	FPreTickFunction PreTickFunction;
 
+	// running after main actor animation evaluation, and before any NON main actor animation evaluation
 	struct FPostTickFunction : public FTickFunction
 	{
 		virtual void ExecuteTick(float DeltaTime, enum ELevelTick TickType, ENamedThreads::Type CurrentThread, const FGraphEventRef& MyCompletionGraphEvent) override;
@@ -197,9 +289,22 @@ private:
 		FInteractionIsland* Island = nullptr;
 	};
 	FPostTickFunction PostTickFunction;
+
+	// running ALL actors animation evaluation
+	struct FClosingTickFunction : public FTickFunction
+	{
+		virtual void ExecuteTick(float DeltaTime, enum ELevelTick TickType, ENamedThreads::Type CurrentThread, const FGraphEventRef& MyCompletionGraphEvent) override;
+		virtual FString DiagnosticMessage() override { return TEXT("UE::PoseSearch::FInteractionIsland::FClosingTickFunction"); }
+		FInteractionIsland* Island = nullptr;
+	};
+	FClosingTickFunction ClosingTickFunction;
+
 	bool bHasTickDependencies = false;
 
 	TArray<TObjectPtr<UActorComponent>> TickActorComponents;
+	// UActorComponent with an ijected tick dependency to the ClosingTickFunction
+	TArray<TObjectPtr<UActorComponent>> PostClosingTickActorComponents;
+
 	// all the AnimContext in this island. Each SearchContexts will contain a subset of IslandAnimContexts 
 	TArray<TObjectPtr<const UObject>> IslandAnimContexts;
 
@@ -221,6 +326,7 @@ private:
 	FThreadSafeCounter TickFunctionsThreadSafeCounter = 0;
 	bool bPreTickFunctionExecuted = false;
 	bool bPostTickFunctionExecuted = false;
+	bool bClosingTickFunctionExecuted = false;
 #endif // ENABLE_ANIM_DEBUG
 };
 
diff --git a/Engine/Plugins/Animation/PoseSearch/Source/Runtime/Public/PoseSearch/PoseSearchInteractionLibrary.h b/Engine/Plugins/Animation/PoseSearch/Source/Runtime/Public/PoseSearch/PoseSearchInteractionLibrary.h
--- a/Engine/Plugins/Animation/PoseSearch/Source/Runtime/Public/PoseSearch/PoseSearchInteractionLibrary.h
+++ b/Engine/Plugins/Animation/PoseSearch/Source/Runtime/Public/PoseSearch/PoseSearchInteractionLibrary.h
@@ -42,10 +42,11 @@ public:
 	static UE_API void MotionMatchInteraction(FPoseSearchBlueprintResult& Result, const TArrayView<const FPoseSearchInteractionAvailability> Availabilities, const UObject* AnimContext, FName PoseHistoryName, const UE::PoseSearch::IPoseHistory* PoseHistory, bool bValidateResultAgainstAvailabilities = true);
 	static UE_API bool MotionMatchInteraction(FPoseSearchBlueprintResult& InOutResult, const TArrayView<const FPoseSearchInteractionAvailability> Availabilities, const UObject* AnimContext, FName PoseHistoryName, const UE::PoseSearch::IPoseHistory* PoseHistory, bool bValidateResultAgainstAvailabilities, bool bKeepInteractionAlive, float BlendTime);
 
-	static UE_API FPoseSearchBlueprintResult GetCachedInteraction(const UObject* AnimContext, bool bCompareOwningActors = false);
-	
 	UFUNCTION(Experimental, BlueprintPure, Category = "Animation|Pose Search", meta = (BlueprintThreadSafe))
 	static UE_API FPoseSearchContinuingProperties GetMontageContinuingProperties(UAnimInstance* AnimInstance);
+
+	UFUNCTION(Experimental, BlueprintPure, Category = "Animation|Pose Search", meta = (BlueprintThreadSafe))
+	static UE_API void GetMotionMatchInteractionConstraint(const UObject* AnimContext, FName BoneName, float& OutDesiredReach, FTransform& OutTransform, bool& IsValid, bool bCompareOwningActors = true);
 };
 
 #undef UE_API
diff --git a/Engine/Plugins/Animation/PoseSearch/Source/Runtime/Public/PoseSearch/PoseSearchInteractionSubsystem.h b/Engine/Plugins/Animation/PoseSearch/Source/Runtime/Public/PoseSearch/PoseSearchInteractionSubsystem.h
--- a/Engine/Plugins/Animation/PoseSearch/Source/Runtime/Public/PoseSearch/PoseSearchInteractionSubsystem.h
+++ b/Engine/Plugins/Animation/PoseSearch/Source/Runtime/Public/PoseSearch/PoseSearchInteractionSubsystem.h
@@ -57,6 +57,13 @@ struct FTagToDatabases;
 // 						Tick SkeletalMeshComponent[k][i] (or AnimNextComponent(s) that DoSearch_AnyThread get the cached result calculated by Tick Actor[k][0].SkeletalMeshComponent) in parallel
 //					}
 //				}
+// 
+//				Tick Island[k].ClosingTickFunction (that calculates all the necessary constraints)
+// 
+//				for each PostClosingTickActorComponent (stored in FInteractionIsland::PostClosingTickActorComponents) with tick dependencies to SkeletalMeshComponent[k][i] (or AnimNextComponent(s))
+//				{
+//					Tick PostClosingTickActorComponent
+//				}
 //			}
 // - next frame UPoseSearchInteractionLibrary::MotionMatchInteraction_Pure(TArray<FPoseSearchInteractionAvailability> Availabilities, UObject* AnimInstance), with the context of all the published 
 //   availabilities and created islands, will find the associated FInteractionIsland to the AnimInstance and call FInteractionIsland::DoSearch_AnyThread 
@@ -90,6 +97,8 @@ public:
 	UE_API void Query_AnyThread(const TArrayView<const FPoseSearchInteractionAvailability> Availabilities, const UObject* AnimContext, 
 		FPoseSearchBlueprintResult& Result,	FName PoseHistoryName, const UE::PoseSearch::IPoseHistory* PoseHistory, bool bValidateResultAgainstAvailabilities);
 
+	// use bCompareOwningActors = true in case you're querying the UPoseSearchInteractionSubsystem from a different AnimContext, 
+	// for example when using this API from a different ABP from a different skeletal mesh component
 	UE_API void GetResult_AnyThread(const UObject* AnimContext, FPoseSearchBlueprintResult& Result, bool bCompareOwningActors = false);
 
 	TConstArrayView<UE::PoseSearch::FInteractionIsland*> GetInteractionIslands() const { return Islands; }
@@ -97,6 +106,12 @@ public:
 	// garbage collection handling
 	static UE_API void AddReferencedObjects(UObject* This, FReferenceCollector& Collector);
 
+	const UE::PoseSearch::FConstraints* GetConstraints(const UE::PoseSearch::FInteractionSearchContext& SearchContext) const;
+
+	// use bCompareOwningActors = true in case you're querying the UPoseSearchInteractionSubsystem from a different AnimContext, 
+	// for example when using this API from a different ABP from a different skeletal mesh component
+	bool GetConstraint(const UObject* AnimContext, FName BoneName, float& OutDesiredReach, FTransform& OutTransform, bool bCompareOwningActors = false);
+
 private:
 #if ENABLE_ANIM_DEBUG
 	friend struct UE::PoseSearch::FInteractionValidator;
@@ -114,7 +129,6 @@ private:
 	bool ValidateAllIslands() const;
 #endif // DO_CHECK
 
-	void PopulateContinuingProperties(float DeltaSeconds, TArrayView<UE::PoseSearch::FInteractionSearchContext> SearchContexts) const;
 	UE::PoseSearch::FInteractionIsland* FindIsland(const UObject* AnimContext, bool bCompareOwningActors = false);
 	
 #if ENABLE_ANIM_DEBUG
@@ -126,11 +140,6 @@ private:
 	void GenerateAnimContextInfosAndTagToDatabases(UE::PoseSearch::FAnimContextInfos& AnimContextInfos, UE::PoseSearch::FTagToDatabases& TagToDatabases) const;
 	void GenerateSearchContexts(float DeltaSeconds, UE::PoseSearch::FInteractionSearchContexts& SearchContexts) const;
 
-	void UpdateValidInteractionSearches();
-	static void OnInteractionStart(UE::PoseSearch::FValidInteractionSearch& ValidInteractionSearch);
-	static void OnInteractionContinuing(UE::PoseSearch::FValidInteractionSearch& ValidInteractionSearch);
-	static void OnInteractionEnd(UE::PoseSearch::FValidInteractionSearch& ValidInteractionSearch);
-
 	// it could contain duplicated availabilities
 	TArray<UE::PoseSearch::FInteractionAnimContextAvailabilities> AnimContextsAvailabilitiesBuffer;
 	FThreadSafeCounter AnimContextsAvailabilitiesIndex = 0;
@@ -139,7 +148,7 @@ private:
 	// array of groups of characters that needs to be anaylzed together for possible interactions
 	TArray<UE::PoseSearch::FInteractionIsland*> Islands;
 
-	TMap<uint32, UE::PoseSearch::FValidInteractionSearch> ValidInteractionSearches;
+	UE::PoseSearch::FValidInteractionSearches ValidInteractionSearches;
 };
 
 #undef UE_API
