PoseSearch - FInteractionIsland::AddReferencedObjects() fix, FRoledAnimContextInfos::Reset() fix, documentation, and LOG(s) improvements

#FailedValidations Preflight

diff --git a/Engine/Plugins/Animation/PoseSearch/Source/Runtime/Private/PoseSearchAssetSampler.cpp b/Engine/Plugins/Animation/PoseSearch/Source/Runtime/Private/PoseSearchAssetSampler.cpp
--- a/Engine/Plugins/Animation/PoseSearch/Source/Runtime/Private/PoseSearchAssetSampler.cpp
+++ b/Engine/Plugins/Animation/PoseSearch/Source/Runtime/Private/PoseSearchAssetSampler.cpp
@@ -672,7 +672,7 @@ void FAnimationAssetSampler::ExtractAnimNotifyContext(float Time, FAnimNotifyCon
 		}
 		else
 		{
-			UE_LOG(LogPoseSearch, Error, TEXT("FAnimationAssetSampler::ExtractAnimNotifyStates: Unsupported BlendSpace NotifyTriggerMode for '%s'"), *BlendSpace->GetName());
+			UE_LOG(LogPoseSearch, Error, TEXT("FAnimationAssetSampler::ExtractAnimNotifyContext: Unsupported BlendSpace NotifyTriggerMode for '%s'"), *BlendSpace->GetName());
 		}
 	}
 	else if (const UAnimSequenceBase* SequenceBase = Cast<UAnimSequenceBase>(AnimationAssetPtr.Get()))
@@ -742,7 +742,7 @@ TConstArrayView<FAnimNotifyEvent> FAnimationAssetSampler::GetAllAnimNotifyEvents
 		}
 		else
 		{
-			UE_LOG(LogPoseSearch, Error, TEXT("FAnimationAssetSampler::ExtractPoseSearchNotifyStates: Unsupported BlendSpace NotifyTriggerMode for '%s'"), *BlendSpace->GetName());
+			UE_LOG(LogPoseSearch, Error, TEXT("FAnimationAssetSampler::GetAllAnimNotifyEvents: Unsupported BlendSpace NotifyTriggerMode for '%s'"), *BlendSpace->GetName());
 		}
 	}
 	else if (const UAnimSequenceBase* SequenceBase = Cast<UAnimSequenceBase>(AnimationAssetPtr.Get()))
diff --git a/Engine/Plugins/Animation/PoseSearch/Source/Runtime/Private/PoseSearchInteractionIsland.cpp b/Engine/Plugins/Animation/PoseSearch/Source/Runtime/Private/PoseSearchInteractionIsland.cpp
--- a/Engine/Plugins/Animation/PoseSearch/Source/Runtime/Private/PoseSearchInteractionIsland.cpp
+++ b/Engine/Plugins/Animation/PoseSearch/Source/Runtime/Private/PoseSearchInteractionIsland.cpp
@@ -786,11 +786,11 @@ void FConstraints::ProcessAll(float DeltaTime, const FSearchResult& SearchResult
 						const FTransform WantedCurrentFromTransform = AnimFromTransform.GetRelativeTransform(AnimToTransform) * CurrentToTransform;
 
 						Constraint.FromBoneTransform.SetTranslation(FMath::Lerp(WantedCurrentFromTransform.GetTranslation(), CurrentFromTransform.GetTranslation(), Constraint.TranslationWeight));
-						Constraint.FromBoneTransform.SetRotation(FMath::Lerp(WantedCurrentFromTransform.GetRotation(), CurrentFromTransform.GetRotation(), Constraint.RotationWeight));
+						Constraint.FromBoneTransform.SetRotation(FQuat::Slerp(WantedCurrentFromTransform.GetRotation(), CurrentFromTransform.GetRotation(), Constraint.RotationWeight));
 						Constraint.FromBoneTransform.SetScale3D(CurrentFromTransform.GetScale3D());
 
 						Constraint.ToBoneTransform.SetTranslation(FMath::Lerp(WantedCurrentToTransform.GetTranslation(), CurrentToTransform.GetTranslation(), 1.f - Constraint.TranslationWeight));
-						Constraint.ToBoneTransform.SetRotation(FMath::Lerp(WantedCurrentToTransform.GetRotation(), CurrentToTransform.GetRotation(), 1.f - Constraint.RotationWeight));
+						Constraint.ToBoneTransform.SetRotation(FQuat::Slerp(WantedCurrentToTransform.GetRotation(), CurrentToTransform.GetRotation(), 1.f - Constraint.RotationWeight));
 						Constraint.ToBoneTransform.SetScale3D(CurrentToTransform.GetScale3D());
 
 #if WITH_EDITORONLY_DATA && ENABLE_ANIM_DEBUG && ENABLE_VISUAL_LOG
@@ -1617,15 +1617,33 @@ void FInteractionIsland::RemoveTickDependencies(bool bValidateTickDependencies)
 		}
 
 		// [undo] adding prerequisite to the ClosingTickFunction to components having a tick dependency on TickActorComponent, so they can run AFTER the ClosingTickFunction
-		for (TObjectPtr<UActorComponent>& PostClosingTickActorComponentPtr : PostClosingTickActorComponents)
+		for (TObjectPtr<UActorComponent>& PostClosingTickActorComponent : PostClosingTickActorComponents)
 		{
-			PostClosingTickActorComponentPtr->PrimaryComponentTick.RemovePrerequisite(InteractionSubsystem, ClosingTickFunction);
+			if (PostClosingTickActorComponent)
+			{
+				PostClosingTickActorComponent->PrimaryComponentTick.RemovePrerequisite(InteractionSubsystem, ClosingTickFunction);
+			}
+		}
+
+		if (!PreTickFunction.GetPrerequisites().IsEmpty())
+		{
+			UE_LOG(LogPoseSearch, Warning, TEXT("FInteractionIsland::RemoveTickDependencies, unable to cleanly remove all the PreTickFunction prerequisites! (This could have happened if an actor got garbage collected or a BP got recompiled while PIE was paused)"));
+			PreTickFunction.GetPrerequisites().Reset();
 		}
 
-		check(PreTickFunction.GetPrerequisites().IsEmpty());
-		check(PostTickFunction.GetPrerequisites().IsEmpty());
-		check(ClosingTickFunction.GetPrerequisites().IsEmpty());
+		if (!PostTickFunction.GetPrerequisites().IsEmpty())
+		{
+			UE_LOG(LogPoseSearch, Warning, TEXT("FInteractionIsland::RemoveTickDependencies, unable to cleanly remove all the PostTickFunction prerequisites! (This could have happened if an actor got garbage collected or a BP got recompiled while PIE was paused)"));
+			PostTickFunction.GetPrerequisites().Reset();
+		}
 
+		if (!ClosingTickFunction.GetPrerequisites().IsEmpty())
+		{
+			UE_LOG(LogPoseSearch, Warning, TEXT("FInteractionIsland::RemoveTickDependencies, unable to cleanly remove all the ClosingTickFunction prerequisites! (This could have happened if an actor got garbage collected or a BP got recompiled while PIE was paused)"));
+			ClosingTickFunction.GetPrerequisites().Reset();
+		}
+
+		PostClosingTickActorComponents.Reset();
 		bHasTickDependencies = false;
 	}
 
@@ -1726,8 +1744,9 @@ void FInteractionIsland::Uninitialize(bool bValidateTickDependencies)
 
 	if (IsInitialized())
 	{
+		check(PostClosingTickActorComponents.IsEmpty());
+
 		TickActorComponents.Reset();
-		PostClosingTickActorComponents.Reset();
 		IslandAnimContexts.Reset();
 
 		SearchContexts.Reset();
@@ -1823,19 +1842,48 @@ void FInteractionIsland::LogTickDependencies() const
 
 void FInteractionIsland::AddReferencedObjects(FReferenceCollector& Collector)
 {
+	TSet<TObjectPtr<AActor>, DefaultKeyFuncs<AActor*>, TInlineSetAllocator<32>> OwningTickActors;
+
 	for (TObjectPtr<UActorComponent>& TickActorComponent : TickActorComponents)
 	{
-		Collector.AddReferencedObject(TickActorComponent);
+		if (TickActorComponent)
+		{
+			Collector.AddReferencedObject(TickActorComponent);
+			// since AddReferencedObject may kill PostClosingTickActorComponent when recompiling blue prints
+			if (TickActorComponent)
+			{
+				OwningTickActors.Add(TickActorComponent->GetOwner());
+			}
+		}
 	}
 
 	for (TObjectPtr<UActorComponent>& PostClosingTickActorComponent : PostClosingTickActorComponents)
 	{
-		Collector.AddReferencedObject(PostClosingTickActorComponent);
+		if (PostClosingTickActorComponent)
+		{
+			Collector.AddReferencedObject(PostClosingTickActorComponent);
+			// since AddReferencedObject may kill PostClosingTickActorComponent when recompiling blue prints
+			if (PostClosingTickActorComponent)
+			{
+				OwningTickActors.Add(PostClosingTickActorComponent->GetOwner());
+			}
+		}
 	}
 	
+	for (TObjectPtr<AActor>& OwningTickActor : OwningTickActors)
+	{
+		if (OwningTickActor)
+		{
+			Collector.AddReferencedObject(OwningTickActor);
+		}
+	}
+
 	for (TObjectPtr<const UObject>& IslandAnimContext : IslandAnimContexts)
 	{
-		Collector.AddReferencedObject(IslandAnimContext);
+		if (IslandAnimContext)
+		{
+			Collector.AddReferencedObject(IslandAnimContext);
+		}
 	}
 
 	for (FInteractionSearchContext& SearchContext : SearchContexts)
diff --git a/Engine/Plugins/Animation/PoseSearch/Source/Runtime/Private/PoseSearchInteractionSubsystem.cpp b/Engine/Plugins/Animation/PoseSearch/Source/Runtime/Private/PoseSearchInteractionSubsystem.cpp
--- a/Engine/Plugins/Animation/PoseSearch/Source/Runtime/Private/PoseSearchInteractionSubsystem.cpp
+++ b/Engine/Plugins/Animation/PoseSearch/Source/Runtime/Private/PoseSearchInteractionSubsystem.cpp
@@ -886,7 +886,7 @@ void UPoseSearchInteractionSubsystem::GenerateSearchContexts(float DeltaSeconds,
 					++PrevIndex;
 
 					if (PrevIndex != Index)
-	{
+					{
 						InternalData[PrevIndex] = MoveTemp(InternalData[Index]);
 					}
 				}
@@ -950,6 +950,7 @@ void UPoseSearchInteractionSubsystem::GenerateSearchContexts(float DeltaSeconds,
 		void Reset()
 		{
 			InternalData.Reset();
+			PresentRoleIndexes.Reset();
 		}
 	};
 
diff --git a/Engine/Plugins/Animation/PoseSearch/Source/Runtime/Public/PoseSearch/PoseSearchAssetSampler.h b/Engine/Plugins/Animation/PoseSearch/Source/Runtime/Public/PoseSearch/PoseSearchAssetSampler.h
--- a/Engine/Plugins/Animation/PoseSearch/Source/Runtime/Public/PoseSearch/PoseSearchAssetSampler.h
+++ b/Engine/Plugins/Animation/PoseSearch/Source/Runtime/Public/PoseSearch/PoseSearchAssetSampler.h
@@ -46,9 +46,11 @@ struct FAnimationAssetSampler
 	// sequence limits when Time is less than zero or greater than the sequence length.
 	POSESEARCH_API FTransform ExtractRootTransform(float Time) const;
 	
-	// Extracts notify states present in the sequence at Time
+	// Populates FAnimNotifyContext from the animation asset at Time
 	POSESEARCH_API void ExtractAnimNotifyContext(float Time, FAnimNotifyContext& PreAllocatedNotifyContext) const;
+	// Extracts notify states present in the animation asset at Time
 	POSESEARCH_API void ExtractAnimNotifyStates(float Time, FAnimNotifyContext& PreAllocatedNotifyContext, const TFunction<bool(UAnimNotifyState*)>& ProcessAnimNotifyState) const;
+	// Extracts notify events present in the animation asset at Time
 	POSESEARCH_API void ExtractAnimNotifyEvents(float Time, FAnimNotifyContext& PreAllocatedNotifyContext, const TFunction<bool(const FAnimNotifyEvent*)>& ProcessAnimNotifyEvent) const;
 	
 	POSESEARCH_API TConstArrayView<FAnimNotifyEvent> GetAllAnimNotifyEvents() const;
