PoseSearch - Socket support for Constraints

diff --git a/Engine/Plugins/Animation/PoseSearch/Source/Editor/Private/PoseSearchInteractionAssetEditor.cpp b/Engine/Plugins/Animation/PoseSearch/Source/Editor/Private/PoseSearchInteractionAssetEditor.cpp
--- a/Engine/Plugins/Animation/PoseSearch/Source/Editor/Private/PoseSearchInteractionAssetEditor.cpp
+++ b/Engine/Plugins/Animation/PoseSearch/Source/Editor/Private/PoseSearchInteractionAssetEditor.cpp
@@ -4,17 +4,18 @@
 #include "Animation/DebugSkelMeshComponent.h"
 #include "AnimPreviewInstance.h"
 #include "AssetEditorModeManager.h"
-#include "PoseSearch/PoseSearchInteractionAsset.h"
-#include "PoseSearchEditor.h"
 #include "Components/StaticMeshComponent.h"
-#include "Framework/Application/SlateApplication.h"
 #include "EngineUtils.h"
+#include "Framework/Application/SlateApplication.h"
 #include "GameFramework/WorldSettings.h"
 #include "Modules/ModuleManager.h"
+#include "PoseSearch/PoseSearchConstraints.h"
+#include "PoseSearch/PoseSearchInteractionAsset.h"
+#include "PoseSearchEditor.h"
 #include "PreviewProfileController.h"
 #include "PropertyEditorModule.h"
-#include "Styling/AppStyle.h"
 #include "SSimpleTimeSlider.h"
+#include "Styling/AppStyle.h"
 #include "UnrealWidget.h"
 #include "Viewports.h"
 #include "Widgets/Docking/SDockTab.h"
@@ -324,9 +325,20 @@ void FInteractionAssetViewModel::Tick(float DeltaSeconds)
 	
 	FRoleToIndex PreviewActorRoleToIndex;
 	PreviewActorRoleToIndex.Reserve(PreviewActors.Num());
+
+	TArray<TObjectPtr<const UObject>, TInlineAllocator<PreallocatedRolesNum>> AnimContexts;
+	AnimContexts.SetNum(NumRoles);
+
+	TArray<FRole, TInlineAllocator<PreallocatedRolesNum>> Roles;
+	Roles.SetNum(NumRoles);
+
 	for (int32 PreviewActorIndex = 0; PreviewActorIndex < PreviewActors.Num(); ++PreviewActorIndex)
 	{
 		PreviewActorRoleToIndex.Add(PreviewActors[PreviewActorIndex].GetRole()) = PreviewActorIndex;
+
+		FInteractionAssetPreviewActor& PreviewActor = PreviewActors[PreviewActorIndex];
+		AnimContexts[PreviewActorIndex] = PreviewActor.GetAnimPreviewInstance();
+		Roles[PreviewActorIndex] = PreviewActor.GetRole();
 	}
 
 	// testing CalculateWarpTransforms
@@ -334,9 +346,7 @@ void FInteractionAssetViewModel::Tick(float DeltaSeconds)
 	if (InteractionAsset->bEnableDebugWarp && PreviewActors.Num() == NumRoles)
 	{
 		TArray<FTransform, TInlineAllocator<PreallocatedRolesNum>> ActorTransforms;
-		TArray<TObjectPtr<const UObject>, TInlineAllocator<PreallocatedRolesNum>> DebugAnimContexts;
 		ActorTransforms.SetNum(NumRoles);
-		DebugAnimContexts.SetNum(NumRoles);
 		for (int32 InteractionAssetRoleIndex = 0; InteractionAssetRoleIndex < NumRoles; ++InteractionAssetRoleIndex)
 		{
 			const FRole& InteractionAssetRole = InteractionAsset->GetRole(InteractionAssetRoleIndex);
@@ -346,14 +356,13 @@ void FInteractionAssetViewModel::Tick(float DeltaSeconds)
 			check(PreviewActor.GetRole() == InteractionAssetRole);
 
 			ActorTransforms[InteractionAssetRoleIndex] = PreviewActor.GetDebugActorTransformFromSampler();
-			DebugAnimContexts[InteractionAssetRoleIndex] = PreviewActor.GetAnimPreviewInstance();
 		}
 
 		TArray<FTransform, TInlineAllocator<PreallocatedRolesNum>> FullAlignedActorTransforms;
 		FullAlignedActorTransforms.SetNum(NumRoles);
 		const TConstArrayView<const UMirrorDataTable*> MirrorDataTables;
 		const TConstArrayView<bool> RelevantRoleIndexes;
-		InteractionAsset->CalculateWarpTransforms(PlayTime, ActorTransforms, FullAlignedActorTransforms, MirrorDataTables, RelevantRoleIndexes, DebugAnimContexts);
+		InteractionAsset->CalculateWarpTransforms(PlayTime, ActorTransforms, FullAlignedActorTransforms, MirrorDataTables, RelevantRoleIndexes, AnimContexts);
 
 		for (int32 InteractionAssetRoleIndex = 0; InteractionAssetRoleIndex < NumRoles; ++InteractionAssetRoleIndex)
 		{
@@ -368,6 +377,23 @@ void FInteractionAssetViewModel::Tick(float DeltaSeconds)
 		}
 	}
 #endif // WITH_EDITORONLY_DATA
+
+	// extracting and drawing Constraints
+	
+	// @todo: fill up those parameters
+	const FVector BlendParameters = FVector::ZeroVector;
+	const bool bIsMirrored = false;
+
+	FConstraints Constraints;
+	Constraints.ProcessAll(InteractionAsset, StepDeltaTime, PlayTime, bIsMirrored, BlendParameters, AnimContexts, Roles);
+	
+	if (const UWorld* World = GetWorld())
+	{
+		for (const FConstraint& Constraint : Constraints.GetAll())
+		{
+			DrawDebugLine(World, Constraint.FromSocketTransform.GetTranslation(), Constraint.ToSocketTransform.GetTranslation(), FColor::Red, false, -1.f, SDPG_Foreground);
+		}
+	}
 }
 
 void FInteractionAssetViewModel::RemovePreviewActors()
diff --git a/Engine/Plugins/Animation/PoseSearch/Source/Runtime/Private/PoseSearchConstraints.cpp b/Engine/Plugins/Animation/PoseSearch/Source/Runtime/Private/PoseSearchConstraints.cpp
new file mode 100644
--- /dev/null
+++ b/Engine/Plugins/Animation/PoseSearch/Source/Runtime/Private/PoseSearchConstraints.cpp
@@ -0,0 +1,385 @@
+// Copyright Epic Games, Inc. All Rights Reserved.
+
+#include "PoseSearch/PoseSearchConstraints.h"
+#include "PoseSearch/PoseSearchDatabase.h"
+#include "PoseSearch/MultiAnimAsset.h"
+#include "PoseSearch/PoseSearchInteractionIsland.h"
+#include "Components/SkeletalMeshComponent.h"
+#include "Engine/SkeletalMesh.h"
+#include "VisualLogger/VisualLogger.h"
+
+namespace UE::PoseSearch
+{
+
+bool FConstraints::GetConstraint(FRole Role, FName SocketName, float& OutDesiredReach, FTransform& OutTransform) const
+{
+	for (const FConstraint& Constraint : Constraints)
+	{
+		if (const FTransform* SocketTransform = Constraint.GetSocketTransform(Role, SocketName))
+		{
+			OutDesiredReach = Constraint.DesiredReach;
+			OutTransform = *SocketTransform;
+			return true;
+		}
+	}
+
+	OutDesiredReach = 0.f;
+	OutTransform = FTransform::Identity;
+	return false;
+}
+
+int32 FConstraints::Num() const
+{
+	return Constraints.Num();
+}
+
+const FConstraint* FConstraints::Find(const FConstraint& Constraint) const
+{
+	const int32 ConstraintIndex = Constraints.Find(Constraint);
+	if (ConstraintIndex != INDEX_NONE)
+	{
+		return &Constraints[ConstraintIndex];
+	}
+	return nullptr;
+}
+
+FConstraint* FConstraints::Find(const FConstraint& Constraint)
+{
+	const int32 ConstraintIndex = Constraints.Find(Constraint);
+	if (ConstraintIndex != INDEX_NONE)
+	{
+		return &Constraints[ConstraintIndex];
+	}
+	return nullptr;
+}
+
+void FConstraints::ProcessAll(float DeltaTime, const FSearchResult& SearchResult, const FInteractionSearchContext& SearchContext, const FConstraints* SearchContextPreviousConstraints)
+{
+	// making sure nobody already called UpdateConstraints or we didn't reset the constraints...
+	check(Constraints.IsEmpty());
+
+	const UPoseSearchDatabase* Database = SearchResult.Database.Get();
+	check(Database);
+
+	const FSearchIndexAsset* SearchIndexAsset = SearchResult.GetSearchIndexAsset();
+	check(SearchIndexAsset);
+
+	const FPoseSearchDatabaseAnimationAssetBase* DatabaseAnimationAssetBase = Database->GetDatabaseAnimationAsset(*SearchIndexAsset);
+	check(DatabaseAnimationAssetBase);
+
+	if (const UMultiAnimAsset* MultiAnimAsset = Cast<UMultiAnimAsset>(DatabaseAnimationAssetBase->GetAnimationAsset()))
+	{
+		TConstArrayView<FConstraint> PreviousConstraints;
+		if (SearchContextPreviousConstraints)
+		{
+			PreviousConstraints = SearchContextPreviousConstraints->Constraints;
+		}
+
+		ProcessAll(MultiAnimAsset, DeltaTime, SearchResult.GetAssetTime(), SearchIndexAsset->IsMirrored(), SearchIndexAsset->GetBlendParameters(), SearchContext.GetAnimContexts(), SearchContext.GetRoles(), PreviousConstraints);
+	}
+}
+
+void FConstraints::ProcessAll(const UMultiAnimAsset* MultiAnimAsset, float DeltaTime, float AssetTime, bool bIsMirrored, const FVector& BlendParameters,
+	const TConstArrayView<TObjectPtr<const UObject>> AnimContexts, const TConstArrayView<FRole> Roles, TConstArrayView<FConstraint> PreviousConstraints)
+{
+	check(MultiAnimAsset);
+	const int32 NumAnimContexts = AnimContexts.Num();
+	check(NumAnimContexts == Roles.Num());
+
+	Constraints.Reset();
+
+	FMemMark Mark(FMemStack::Get());
+
+	TArray<bool, TInlineAllocator<NumPreallocatedContraints, TMemStackAllocator<>>> VisitedPreviousConstraints;
+	VisitedPreviousConstraints.SetNumZeroed(PreviousConstraints.Num());
+		
+	// extracting the contraints for this SearchContext from MultiAnimAsset
+	FAnimNotifyContext PreAllocatedNotifyContext;
+	for (int32 AnimContextIndex = 0; AnimContextIndex < NumAnimContexts; ++AnimContextIndex)
+	{
+		const FRole& Role = Roles[AnimContextIndex];
+		const UAnimationAsset* Animation = MultiAnimAsset->GetAnimationAsset(Role);
+		check(Animation);
+
+		// initializing FAnimationAssetSampler only for extracting AnimNotifyStates
+		FAnimationAssetSampler Sampler(Animation, FTransform::Identity, FVector::ZeroVector, FAnimationAssetSampler::DefaultRootTransformSamplingRate, false, false);
+		Sampler.ExtractAnimNotifyEvents(AssetTime, PreAllocatedNotifyContext, [AssetTime, DeltaTime, this, &PreviousConstraints, &VisitedPreviousConstraints](const FAnimNotifyEvent* AnimNotifyEvent)
+			{
+				if (const UAnimNotifyState_PoseSearchConstraint* ConstraintNotifyState = Cast<const UAnimNotifyState_PoseSearchConstraint>(AnimNotifyEvent->NotifyStateClass))
+				{
+					// subtracting ConstraintNotifyState->CoolDownTime from the validity of this UAnimNotifyState_PoseSearchConstraint
+					const float NotifyEndTime = AnimNotifyEvent->GetEndTriggerTime();
+					if (AssetTime <= (NotifyEndTime - ConstraintNotifyState->CoolDownTime))
+					{
+						FConstraint NewConstraint(ConstraintNotifyState);
+						const int32 PreviousConstraintIndex = PreviousConstraints.Find(NewConstraint);
+						if (PreviousConstraintIndex != INDEX_NONE)
+						{
+							VisitedPreviousConstraints[PreviousConstraintIndex] = true;
+							const FConstraint& PreviousConstraint = PreviousConstraints[PreviousConstraintIndex];
+
+							// @todo: for now we just ignore conflicting constraints, but we should probably blend PreviousConstraint into NewConstraint :/
+							//        perhaps keeping the min or RampUpTime, to make the constraint as following the fastest one to get into full DesiredReach, and the max CoolDownTime to fade off with the slowest rate
+							//        not really sure what to do about TranslationWeight and RotationWeight.. maybe slowly blending them to the newer one?
+							if (PreviousConstraint.RampUpTime > UE_KINDA_SMALL_NUMBER)
+							{
+								NewConstraint.DesiredReach = FMath::Min(1.f, PreviousConstraint.DesiredReach + (DeltaTime / PreviousConstraint.RampUpTime));
+							}
+							else
+							{
+								NewConstraint.DesiredReach = 1.f;
+							}
+						}
+						else if (NewConstraint.RampUpTime <= UE_KINDA_SMALL_NUMBER)
+						{
+							NewConstraint.DesiredReach = 1.f;
+						}
+
+						// @todo: for now we also ignore duplicate constraints, and perhaps we should use a similar strategy of blending the constraints
+						Constraints.AddUnique(NewConstraint);
+					}
+				}
+				return true;
+			});
+	}
+
+	// adding remaining unprocessed constraints from PreviousConstraints
+	for (int32 PreviousConstraintIndex = 0; PreviousConstraintIndex < PreviousConstraints.Num(); ++PreviousConstraintIndex)
+	{
+		if (!VisitedPreviousConstraints[PreviousConstraintIndex])
+		{
+			const FConstraint& PreviousConstraint = PreviousConstraints[PreviousConstraintIndex];
+			if (PreviousConstraint.CoolDownTime > UE_KINDA_SMALL_NUMBER)
+			{
+				const float DesiredReach = PreviousConstraint.DesiredReach - (DeltaTime / PreviousConstraint.CoolDownTime);
+				if (DesiredReach > UE_KINDA_SMALL_NUMBER)
+				{
+					FConstraint& NewConstraint = Constraints.Add_GetRef(PreviousConstraint);
+					NewConstraint.DesiredReach = DesiredReach;
+				}
+			}
+		}
+	}
+
+	// process all the collected constraints
+	if (!Constraints.IsEmpty())
+	{
+		TArray<FTransform, TInlineAllocator<PreallocatedRolesNum, TMemStackAllocator<>>> ComponentToWorldTransforms;
+		TArray<TConstArrayView<FTransform>, TInlineAllocator<PreallocatedRolesNum, TMemStackAllocator<>>> AnimContextsComponentSpaceTransforms;
+		TArray<FBoneContainer, TInlineAllocator<PreallocatedRolesNum, TMemStackAllocator<>>> BoneContainers;
+		TArray<FCompactPose, TInlineAllocator<PreallocatedRolesNum, TMemStackAllocator<>>> Poses;
+		TArray<FCSPose<FCompactPose>, TInlineAllocator<PreallocatedRolesNum, TMemStackAllocator<>>> ComponentSpacePoses;
+		TArray<const USkeletalMeshComponent*, TInlineAllocator<PreallocatedRolesNum, TMemStackAllocator<>>> SkeletalMeshComponents;
+
+		ComponentToWorldTransforms.SetNum(NumAnimContexts);
+		AnimContextsComponentSpaceTransforms.SetNum(NumAnimContexts);
+		BoneContainers.SetNum(NumAnimContexts);
+		Poses.SetNum(NumAnimContexts);
+		ComponentSpacePoses.SetNum(NumAnimContexts);
+		SkeletalMeshComponents.SetNum(NumAnimContexts);
+
+		FRoleToIndex RoleToIndex;
+		for (int32 AnimContextIndex = 0; AnimContextIndex < NumAnimContexts; ++AnimContextIndex)
+		{
+			const UObject* AnimContext = AnimContexts[AnimContextIndex];
+			check(AnimContext);
+			const USkeletalMeshComponent* SkeletalMeshComponent = GetContextSkeletalMeshComponent(AnimContext, false);
+			check(SkeletalMeshComponent);
+			SkeletalMeshComponents[AnimContextIndex] = SkeletalMeshComponent;
+
+			// @todo: do we need this logic about URO etc? if we're executing this is because we update the mesh, so we probably don't need any interpolation...
+			const bool bUROInSync = SkeletalMeshComponent->ShouldUseUpdateRateOptimizations() && SkeletalMeshComponent->AnimUpdateRateParams != nullptr;
+			const bool bUsingExternalInterpolation = SkeletalMeshComponent->IsUsingExternalInterpolation();
+
+			const TArray<FTransform>& CachedComponentSpaceTransforms = SkeletalMeshComponent->GetCachedComponentSpaceTransforms();
+			const bool bArraySizesMatch = CachedComponentSpaceTransforms.Num() == SkeletalMeshComponent->GetComponentSpaceTransforms().Num();
+
+			ComponentToWorldTransforms[AnimContextIndex] = SkeletalMeshComponent->GetComponentTransform();
+			AnimContextsComponentSpaceTransforms[AnimContextIndex] = (bUROInSync || bUsingExternalInterpolation) && bArraySizesMatch ? CachedComponentSpaceTransforms : SkeletalMeshComponent->GetComponentSpaceTransforms();
+
+			// @todo: check FAnimNode_CopyPoseFromMesh::PreUpdate to see how to GetAnimCurves and GetCustomAttributes
+			const FRole& Role = Roles[AnimContextIndex];
+			RoleToIndex.Add(Role) = AnimContextIndex;
+
+			const UAnimationAsset* Animation = MultiAnimAsset->GetAnimationAsset(Role);
+			check(Animation);
+			const FTransform Origin = MultiAnimAsset->GetOrigin(Role);
+
+			FAnimationAssetSampler Sampler(Animation, Origin, BlendParameters);
+			const FTransform RootTransform = Sampler.ExtractRootTransform(AssetTime);
+
+			// @todo add support for mirroring if needed
+			//if (MirrorDataTables.IsValidIndex(ItemIndex) && MirrorDataTables[ItemIndex])
+			//{
+			//	const FMirrorDataCache MirrorDataCache(MirrorDataTables[ItemIndex]);
+			//	AssetRootBoneTransforms[ItemIndex] = MirrorDataCache.MirrorTransform(AssetRootBoneTransforms[ItemIndex]);
+			//}
+
+			// array containing the bone index of the root bone (0)
+			const USkeleton* Skeleton = Animation->GetSkeleton();
+			check(Skeleton && SkeletalMeshComponent->GetSkeletalMeshAsset()->GetSkeleton() == Skeleton);
+
+			const FReferenceSkeleton& ReferenceSkeleton = Skeleton->GetReferenceSkeleton();
+
+			// @todo: refine this BoneIndices / cache bone containers
+			const int32 NumBones = ReferenceSkeleton.GetNum();
+			TArray<uint16, TMemStackAllocator<>> BoneIndices;
+			BoneIndices.SetNumUninitialized(NumBones);
+			for (int32 BoneIndex = 0; BoneIndex < NumBones; ++BoneIndex)
+			{
+				BoneIndices[BoneIndex] = BoneIndex;
+			}
+
+			// extracting the pose, containing only the root bone from the Sampler 
+			BoneContainers[AnimContextIndex].InitializeTo(BoneIndices, UE::Anim::FCurveFilterSettings(UE::Anim::ECurveFilterMode::DisallowAll), *Skeleton);
+			Poses[AnimContextIndex].SetBoneContainer(&BoneContainers[AnimContextIndex]);
+			Sampler.ExtractPose(AssetTime, Poses[AnimContextIndex]);
+
+			// making sure the animation root bone transform is Identity, so we can confuse the root with the root BONE transform and preserve performances!
+			FTransform& RootBoneTransform = Poses[AnimContextIndex][FCompactPoseBoneIndex(0)];
+			if (!RootBoneTransform.Equals(FTransform::Identity))
+			{
+				const FVector Pos = RootBoneTransform.GetTranslation();
+				const FRotator Rot(RootBoneTransform.GetRotation());
+				UE_LOG(LogPoseSearch, Error, TEXT("FInteractionIsland::UpdateConstraints unsupported non identity root bone in %s at time %f Pos(%f, %f, %f), Rot(%f, %f, %f)"), *Animation->GetName(), AssetTime, Pos.X, Pos.Y, Pos.Z, Rot.Pitch, Rot.Yaw, Rot.Roll);
+			}
+			RootBoneTransform = Sampler.ExtractRootTransform(AssetTime);
+			ComponentSpacePoses[AnimContextIndex].InitPose(Poses[AnimContextIndex]);
+
+#if ENABLE_VISUAL_LOG
+			if (FVisualLogger::IsRecording())
+			{
+				static const TCHAR* LogName = TEXT("PoseSearchInteraction_Constraints");
+
+				for (int32 BoneIndex = 0; BoneIndex < ReferenceSkeleton.GetNum(); ++BoneIndex)
+				{
+					const int32 ParentBoneIndex = ReferenceSkeleton.GetParentIndex(BoneIndex);
+					if (ParentBoneIndex != INDEX_NONE)
+					{
+						const FVector CurrentBoneWorldLocation = ComponentToWorldTransforms[AnimContextIndex].TransformPosition(AnimContextsComponentSpaceTransforms[AnimContextIndex][BoneIndex].GetTranslation());
+						const FVector CurrentParentBoneWorldLocation = ComponentToWorldTransforms[AnimContextIndex].TransformPosition(AnimContextsComponentSpaceTransforms[AnimContextIndex][ParentBoneIndex].GetTranslation());
+
+						const FVector AnimBoneWorldLocation = ComponentSpacePoses[AnimContextIndex].GetComponentSpaceTransform(FCompactPoseBoneIndex(BoneIndex)).GetTranslation();
+						const FVector AnimParentBoneWorldLocation = ComponentSpacePoses[AnimContextIndex].GetComponentSpaceTransform(FCompactPoseBoneIndex(ParentBoneIndex)).GetTranslation();
+
+						for (int32 DrawAnimContextIndex = 0; DrawAnimContextIndex < NumAnimContexts; ++DrawAnimContextIndex)
+						{
+							const UObject* DrawAnimContext = AnimContexts[DrawAnimContextIndex];
+							const FColor Color = DrawAnimContext == AnimContext ? FColor::Red : FColor::Blue;
+
+							// current characters poses
+							UE_VLOG_SEGMENT(DrawAnimContext, LogName, Display, CurrentBoneWorldLocation, CurrentParentBoneWorldLocation, Color, TEXT(""));
+
+							// sampled poses from UAnimationAsset
+							UE_VLOG_SEGMENT(DrawAnimContext, LogName, Display, AnimBoneWorldLocation, AnimParentBoneWorldLocation, Color, TEXT(""));
+						}
+					}
+				}
+			}
+#endif //ENABLE_VISUAL_LOG
+		}
+
+		for (TConstraints::TIterator ConstraintIt = Constraints.CreateIterator(); ConstraintIt; ++ConstraintIt)
+		{
+			FConstraint& Constraint = *ConstraintIt;
+			const int32* FromSocketRoleIndex = RoleToIndex.Find(Constraint.FromSocketRole);
+			const int32* ToSocketRoleIndex = RoleToIndex.Find(Constraint.ToSocketRole);
+
+			if (!FromSocketRoleIndex || !ToSocketRoleIndex)
+			{
+#if !NO_LOGGING
+				if (!FromSocketRoleIndex)
+				{
+					UE_LOG(LogPoseSearch, Error, TEXT("FConstraints::ProcessAll, Failed to find Role '%s' used in a PoseSearchConstraint in asset '%s'"), *Constraint.FromSocketRole.ToString(), *MultiAnimAsset->GetName());
+				}
+
+				if (!ToSocketRoleIndex)
+				{
+					UE_LOG(LogPoseSearch, Error, TEXT("FConstraints::ProcessAll, Failed to find Role '%s' used in a PoseSearchConstraint in asset '%s'"), *Constraint.ToSocketRole.ToString(), *MultiAnimAsset->GetName());
+				}
+#endif // !NO_LOGGING
+				ConstraintIt.RemoveCurrent();
+			}
+			else
+			{
+				int32 FromBoneIndex = INDEX_NONE;
+				int32 ToBoneIndex = INDEX_NONE;
+
+				FTransform FromSocketLocalTransform(FTransform::Identity);
+				FTransform ToSocketLocalTransform(FTransform::Identity);
+					
+				if (!SkeletalMeshComponents[*FromSocketRoleIndex]->GetSocketInfoByName(Constraint.FromSocketName, FromSocketLocalTransform, FromBoneIndex))
+				{
+					FromBoneIndex = SkeletalMeshComponents[*FromSocketRoleIndex]->GetBoneIndex(Constraint.FromSocketName);
+				}
+
+				if (!SkeletalMeshComponents[*ToSocketRoleIndex]->GetSocketInfoByName(Constraint.ToSocketName, ToSocketLocalTransform, ToBoneIndex))
+				{
+					ToBoneIndex = SkeletalMeshComponents[*ToSocketRoleIndex]->GetBoneIndex(Constraint.ToSocketName);
+				}
+
+				if (FromBoneIndex == INDEX_NONE || ToBoneIndex == INDEX_NONE)
+				{
+#if !NO_LOGGING
+					if (FromBoneIndex == INDEX_NONE)
+					{
+						UE_LOG(LogPoseSearch, Error, TEXT("FConstraints::ProcessAll, Failed to find Socket '%s' from SkeletalMeshComponent '%s' for Role '%s'"), *Constraint.FromSocketName.ToString(), *GetNameSafe(SkeletalMeshComponents[*FromSocketRoleIndex]), *Constraint.FromSocketRole.ToString());
+					}
+
+					if (ToBoneIndex == INDEX_NONE)
+					{
+						UE_LOG(LogPoseSearch, Error, TEXT("FConstraints::ProcessAll, Failed to find Socket '%s' from SkeletalMeshComponent '%s' for Role '%s'"), *Constraint.ToSocketName.ToString(), *GetNameSafe(SkeletalMeshComponents[*ToSocketRoleIndex]), *Constraint.ToSocketRole.ToString());
+					}
+#endif // !NO_LOGGING
+					ConstraintIt.RemoveCurrent();
+				}
+				else
+				{
+					const FTransform CurrentFromTransform = FromSocketLocalTransform * (AnimContextsComponentSpaceTransforms[*FromSocketRoleIndex][FromBoneIndex] * ComponentToWorldTransforms[*FromSocketRoleIndex]);
+					const FTransform CurrentToTransform = ToSocketLocalTransform * (AnimContextsComponentSpaceTransforms[*ToSocketRoleIndex][ToBoneIndex] * ComponentToWorldTransforms[*ToSocketRoleIndex]);
+
+					const FTransform AnimFromTransform = FromSocketLocalTransform * (ComponentSpacePoses[*FromSocketRoleIndex].GetComponentSpaceTransform(FCompactPoseBoneIndex(FromBoneIndex)));
+					const FTransform AnimToTransform = ToSocketLocalTransform * (ComponentSpacePoses[*ToSocketRoleIndex].GetComponentSpaceTransform(FCompactPoseBoneIndex(ToBoneIndex)));
+
+					const FTransform WantedCurrentToTransform = AnimToTransform.GetRelativeTransform(AnimFromTransform) * CurrentFromTransform;
+					const FTransform WantedCurrentFromTransform = AnimFromTransform.GetRelativeTransform(AnimToTransform) * CurrentToTransform;
+
+					Constraint.FromSocketTransform.SetTranslation(FMath::Lerp(WantedCurrentFromTransform.GetTranslation(), CurrentFromTransform.GetTranslation(), Constraint.TranslationWeight));
+					Constraint.FromSocketTransform.SetRotation(FQuat::Slerp(WantedCurrentFromTransform.GetRotation(), CurrentFromTransform.GetRotation(), Constraint.RotationWeight));
+					Constraint.FromSocketTransform.SetScale3D(CurrentFromTransform.GetScale3D());
+
+					Constraint.ToSocketTransform.SetTranslation(FMath::Lerp(WantedCurrentToTransform.GetTranslation(), CurrentToTransform.GetTranslation(), 1.f - Constraint.TranslationWeight));
+					Constraint.ToSocketTransform.SetRotation(FQuat::Slerp(WantedCurrentToTransform.GetRotation(), CurrentToTransform.GetRotation(), 1.f - Constraint.RotationWeight));
+					Constraint.ToSocketTransform.SetScale3D(CurrentToTransform.GetScale3D());
+
+#if WITH_EDITORONLY_DATA && ENABLE_ANIM_DEBUG && ENABLE_VISUAL_LOG
+					if (FVisualLogger::IsRecording())
+					{
+						static const TCHAR* LogName = TEXT("PoseSearchInteraction_Constraints");
+
+						for (int32 DrawAnimContextIndex = 0; DrawAnimContextIndex < NumAnimContexts; ++DrawAnimContextIndex)
+						{
+							const UObject* DrawAnimContext = AnimContexts[DrawAnimContextIndex];
+
+							// current characters poses
+							UE_VLOG_SEGMENT(DrawAnimContext, LogName, Display, CurrentFromTransform.GetTranslation(), CurrentToTransform.GetTranslation(), FColor::Yellow, TEXT(""));
+
+							// sampled poses from UAnimationAsset
+							UE_VLOG_SEGMENT(DrawAnimContext, LogName, Display, AnimFromTransform.GetTranslation(), AnimToTransform.GetTranslation(), FColor::Yellow, TEXT(""));
+
+							// wanted transforms
+							UE_VLOG_SEGMENT(DrawAnimContext, LogName, Display, CurrentFromTransform.GetTranslation(), WantedCurrentToTransform.GetTranslation(), FColor::Green, TEXT(""));
+							UE_VLOG_SEGMENT(DrawAnimContext, LogName, Display, CurrentToTransform.GetTranslation(), WantedCurrentFromTransform.GetTranslation(), FColor::Green, TEXT(""));
+
+							UE_VLOG_SEGMENT(DrawAnimContext, LogName, Display, CurrentFromTransform.GetTranslation(), Constraint.FromSocketTransform.GetTranslation(), FColor::Orange, TEXT(""));
+							UE_VLOG_SEGMENT(DrawAnimContext, LogName, Display, CurrentToTransform.GetTranslation(), Constraint.ToSocketTransform.GetTranslation(), FColor::Orange, TEXT(""));
+						}
+					}
+#endif // WITH_EDITORONLY_DATA && ENABLE_ANIM_DEBUG && ENABLE_VISUAL_LOG
+				}
+			}
+		}
+	}
+}
+
+} // namespace UE::PoseSearch
diff --git a/Engine/Plugins/Animation/PoseSearch/Source/Runtime/Private/PoseSearchInteractionIsland.cpp b/Engine/Plugins/Animation/PoseSearch/Source/Runtime/Private/PoseSearchInteractionIsland.cpp
--- a/Engine/Plugins/Animation/PoseSearch/Source/Runtime/Private/PoseSearchInteractionIsland.cpp
+++ b/Engine/Plugins/Animation/PoseSearch/Source/Runtime/Private/PoseSearchInteractionIsland.cpp
@@ -459,26 +459,26 @@ bool FInteractionSearchContextBase::CheckForConsistency() const
 ///////////////////////////////////////////////////////////
 FConstraint::FConstraint(const UAnimNotifyState_PoseSearchConstraint* ConstraintNotifyState)
 {
-	FromBoneName = ConstraintNotifyState->FromBoneName;
-	FromBoneRole = ConstraintNotifyState->FromBoneRole;
-	ToBoneName = ConstraintNotifyState->ToBoneName;
-	ToBoneRole = ConstraintNotifyState->ToBoneRole;
+	FromSocketName = ConstraintNotifyState->FromSocketName;
+	FromSocketRole = ConstraintNotifyState->FromSocketRole;
+	ToSocketName = ConstraintNotifyState->ToSocketName;
+	ToSocketRole = ConstraintNotifyState->ToSocketRole;
 	TranslationWeight = ConstraintNotifyState->TranslationWeight;
 	RotationWeight = ConstraintNotifyState->RotationWeight;
 	RampUpTime = ConstraintNotifyState->RampUpTime;
 	CoolDownTime = ConstraintNotifyState->CoolDownTime;
 }
 
-const FTransform* FConstraint::GetBoneTransform(FRole Role, FName BoneName) const
+const FTransform* FConstraint::GetSocketTransform(FRole Role, FName SocketName) const
 {
-	if (FromBoneRole == Role && FromBoneName == BoneName)
+	if (FromSocketRole == Role && FromSocketName == SocketName)
 	{
-		return &FromBoneTransform;
+		return &FromSocketTransform;
 	}
 
-	if (ToBoneRole == Role && ToBoneName == BoneName)
+	if (ToSocketRole == Role && ToSocketName == SocketName)
 	{
-		return &ToBoneTransform;
+		return &ToSocketTransform;
 	}
 
 	return nullptr;
@@ -487,13 +487,13 @@ const FTransform* FConstraint::GetBoneTransform(FRole Role, FName BoneName) cons
 bool FConstraint::operator==(const FConstraint& Other) const
 {
 	// same order of From and To properties
-	if (FromBoneRole == Other.FromBoneRole && FromBoneName == Other.FromBoneName && ToBoneRole == Other.ToBoneRole && ToBoneName == Other.ToBoneName)
+	if (FromSocketRole == Other.FromSocketRole && FromSocketName == Other.FromSocketName && ToSocketRole == Other.ToSocketRole && ToSocketName == Other.ToSocketName)
 	{
 		return true;
 	}
 
 	// opposite order of From and To properties
-	if (FromBoneRole == Other.ToBoneRole && FromBoneName == Other.ToBoneName && ToBoneRole == Other.FromBoneRole && ToBoneName == Other.FromBoneName)
+	if (FromSocketRole == Other.ToSocketRole && FromSocketName == Other.ToSocketName && ToSocketRole == Other.FromSocketRole && ToSocketName == Other.FromSocketName)
 	{
 		return true;
 	}
@@ -501,329 +501,6 @@ bool FConstraint::operator==(const FConstraint& Other) const
 	return false;
 }
 
-// FConstraints
-///////////////////////////////////////////////////////////
-bool FConstraints::GetConstraint(FRole Role, FName BoneName, float& OutDesiredReach, FTransform& OutTransform) const
-{
-	for (const FConstraint& Constraint : Constraints)
-	{
-		if (const FTransform* BoneTransform = Constraint.GetBoneTransform(Role, BoneName))
-		{
-			OutDesiredReach = Constraint.DesiredReach;
-			OutTransform = *BoneTransform;
-			return true;
-		}
-	}
-
-	OutDesiredReach = 0.f;
-	OutTransform = FTransform::Identity;
-	return false;
-}
-
-int32 FConstraints::Num() const
-{
-	return Constraints.Num();
-}
-
-const FConstraint* FConstraints::Find(const FConstraint& Constraint) const
-{
-	const int32 ConstraintIndex = Constraints.Find(Constraint);
-	if (ConstraintIndex != INDEX_NONE)
-	{
-		return &Constraints[ConstraintIndex];
-	}
-	return nullptr;
-}
-
-FConstraint* FConstraints::Find(const FConstraint& Constraint)
-{
-	const int32 ConstraintIndex = Constraints.Find(Constraint);
-	if (ConstraintIndex != INDEX_NONE)
-	{
-		return &Constraints[ConstraintIndex];
-	}
-	return nullptr;
-}
-
-void FConstraints::ProcessAll(float DeltaTime, const FSearchResult& SearchResult, const FInteractionSearchContext& SearchContext, const FConstraints* SearchContextPreviousConstraints)
-{
-	// making sure nobody already called UpdateConstraints or we didn't reset the constraints...
-	check(Constraints.IsEmpty());
-
-	const UPoseSearchDatabase* Database = SearchResult.Database.Get();
-	check(Database);
-
-	const FSearchIndexAsset* SearchIndexAsset = SearchResult.GetSearchIndexAsset();
-	check(SearchIndexAsset);
-
-	const FPoseSearchDatabaseAnimationAssetBase* DatabaseAnimationAssetBase = Database->GetDatabaseAnimationAsset(*SearchIndexAsset);
-	check(DatabaseAnimationAssetBase);
-
-	if (const UMultiAnimAsset* MultiAnimAsset = Cast<UMultiAnimAsset>(DatabaseAnimationAssetBase->GetAnimationAsset()))
-	{
-		const int32 NumAnimContexts = SearchContext.Num();
-		const float AssetTime = SearchResult.GetAssetTime();
-		const bool bIsMirrored = SearchIndexAsset->IsMirrored();
-		const FVector BlendParameters = SearchIndexAsset->GetBlendParameters();
-
-		FMemMark Mark(FMemStack::Get());
-
-		TConstArrayView<FConstraint> PreviousConstraints;
-		if (SearchContextPreviousConstraints)
-		{
-			PreviousConstraints = SearchContextPreviousConstraints->Constraints;
-		}
-
-		TArray<bool, TInlineAllocator<NumPreallocatedContraints, TMemStackAllocator<>>> VisitedPreviousConstraints;
-		VisitedPreviousConstraints.SetNumZeroed(PreviousConstraints.Num());
-		
-		// extracting the contraints for this SearchContext from MultiAnimAsset
-		FAnimNotifyContext PreAllocatedNotifyContext;
-		for (int32 AnimContextIndex = 0; AnimContextIndex < NumAnimContexts; ++AnimContextIndex)
-		{
-			const FRole Role = SearchContext.GetRole(AnimContextIndex);
-			const UAnimationAsset* Animation = MultiAnimAsset->GetAnimationAsset(Role);
-			check(Animation);
-
-			// initializing FAnimationAssetSampler only for extracting AnimNotifyStates
-			FAnimationAssetSampler Sampler(Animation, FTransform::Identity, FVector::ZeroVector, FAnimationAssetSampler::DefaultRootTransformSamplingRate, false, false);
-			Sampler.ExtractAnimNotifyEvents(AssetTime, PreAllocatedNotifyContext, [AssetTime, DeltaTime, this, &PreviousConstraints, &VisitedPreviousConstraints](const FAnimNotifyEvent* AnimNotifyEvent)
-				{
-					if (const UAnimNotifyState_PoseSearchConstraint* ConstraintNotifyState = Cast<const UAnimNotifyState_PoseSearchConstraint>(AnimNotifyEvent->NotifyStateClass))
-					{
-						// subtracting ConstraintNotifyState->CoolDownTime from the validity of this UAnimNotifyState_PoseSearchConstraint
-						const float NotifyEndTime = AnimNotifyEvent->GetEndTriggerTime();
-						if (AssetTime <= (NotifyEndTime - ConstraintNotifyState->CoolDownTime))
-						{
-							FConstraint NewConstraint(ConstraintNotifyState);
-							const int32 PreviousConstraintIndex = PreviousConstraints.Find(NewConstraint);
-							if (PreviousConstraintIndex != INDEX_NONE)
-							{
-								VisitedPreviousConstraints[PreviousConstraintIndex] = true;
-								const FConstraint& PreviousConstraint = PreviousConstraints[PreviousConstraintIndex];
-
-								// @todo: for now we just ignore conflicting constraints, but we should probably blend PreviousConstraint into NewConstraint :/
-								//        perhaps keeping the min or RampUpTime, to make the constraint as following the fastest one to get into full DesiredReach, and the max CoolDownTime to fade off with the slowest rate
-								//        not really sure what to do about TranslationWeight and RotationWeight.. maybe slowly blending them to the newer one?
-								if (PreviousConstraint.RampUpTime > UE_KINDA_SMALL_NUMBER)
-								{
-									NewConstraint.DesiredReach = FMath::Min(1.f, PreviousConstraint.DesiredReach + (DeltaTime / PreviousConstraint.RampUpTime));
-								}
-								else
-								{
-									NewConstraint.DesiredReach = 1.f;
-								}
-							}
-							else if (NewConstraint.RampUpTime <= UE_KINDA_SMALL_NUMBER)
-							{
-								NewConstraint.DesiredReach = 1.f;
-							}
-
-							// @todo: for now we also ignore duplicate constraints, and perhaps we should use a similar strategy of blending the constraints
-							Constraints.AddUnique(NewConstraint);
-						}
-					}
-					return true;
-				});
-		}
-
-		// adding remaining unprocessed constraints from PreviousConstraints
-		for (int32 PreviousConstraintIndex = 0; PreviousConstraintIndex < PreviousConstraints.Num(); ++PreviousConstraintIndex)
-		{
-			if (!VisitedPreviousConstraints[PreviousConstraintIndex])
-			{
-				const FConstraint& PreviousConstraint = PreviousConstraints[PreviousConstraintIndex];
-				if (PreviousConstraint.CoolDownTime > UE_KINDA_SMALL_NUMBER)
-				{
-					const float DesiredReach = PreviousConstraint.DesiredReach - (DeltaTime / PreviousConstraint.CoolDownTime);
-					if (DesiredReach > UE_KINDA_SMALL_NUMBER)
-					{
-						FConstraint& NewConstraint = Constraints.Add_GetRef(PreviousConstraint);
-						NewConstraint.DesiredReach = DesiredReach;
-					}
-				}
-			}
-		}
-
-		// process all the collected constraints
-		if (!Constraints.IsEmpty())
-		{
-			TArray<FTransform, TInlineAllocator<PreallocatedRolesNum, TMemStackAllocator<>>> ComponentToWorldTransforms;
-			TArray<TConstArrayView<FTransform>, TInlineAllocator<PreallocatedRolesNum, TMemStackAllocator<>>> AnimContextsComponentSpaceTransforms;
-			TArray<FBoneContainer, TInlineAllocator<PreallocatedRolesNum, TMemStackAllocator<>>> BoneContainers;
-			TArray<FCompactPose, TInlineAllocator<PreallocatedRolesNum, TMemStackAllocator<>>> Poses;
-			TArray<FCSPose<FCompactPose>, TInlineAllocator<PreallocatedRolesNum, TMemStackAllocator<>>> ComponentSpacePoses;
-			TArray<const USkeleton*, TInlineAllocator<PreallocatedRolesNum, TMemStackAllocator<>>> Skeletons;
-
-			ComponentToWorldTransforms.SetNum(NumAnimContexts);
-			AnimContextsComponentSpaceTransforms.SetNum(NumAnimContexts);
-			BoneContainers.SetNum(NumAnimContexts);
-			Poses.SetNum(NumAnimContexts);
-			ComponentSpacePoses.SetNum(NumAnimContexts);
-			Skeletons.SetNum(NumAnimContexts);
-
-			FRoleToIndex RoleToIndex;
-			for (int32 AnimContextIndex = 0; AnimContextIndex < NumAnimContexts; ++AnimContextIndex)
-			{
-				const UObject* AnimContext = SearchContext.GetAnimContext(AnimContextIndex);
-				check(AnimContext);
-				const USkeletalMeshComponent* SkeletalMeshComponent = GetContextSkeletalMeshComponent(AnimContext, false);
-				check(SkeletalMeshComponent);
-
-				// @todo: do we need this logic about URO etc? if we're executing this is because we update the mesh, so we probably don't need any interpolation...
-				const bool bUROInSync = SkeletalMeshComponent->ShouldUseUpdateRateOptimizations() && SkeletalMeshComponent->AnimUpdateRateParams != nullptr;
-				const bool bUsingExternalInterpolation = SkeletalMeshComponent->IsUsingExternalInterpolation();
-
-				const TArray<FTransform>& CachedComponentSpaceTransforms = SkeletalMeshComponent->GetCachedComponentSpaceTransforms();
-				const bool bArraySizesMatch = CachedComponentSpaceTransforms.Num() == SkeletalMeshComponent->GetComponentSpaceTransforms().Num();
-
-				ComponentToWorldTransforms[AnimContextIndex] = SkeletalMeshComponent->GetComponentTransform();
-				AnimContextsComponentSpaceTransforms[AnimContextIndex] = (bUROInSync || bUsingExternalInterpolation) && bArraySizesMatch ? CachedComponentSpaceTransforms : SkeletalMeshComponent->GetComponentSpaceTransforms();
-
-				// @todo: check FAnimNode_CopyPoseFromMesh::PreUpdate to see how to GetAnimCurves and GetCustomAttributes
-				const FRole Role = SearchContext.GetRole(AnimContextIndex);
-				RoleToIndex.Add(Role) = AnimContextIndex;
-
-				const UAnimationAsset* Animation = MultiAnimAsset->GetAnimationAsset(Role);
-				check(Animation);
-				const FTransform Origin = MultiAnimAsset->GetOrigin(Role);
-
-				FAnimationAssetSampler Sampler(Animation, Origin, BlendParameters);
-				const FTransform RootTransform = Sampler.ExtractRootTransform(AssetTime);
-
-				// @todo add support for mirroring if needed
-				//if (MirrorDataTables.IsValidIndex(ItemIndex) && MirrorDataTables[ItemIndex])
-				//{
-				//	const FMirrorDataCache MirrorDataCache(MirrorDataTables[ItemIndex]);
-				//	AssetRootBoneTransforms[ItemIndex] = MirrorDataCache.MirrorTransform(AssetRootBoneTransforms[ItemIndex]);
-				//}
-
-				// array containing the bone index of the root bone (0)
-				Skeletons[AnimContextIndex] = Animation->GetSkeleton();
-				check(Skeletons[AnimContextIndex]);
-				check(SkeletalMeshComponent->GetSkeletalMeshAsset()->GetSkeleton() == Skeletons[AnimContextIndex]);
-
-				const FReferenceSkeleton& ReferenceSkeleton = Skeletons[AnimContextIndex]->GetReferenceSkeleton();
-
-				// @todo: refine this BoneIndices / cache bone containers
-				const int32 NumBones = ReferenceSkeleton.GetNum();
-				TArray<uint16, TMemStackAllocator<>> BoneIndices;
-				BoneIndices.SetNumUninitialized(NumBones);
-				for (int32 BoneIndex = 0; BoneIndex < NumBones; ++BoneIndex)
-				{
-					BoneIndices[BoneIndex] = BoneIndex;
-				}
-
-				// extracting the pose, containing only the root bone from the Sampler 
-				BoneContainers[AnimContextIndex].InitializeTo(BoneIndices, UE::Anim::FCurveFilterSettings(UE::Anim::ECurveFilterMode::DisallowAll), *Skeletons[AnimContextIndex]);
-				Poses[AnimContextIndex].SetBoneContainer(&BoneContainers[AnimContextIndex]);
-				Sampler.ExtractPose(AssetTime, Poses[AnimContextIndex]);
-
-
-				// making sure the animation root bone transform is Identity, so we can confuse the root with the root BONE transform and preserve performances!
-				FTransform& RootBoneTransform = Poses[AnimContextIndex][FCompactPoseBoneIndex(0)];
-				if (!RootBoneTransform.Equals(FTransform::Identity))
-				{
-					const FVector Pos = RootBoneTransform.GetTranslation();
-					const FRotator Rot(RootBoneTransform.GetRotation());
-					UE_LOG(LogPoseSearch, Error, TEXT("FInteractionIsland::UpdateConstraints unsupported non identity root bone in %s at time %f Pos(%f, %f, %f), Rot(%f, %f, %f)"), *Animation->GetName(), AssetTime, Pos.X, Pos.Y, Pos.Z, Rot.Pitch, Rot.Yaw, Rot.Roll);
-				}
-				RootBoneTransform = Sampler.ExtractRootTransform(AssetTime);
-				ComponentSpacePoses[AnimContextIndex].InitPose(Poses[AnimContextIndex]);
-
-#if ENABLE_VISUAL_LOG
-				if (FVisualLogger::IsRecording())
-				{
-					static const TCHAR* LogName = TEXT("PoseSearchInteraction_Constraints");
-
-					for (int32 BoneIndex = 0; BoneIndex < ReferenceSkeleton.GetNum(); ++BoneIndex)
-					{
-						const int32 ParentBoneIndex = ReferenceSkeleton.GetParentIndex(BoneIndex);
-						if (ParentBoneIndex != INDEX_NONE)
-						{
-							const FVector CurrentBoneWorldLocation = ComponentToWorldTransforms[AnimContextIndex].TransformPosition(AnimContextsComponentSpaceTransforms[AnimContextIndex][BoneIndex].GetTranslation());
-							const FVector CurrentParentBoneWorldLocation = ComponentToWorldTransforms[AnimContextIndex].TransformPosition(AnimContextsComponentSpaceTransforms[AnimContextIndex][ParentBoneIndex].GetTranslation());
-
-							const FVector AnimBoneWorldLocation = ComponentSpacePoses[AnimContextIndex].GetComponentSpaceTransform(FCompactPoseBoneIndex(BoneIndex)).GetTranslation();
-							const FVector AnimParentBoneWorldLocation = ComponentSpacePoses[AnimContextIndex].GetComponentSpaceTransform(FCompactPoseBoneIndex(ParentBoneIndex)).GetTranslation();
-
-							for (int32 DrawAnimContextIndex = 0; DrawAnimContextIndex < NumAnimContexts; ++DrawAnimContextIndex)
-							{
-								const UObject* DrawAnimContext = SearchContext.GetAnimContext(DrawAnimContextIndex);
-								const FColor Color = DrawAnimContext == AnimContext ? FColor::Red : FColor::Blue;
-
-								// current characters poses
-								UE_VLOG_SEGMENT(DrawAnimContext, LogName, Display, CurrentBoneWorldLocation, CurrentParentBoneWorldLocation, Color, TEXT(""));
-
-								// sampled poses from UAnimationAsset
-								UE_VLOG_SEGMENT(DrawAnimContext, LogName, Display, AnimBoneWorldLocation, AnimParentBoneWorldLocation, Color, TEXT(""));
-							}
-						}
-					}
-				}
-#endif //ENABLE_VISUAL_LOG
-			}
-
-			for (FConstraint& Constraint : Constraints)
-			{
-				const int32* FromBoneRoleIndex = RoleToIndex.Find(Constraint.FromBoneRole);
-				const int32* ToBoneRoleIndex = RoleToIndex.Find(Constraint.ToBoneRole);
-				if (FromBoneRoleIndex && ToBoneRoleIndex)
-				{
-					FBoneReference FromBoneReference(Constraint.FromBoneName);
-					FBoneReference ToBoneReference(Constraint.ToBoneName);
-
-					if (FromBoneReference.Initialize(Skeletons[*FromBoneRoleIndex]) &&
-						ToBoneReference.Initialize(Skeletons[*ToBoneRoleIndex]))
-					{
-						const FTransform CurrentFromTransform = AnimContextsComponentSpaceTransforms[*FromBoneRoleIndex][FromBoneReference.BoneIndex] * ComponentToWorldTransforms[*FromBoneRoleIndex];
-						const FTransform CurrentToTransform = AnimContextsComponentSpaceTransforms[*ToBoneRoleIndex][ToBoneReference.BoneIndex] * ComponentToWorldTransforms[*ToBoneRoleIndex];
-
-						const FTransform& AnimFromTransform = ComponentSpacePoses[*FromBoneRoleIndex].GetComponentSpaceTransform(FCompactPoseBoneIndex(FromBoneReference.BoneIndex));
-						const FTransform& AnimToTransform = ComponentSpacePoses[*ToBoneRoleIndex].GetComponentSpaceTransform(FCompactPoseBoneIndex(ToBoneReference.BoneIndex));
-
-						const FTransform WantedCurrentToTransform = AnimToTransform.GetRelativeTransform(AnimFromTransform) * CurrentFromTransform;
-						const FTransform WantedCurrentFromTransform = AnimFromTransform.GetRelativeTransform(AnimToTransform) * CurrentToTransform;
-
-						Constraint.FromBoneTransform.SetTranslation(FMath::Lerp(WantedCurrentFromTransform.GetTranslation(), CurrentFromTransform.GetTranslation(), Constraint.TranslationWeight));
-						Constraint.FromBoneTransform.SetRotation(FQuat::Slerp(WantedCurrentFromTransform.GetRotation(), CurrentFromTransform.GetRotation(), Constraint.RotationWeight));
-						Constraint.FromBoneTransform.SetScale3D(CurrentFromTransform.GetScale3D());
-
-						Constraint.ToBoneTransform.SetTranslation(FMath::Lerp(WantedCurrentToTransform.GetTranslation(), CurrentToTransform.GetTranslation(), 1.f - Constraint.TranslationWeight));
-						Constraint.ToBoneTransform.SetRotation(FQuat::Slerp(WantedCurrentToTransform.GetRotation(), CurrentToTransform.GetRotation(), 1.f - Constraint.RotationWeight));
-						Constraint.ToBoneTransform.SetScale3D(CurrentToTransform.GetScale3D());
-
-#if WITH_EDITORONLY_DATA && ENABLE_ANIM_DEBUG && ENABLE_VISUAL_LOG
-						if (FVisualLogger::IsRecording())
-						{
-							static const TCHAR* LogName = TEXT("PoseSearchInteraction_Constraints");
-
-							for (int32 DrawAnimContextIndex = 0; DrawAnimContextIndex < NumAnimContexts; ++DrawAnimContextIndex)
-							{
-								const UObject* DrawAnimContext = SearchContext.GetAnimContext(DrawAnimContextIndex);
-
-								// current characters poses
-								UE_VLOG_SEGMENT(DrawAnimContext, LogName, Display, CurrentFromTransform.GetTranslation(), CurrentToTransform.GetTranslation(), FColor::Yellow, TEXT(""));
-
-								// sampled poses from UAnimationAsset
-								UE_VLOG_SEGMENT(DrawAnimContext, LogName, Display, AnimFromTransform.GetTranslation(), AnimToTransform.GetTranslation(), FColor::Yellow, TEXT(""));
-
-								// wanted transforms
-								UE_VLOG_SEGMENT(DrawAnimContext, LogName, Display, CurrentFromTransform.GetTranslation(), WantedCurrentToTransform.GetTranslation(), FColor::Green, TEXT(""));
-								UE_VLOG_SEGMENT(DrawAnimContext, LogName, Display, CurrentToTransform.GetTranslation(), WantedCurrentFromTransform.GetTranslation(), FColor::Green, TEXT(""));
-
-								UE_VLOG_SEGMENT(DrawAnimContext, LogName, Display, CurrentFromTransform.GetTranslation(), Constraint.FromBoneTransform.GetTranslation(), FColor::Orange, TEXT(""));
-								UE_VLOG_SEGMENT(DrawAnimContext, LogName, Display, CurrentToTransform.GetTranslation(), Constraint.ToBoneTransform.GetTranslation(), FColor::Orange, TEXT(""));
-							}
-						}
-#endif // WITH_EDITORONLY_DATA && ENABLE_ANIM_DEBUG && ENABLE_VISUAL_LOG
-					}
-				}
-			}
-		}
-	}
-}
-
 // FValidInteractionSearch
 ///////////////////////////////////////////////////////////
 FValidInteractionSearch& FValidInteractionSearch::Init(const FInteractionSearchContextBase& SearchContextBase)
@@ -2102,12 +1779,12 @@ void FInteractionIsland::UpdateConstraints(float DeltaTime)
 	}
 }
 
-bool FInteractionIsland::GetConstraint(const UObject* AnimContext, FName BoneName, float& OutDesiredReach, FTransform& OutTransform, bool bCompareOwningActors) const
+bool FInteractionIsland::GetConstraint(const UObject* AnimContext, FName SocketName, float& OutDesiredReach, FTransform& OutTransform, bool bCompareOwningActors) const
 {
 	FRole AnimContextRole;
 	if (const FInteractionSearchResult* SearchResult = GetInteractionSearchResult_AnyThread(AnimContext, bCompareOwningActors, AnimContextRole))
 	{
-		return SearchResult->Constraints.GetConstraint(AnimContextRole, BoneName, OutDesiredReach, OutTransform);
+		return SearchResult->Constraints.GetConstraint(AnimContextRole, SocketName, OutDesiredReach, OutTransform);
 	}
 
 	OutDesiredReach = 0.f;
diff --git a/Engine/Plugins/Animation/PoseSearch/Source/Runtime/Private/PoseSearchInteractionLibrary.cpp b/Engine/Plugins/Animation/PoseSearch/Source/Runtime/Private/PoseSearchInteractionLibrary.cpp
--- a/Engine/Plugins/Animation/PoseSearch/Source/Runtime/Private/PoseSearchInteractionLibrary.cpp
+++ b/Engine/Plugins/Animation/PoseSearch/Source/Runtime/Private/PoseSearchInteractionLibrary.cpp
@@ -139,11 +139,11 @@ FPoseSearchContinuingProperties UPoseSearchInteractionLibrary::GetMontageContinu
 	return ContinuingProperties;
 }
 
-void UPoseSearchInteractionLibrary::GetMotionMatchInteractionConstraint(const UObject* AnimContext, FName BoneName, float& OutDesiredReach, FTransform& OutTransform, bool& OutIsValid, bool bCompareOwningActors)
+void UPoseSearchInteractionLibrary::GetMotionMatchInteractionConstraint(const UObject* AnimContext, FName SocketName, float& OutDesiredReach, FTransform& OutTransform, bool& OutIsValid, bool bCompareOwningActors)
 {
 	if (UPoseSearchInteractionSubsystem* InteractionSubsystem = UPoseSearchInteractionSubsystem::GetSubsystem_AnyThread(AnimContext))
 	{
-		OutIsValid = InteractionSubsystem->GetConstraint(AnimContext, BoneName, OutDesiredReach, OutTransform, bCompareOwningActors);
+		OutIsValid = InteractionSubsystem->GetConstraint(AnimContext, SocketName, OutDesiredReach, OutTransform, bCompareOwningActors);
 	}
 	else
 	{
diff --git a/Engine/Plugins/Animation/PoseSearch/Source/Runtime/Private/PoseSearchInteractionSubsystem.cpp b/Engine/Plugins/Animation/PoseSearch/Source/Runtime/Private/PoseSearchInteractionSubsystem.cpp
--- a/Engine/Plugins/Animation/PoseSearch/Source/Runtime/Private/PoseSearchInteractionSubsystem.cpp
+++ b/Engine/Plugins/Animation/PoseSearch/Source/Runtime/Private/PoseSearchInteractionSubsystem.cpp
@@ -1392,13 +1392,13 @@ const UE::PoseSearch::FConstraints* UPoseSearchInteractionSubsystem::GetConstrai
 	return nullptr;
 }
 
-bool UPoseSearchInteractionSubsystem::GetConstraint(const UObject* AnimContext, FName BoneName, float& OutDesiredReach, FTransform& OutTransform, bool bCompareOwningActors)
+bool UPoseSearchInteractionSubsystem::GetConstraint(const UObject* AnimContext, FName SocketName, float& OutDesiredReach, FTransform& OutTransform, bool bCompareOwningActors)
 {
 	using namespace UE::PoseSearch;
 
 	if (const FInteractionIsland* Island = FindIsland(AnimContext, bCompareOwningActors))
 	{
-		return Island->GetConstraint(AnimContext, BoneName, OutDesiredReach, OutTransform, bCompareOwningActors);
+		return Island->GetConstraint(AnimContext, SocketName, OutDesiredReach, OutTransform, bCompareOwningActors);
 	}
 	
 	OutDesiredReach = 0.f;
diff --git a/Engine/Plugins/Animation/PoseSearch/Source/Runtime/Public/PoseSearch/PoseSearchAnimNotifies.h b/Engine/Plugins/Animation/PoseSearch/Source/Runtime/Public/PoseSearch/PoseSearchAnimNotifies.h
--- a/Engine/Plugins/Animation/PoseSearch/Source/Runtime/Public/PoseSearch/PoseSearchAnimNotifies.h
+++ b/Engine/Plugins/Animation/PoseSearch/Source/Runtime/Public/PoseSearch/PoseSearchAnimNotifies.h
@@ -130,22 +130,22 @@ class UAnimNotifyState_PoseSearchConstraint : public UAnimNotifyState_PoseSearch
 
 public:
 	UPROPERTY(EditAnywhere, Category = Config)
-	FName FromBoneName;
+	FName FromSocketName;
 
 	UPROPERTY(EditAnywhere, Category = Config)
-	FName FromBoneRole;
+	FName FromSocketRole;
 	
 	UPROPERTY(EditAnywhere, Category = Config)
-	FName ToBoneName;
+	FName ToSocketName;
 	
 	UPROPERTY(EditAnywhere, Category = Config)
-	FName ToBoneRole;
+	FName ToSocketRole;
 
-	// who should translate to reach full aligment? FromBone or ToBone? If TranslationWeight is 0 FromBone will move 100% and ToBone 0%, if TranslationWeight is 1 FromBone will move 0% and ToBone 100%,
+	// who should translate to reach full aligment? FromSocket or ToSocket? If TranslationWeight is 0 FromSocket will move 100% and ToSocket 0%, if TranslationWeight is 1 FromSocket will move 0% and ToSocket 100%,
 	UPROPERTY(EditAnywhere, Category = Config, meta = (ClampMin = "0", ClampMax = "1", UIMin = "0", UIMax = "1"))
 	float TranslationWeight = 0.f;
 	
-	// who should rotate to reach full aligment? FromBone or ToBone? If RotationWeight is 0 FromBone will move 100% and ToBone 0%, if RotationWeight is 1 FromBone will move 0% and ToBone 100%,
+	// who should rotate to reach full aligment? FromSocket or ToSocket? If RotationWeight is 0 FromSocket will move 100% and ToSocket 0%, if RotationWeight is 1 FromSocket will move 0% and ToSocket 100%,
 	UPROPERTY(EditAnywhere, Category = Config, meta = (ClampMin = "0", ClampMax = "1", UIMin = "0", UIMax = "1"))
 	float RotationWeight = 0.f;
 
diff --git a/Engine/Plugins/Animation/PoseSearch/Source/Runtime/Public/PoseSearch/PoseSearchConstraints.h b/Engine/Plugins/Animation/PoseSearch/Source/Runtime/Public/PoseSearch/PoseSearchConstraints.h
new file mode 100644
--- /dev/null
+++ b/Engine/Plugins/Animation/PoseSearch/Source/Runtime/Public/PoseSearch/PoseSearchConstraints.h
@@ -0,0 +1,80 @@
+// Copyright Epic Games, Inc. All Rights Reserved.
+
+#pragma once
+
+#include "PoseSearch/PoseSearchRole.h"
+
+#define UE_API POSESEARCH_API
+
+class UAnimNotifyState_PoseSearchConstraint;
+class UMultiAnimAsset;
+
+namespace UE::PoseSearch
+{
+struct FSearchResult;
+struct FInteractionSearchContext;
+
+// Experimental, this feature might be removed without warning, not for production use
+struct FConstraint
+{
+	FConstraint() = default;
+	FConstraint(const FConstraint& Other) = default;
+	FConstraint(FConstraint&& Other) = default;
+	FConstraint& operator=(const FConstraint& Other) = default;
+	FConstraint& operator=(FConstraint&& Other) = default;
+
+	FConstraint(const UAnimNotifyState_PoseSearchConstraint* ConstraintNotifyState);
+	bool operator==(const FConstraint& Other) const;
+	const FTransform* GetSocketTransform(FRole Role, FName SocketName) const;
+
+	FName FromSocketName;
+	FName FromSocketRole;
+	FName ToSocketName;
+	FName ToSocketRole;
+
+	// who should translate to reach full aligment? FromBone or ToBone? If TranslationWeight is 0 FromBone will move 100% and ToBone 0%, if TranslationWeight is 1 FromBone will move 0% and ToBone 100%,
+	float TranslationWeight = 0.f;
+	
+	// who should rotate to reach full aligment? FromBone or ToBone? If RotationWeight is 0 FromBone will move 100% and ToBone 0%, if RotationWeight is 1 FromBone will move 0% and ToBone 100%,
+	float RotationWeight = 0.f;
+
+	// rump up time used to calculate the desired reach for this constraint 
+	// (it doesn't start at the beginning of this notify state, BUT when this notify state is first seen, since MM can jump into a random frame of the animation containing the notify state)
+	float RampUpTime = 0.2f;
+
+	// cool down time used to calculate the desired reach for this constraint
+	// (it doesn't start at the end of this notify state, BUT when this notify state gets out of scope
+	// (ultimately getting out of scope when reaching the notify state end time minus CoolDownTime, since MM can jump into a random frame of the animation containing the notify state)
+	float CoolDownTime = 0.2f;
+
+	FTransform FromSocketTransform;
+	FTransform ToSocketTransform;
+
+	float DesiredReach = 0.f;
+};
+
+// Experimental, this feature might be removed without warning, not for production use
+struct FConstraints
+{
+	// @todo: mitigate / budget allocations properly
+	enum { NumPreallocatedContraints = 4 };
+
+	void ProcessAll(float DeltaTime, const FSearchResult& SearchResult, const FInteractionSearchContext& SearchContext, const FConstraints* SearchContextPreviousConstraints);
+	UE_API void ProcessAll(const UMultiAnimAsset* MultiAnimAsset, float DeltaTime, float AssetTime, bool bIsMirrored, const FVector& BlendParameters,
+		const TConstArrayView<TObjectPtr<const UObject>> AnimContexts, const TConstArrayView<FRole> Roles, TConstArrayView<FConstraint> PreviousConstraints = TConstArrayView<FConstraint>());
+	bool GetConstraint(FRole Role, FName SocketName, float& OutDesiredReach, FTransform& OutTransform) const;
+
+	const TConstArrayView<FConstraint> GetAll() const { return Constraints; }
+
+private:
+	int32 Num() const;
+	const FConstraint* Find(const FConstraint& Constraint) const;
+	FConstraint* Find(const FConstraint& Constraint);
+
+	typedef TArray<FConstraint, TInlineAllocator<NumPreallocatedContraints>> TConstraints;
+	TConstraints Constraints;
+};
+
+} // namespace UE::PoseSearch
+
+#undef UE_API
diff --git a/Engine/Plugins/Animation/PoseSearch/Source/Runtime/Public/PoseSearch/PoseSearchInteractionIsland.h b/Engine/Plugins/Animation/PoseSearch/Source/Runtime/Public/PoseSearch/PoseSearchInteractionIsland.h
--- a/Engine/Plugins/Animation/PoseSearch/Source/Runtime/Public/PoseSearch/PoseSearchInteractionIsland.h
+++ b/Engine/Plugins/Animation/PoseSearch/Source/Runtime/Public/PoseSearch/PoseSearchInteractionIsland.h
@@ -4,6 +4,7 @@
 
 #include "Engine/EngineBaseTypes.h"
 #include "PoseSearch/PoseSearchAnimNotifies.h"
+#include "PoseSearch/PoseSearchConstraints.h"
 #include "PoseSearch/PoseSearchInteractionAvailability.h"
 #include "PoseSearch/PoseSearchInteractionValidator.h"
 #include "PoseSearch/PoseSearchLibrary.h"
@@ -33,6 +34,7 @@ struct FInteractionSearchContextBase
 	const FRole GetRole(int32 Index) const { return Roles[Index]; }
 	const UPoseSearchDatabase* GetDatabase() const { return Database; }
 	const TConstArrayView<FRole> GetRoles() const { return Roles; }
+	const TConstArrayView<TObjectPtr<const UObject>> GetAnimContexts() const { return AnimContexts; }
 
 	bool IsDisableCollisions() const { return bDisableCollisions; }
 
@@ -64,60 +66,6 @@ private:
 	uint32 CachedHashForEquivalence = InvalidCachedHashForEquivalence;
 };
 
-struct FConstraint
-{
-	FConstraint() = default;
-	FConstraint(const FConstraint& Other) = default;
-	FConstraint(FConstraint&& Other) = default;
-	FConstraint& operator=(const FConstraint& Other) = default;
-	FConstraint& operator=(FConstraint&& Other) = default;
-
-	FConstraint(const UAnimNotifyState_PoseSearchConstraint* ConstraintNotifyState);
-	bool operator==(const FConstraint& Other) const;
-	const FTransform* GetBoneTransform(FRole Role, FName BoneName) const;
-
-	FName FromBoneName;
-	FName FromBoneRole;
-	FName ToBoneName;
-	FName ToBoneRole;
-
-	// who should translate to reach full aligment? FromBone or ToBone? If TranslationWeight is 0 FromBone will move 100% and ToBone 0%, if TranslationWeight is 1 FromBone will move 0% and ToBone 100%,
-	float TranslationWeight = 0.f;
-	
-	// who should rotate to reach full aligment? FromBone or ToBone? If RotationWeight is 0 FromBone will move 100% and ToBone 0%, if RotationWeight is 1 FromBone will move 0% and ToBone 100%,
-	float RotationWeight = 0.f;
-
-	// rump up time used to calculate the desired reach for this constraint 
-	// (it doesn't start at the beginning of this notify state, BUT when this notify state is first seen, since MM can jump into a random frame of the animation containing the notify state)
-	float RampUpTime = 0.2f;
-
-	// cool down time used to calculate the desired reach for this constraint
-	// (it doesn't start at the end of this notify state, BUT when this notify state gets out of scope
-	// (ultimately getting out of scope when reaching the notify state end time minus CoolDownTime, since MM can jump into a random frame of the animation containing the notify state)
-	float CoolDownTime = 0.2f;
-
-	FTransform FromBoneTransform;
-	FTransform ToBoneTransform;
-
-	float DesiredReach = 0.f;
-};
-
-struct FConstraints
-{
-	// @todo: mitigate / budget allocations properly
-	enum { NumPreallocatedContraints = 4 };
-
-	void ProcessAll(float DeltaTime, const FSearchResult& SearchResult, const FInteractionSearchContext& SearchContext, const FConstraints* SearchContextPreviousConstraints);
-
-	bool GetConstraint(FRole Role, FName BoneName, float& OutDesiredReach, FTransform& OutTransform) const;
-
-private:
-	int32 Num() const;
-	const FConstraint* Find(const FConstraint& Constraint) const;
-	FConstraint* Find(const FConstraint& Constraint);
-	TArray<FConstraint, TInlineAllocator<NumPreallocatedContraints>> Constraints;
-};
-
 // Experimental, this feature might be removed without warning, not for production use
 struct FInteractionSearchResult : public FSearchResult
 {
@@ -256,7 +204,7 @@ struct FInteractionIsland
 
 	void AddReferencedObjects(FReferenceCollector& Collector);
 	void UpdateConstraints(float DeltaTime);
-	bool GetConstraint(const UObject* AnimContext, FName BoneName, float& OutDesiredReach, FTransform& OutTransform, bool bCompareOwningActors = false) const;
+	bool GetConstraint(const UObject* AnimContext, FName SocketName, float& OutDesiredReach, FTransform& OutTransform, bool bCompareOwningActors = false) const;
 
 private:
 #if ENABLE_ANIM_DEBUG
diff --git a/Engine/Plugins/Animation/PoseSearch/Source/Runtime/Public/PoseSearch/PoseSearchInteractionLibrary.h b/Engine/Plugins/Animation/PoseSearch/Source/Runtime/Public/PoseSearch/PoseSearchInteractionLibrary.h
--- a/Engine/Plugins/Animation/PoseSearch/Source/Runtime/Public/PoseSearch/PoseSearchInteractionLibrary.h
+++ b/Engine/Plugins/Animation/PoseSearch/Source/Runtime/Public/PoseSearch/PoseSearchInteractionLibrary.h
@@ -46,7 +46,7 @@ public:
 	static UE_API FPoseSearchContinuingProperties GetMontageContinuingProperties(UAnimInstance* AnimInstance);
 
 	UFUNCTION(Experimental, BlueprintPure, Category = "Animation|Pose Search", meta = (BlueprintThreadSafe))
-	static UE_API void GetMotionMatchInteractionConstraint(const UObject* AnimContext, FName BoneName, float& OutDesiredReach, FTransform& OutTransform, bool& IsValid, bool bCompareOwningActors = true);
+	static UE_API void GetMotionMatchInteractionConstraint(const UObject* AnimContext, FName SocketName, float& OutDesiredReach, FTransform& OutTransform, bool& IsValid, bool bCompareOwningActors = true);
 };
 
 #undef UE_API
diff --git a/Engine/Plugins/Animation/PoseSearch/Source/Runtime/Public/PoseSearch/PoseSearchInteractionSubsystem.h b/Engine/Plugins/Animation/PoseSearch/Source/Runtime/Public/PoseSearch/PoseSearchInteractionSubsystem.h
--- a/Engine/Plugins/Animation/PoseSearch/Source/Runtime/Public/PoseSearch/PoseSearchInteractionSubsystem.h
+++ b/Engine/Plugins/Animation/PoseSearch/Source/Runtime/Public/PoseSearch/PoseSearchInteractionSubsystem.h
@@ -110,7 +110,7 @@ public:
 
 	// use bCompareOwningActors = true in case you're querying the UPoseSearchInteractionSubsystem from a different AnimContext, 
 	// for example when using this API from a different ABP from a different skeletal mesh component
-	bool GetConstraint(const UObject* AnimContext, FName BoneName, float& OutDesiredReach, FTransform& OutTransform, bool bCompareOwningActors = false);
+	bool GetConstraint(const UObject* AnimContext, FName SocketName, float& OutDesiredReach, FTransform& OutTransform, bool bCompareOwningActors = false);
 
 private:
 #if ENABLE_ANIM_DEBUG
