PoseSearch - fix comments typos

diff --git a/Engine/Plugins/Animation/PoseSearch/Source/Runtime/Private/PoseSearchInteractionUtils.cpp b/Engine/Plugins/Animation/PoseSearch/Source/Runtime/Private/PoseSearchInteractionUtils.cpp
--- a/Engine/Plugins/Animation/PoseSearch/Source/Runtime/Private/PoseSearchInteractionUtils.cpp
+++ b/Engine/Plugins/Animation/PoseSearch/Source/Runtime/Private/PoseSearchInteractionUtils.cpp
@@ -97,7 +97,7 @@ FTransform CalculateDeltaAlignment(const FTransform& MeshWithoutOffset, const FT
 	// calculating the FullDeltaAlignment as the delta transform that brings the actor to its full aligned transform.
 	const FTransform FullDeltaAlignment = FullAlignedTransform.GetRelativeTransform(MeshWithOffset);
 
-	// calculating the DeltaAlignment as blend between the NoDeltaAlignment and the FullDeltaAlignment: how much the character need to move to get to the desired aligment
+	// calculating the DeltaAlignment as blend between the NoDeltaAlignment and the FullDeltaAlignment: how much the character need to move to get to the desired alignment
 	const FTransform DeltaAlignment(FMath::Lerp(NoDeltaAlignment.GetRotation(), FullDeltaAlignment.GetRotation(), FMath::Clamp(WarpingRotationRatio, 0.f, 1.f)),
 		FMath::Lerp(NoDeltaAlignment.GetTranslation(), FullDeltaAlignment.GetTranslation(), FMath::Clamp(WarpingTranslationRatio, 0.f, 1.f)), FVector::OneVector);
 
diff --git a/Engine/Plugins/Animation/PoseSearch/Source/Runtime/Public/PoseSearch/PoseSearchAnimNotifies.h b/Engine/Plugins/Animation/PoseSearch/Source/Runtime/Public/PoseSearch/PoseSearchAnimNotifies.h
--- a/Engine/Plugins/Animation/PoseSearch/Source/Runtime/Public/PoseSearch/PoseSearchAnimNotifies.h
+++ b/Engine/Plugins/Animation/PoseSearch/Source/Runtime/Public/PoseSearch/PoseSearchAnimNotifies.h
@@ -141,15 +141,15 @@ public:
 	UPROPERTY(EditAnywhere, Category = Config)
 	FName ToSocketRole;
 
-	// who should translate to reach full aligment? FromSocket or ToSocket? If TranslationWeight is 0 FromSocket will move 100% and ToSocket 0%, if TranslationWeight is 1 FromSocket will move 0% and ToSocket 100%,
+	// who should translate to reach full alignment? FromSocket or ToSocket? If TranslationWeight is 0 FromSocket will move 100% and ToSocket 0%, if TranslationWeight is 1 FromSocket will move 0% and ToSocket 100%,
 	UPROPERTY(EditAnywhere, Category = Config, meta = (ClampMin = "0", ClampMax = "1", UIMin = "0", UIMax = "1"))
 	float TranslationWeight = 0.f;
 	
-	// who should rotate to reach full aligment? FromSocket or ToSocket? If RotationWeight is 0 FromSocket will move 100% and ToSocket 0%, if RotationWeight is 1 FromSocket will move 0% and ToSocket 100%,
+	// who should rotate to reach full alignment? FromSocket or ToSocket? If RotationWeight is 0 FromSocket will move 100% and ToSocket 0%, if RotationWeight is 1 FromSocket will move 0% and ToSocket 100%,
 	UPROPERTY(EditAnywhere, Category = Config, meta = (ClampMin = "0", ClampMax = "1", UIMin = "0", UIMax = "1"))
 	float RotationWeight = 0.f;
 
-	// rump up time used to calculate the desired reach for this constraint 
+	// ramp up time used to calculate the desired reach for this constraint 
 	// (it doesn't start at the beginning of this notify state, BUT when this notify state is first seen, since MM can jump into a random frame of the animation containing the notify state)
 	UPROPERTY(EditAnywhere, Category = Config, meta = (ClampMin = "0", UIMin = "0"))
 	float RampUpTime = 0.2f;
diff --git a/Engine/Plugins/Animation/PoseSearch/Source/Runtime/Public/PoseSearch/PoseSearchConstraints.h b/Engine/Plugins/Animation/PoseSearch/Source/Runtime/Public/PoseSearch/PoseSearchConstraints.h
--- a/Engine/Plugins/Animation/PoseSearch/Source/Runtime/Public/PoseSearch/PoseSearchConstraints.h
+++ b/Engine/Plugins/Animation/PoseSearch/Source/Runtime/Public/PoseSearch/PoseSearchConstraints.h
@@ -32,13 +32,13 @@ struct FConstraint
 	FName ToSocketName;
 	FName ToSocketRole;
 
-	// who should translate to reach full aligment? FromBone or ToBone? If TranslationWeight is 0 FromBone will move 100% and ToBone 0%, if TranslationWeight is 1 FromBone will move 0% and ToBone 100%,
+	// who should translate to reach full alignment? FromBone or ToBone? If TranslationWeight is 0 FromBone will move 100% and ToBone 0%, if TranslationWeight is 1 FromBone will move 0% and ToBone 100%,
 	float TranslationWeight = 0.f;
 	
-	// who should rotate to reach full aligment? FromBone or ToBone? If RotationWeight is 0 FromBone will move 100% and ToBone 0%, if RotationWeight is 1 FromBone will move 0% and ToBone 100%,
+	// who should rotate to reach full alignment? FromBone or ToBone? If RotationWeight is 0 FromBone will move 100% and ToBone 0%, if RotationWeight is 1 FromBone will move 0% and ToBone 100%,
 	float RotationWeight = 0.f;
 
-	// rump up time used to calculate the desired reach for this constraint 
+	// ramp up time used to calculate the desired reach for this constraint 
 	// (it doesn't start at the beginning of this notify state, BUT when this notify state is first seen, since MM can jump into a random frame of the animation containing the notify state)
 	float RampUpTime = 0.2f;
 
