fix posesearchcolumn making shared database assets unintentionally

diff --git a/Engine/Plugins/Animation/PoseSearch/Source/Runtime/Private/Chooser/PoseSearchChooserColumn.cpp b/Engine/Plugins/Animation/PoseSearch/Source/Runtime/Private/Chooser/PoseSearchChooserColumn.cpp
--- a/Engine/Plugins/Animation/PoseSearch/Source/Runtime/Private/Chooser/PoseSearchChooserColumn.cpp
+++ b/Engine/Plugins/Animation/PoseSearch/Source/Runtime/Private/Chooser/PoseSearchChooserColumn.cpp
@@ -742,7 +742,9 @@ void FPoseSearchColumn::DeinitializeScratchArea(TArrayView<uint8> ScratchArea) c
 void FPoseSearchColumn::Initialize(UChooserTable* OuterChooser)
 {
 	FChooserColumnBase::Initialize(OuterChooser);
-	InternalDatabase = NewObject<UPoseSearchDatabase>(OuterChooser, OuterChooser->GetFName(), RF_Transactional);
+	FName BaseName = FName(FString::Format(TEXT("{0}_PSD"), { OuterChooser->GetName() }));
+	FName DatabaseName = MakeUniqueObjectName(OuterChooser, UPoseSearchDatabase::StaticClass(), BaseName);
+	InternalDatabase = NewObject<UPoseSearchDatabase>(OuterChooser, DatabaseName, RF_Transactional);
 }
 
 void FPoseSearchColumn::AddToDetails(FInstancedPropertyBag& PropertyBag, int32 ColumnIndex, int32 RowIndex)
