Accounting for AutoRTFM being enabled for these TickableObject systems.

#rnx

diff --git a/Engine/Plugins/Animation/PoseSearch/Source/Runtime/Private/PoseSearchInteractionSubsystem.cpp b/Engine/Plugins/Animation/PoseSearch/Source/Runtime/Private/PoseSearchInteractionSubsystem.cpp
--- a/Engine/Plugins/Animation/PoseSearch/Source/Runtime/Private/PoseSearchInteractionSubsystem.cpp
+++ b/Engine/Plugins/Animation/PoseSearch/Source/Runtime/Private/PoseSearchInteractionSubsystem.cpp
@@ -1299,6 +1299,22 @@ void UPoseSearchInteractionSubsystem::Tick(float DeltaSeconds)
 
 	check(IsInGameThread());
 
+	// If we're in an AutoRTFM Transaction, exit the TickableObject Tick (which is top-level) and run again outside the transaction
+	// This is temporary because AutoRTFM does not have a ThreadSafeCounter; we will eventually need to solve this as we do derference
+	// UObjects inside this function.
+	if (AutoRTFM::IsTransactional())
+	{
+		AutoRTFM::OnCommit([WeakThis = TWeakObjectPtr(this), DeltaSeconds]()
+		{
+			if (WeakThis.IsValid())
+			{
+				WeakThis->Tick(DeltaSeconds);
+			}
+		});
+
+		return;
+	}
+
 	Super::Tick(DeltaSeconds);
 
 	FMemMark Mark(FMemStack::Get());
diff --git a/Engine/Plugins/PCG/Source/PCG/Private/Subsystems/PCGSubsystem.cpp b/Engine/Plugins/PCG/Source/PCG/Private/Subsystems/PCGSubsystem.cpp
--- a/Engine/Plugins/PCG/Source/PCG/Private/Subsystems/PCGSubsystem.cpp
+++ b/Engine/Plugins/PCG/Source/PCG/Private/Subsystems/PCGSubsystem.cpp
@@ -325,6 +325,20 @@ void UPCGSubsystem::Tick(float DeltaSeconds)
 	TRACE_CPUPROFILER_EVENT_SCOPE(UPCGSubsystem::Tick);
 	LLM_SCOPE_BYTAG(PCG);
 
+	// If we're in an AutoRTFM Transaction, exit the TickableObject Tick (which is top-level) and run again outside the transaction
+	if (AutoRTFM::IsTransactional())
+	{
+		AutoRTFM::OnCommit([WeakThis = TWeakObjectPtr(this), DeltaSeconds]()
+		{
+			if (WeakThis.IsValid())
+			{
+				WeakThis->Tick(DeltaSeconds);
+			}
+		});
+
+		return;
+	}
+
 	Super::Tick(DeltaSeconds);
 
 	ExecuteBeginTickActions();
