PoseSearch - supporting TimeOffset to UPoseSearchInteractionAsset::CalculateWarpTransforms

diff --git a/Engine/Plugins/Animation/PoseSearch/Source/Editor/Private/PoseSearchInteractionAssetEditor.cpp b/Engine/Plugins/Animation/PoseSearch/Source/Editor/Private/PoseSearchInteractionAssetEditor.cpp
--- a/Engine/Plugins/Animation/PoseSearch/Source/Editor/Private/PoseSearchInteractionAssetEditor.cpp
+++ b/Engine/Plugins/Animation/PoseSearch/Source/Editor/Private/PoseSearchInteractionAssetEditor.cpp
@@ -362,7 +362,12 @@ void FInteractionAssetViewModel::Tick(float DeltaSeconds)
 		FullAlignedActorTransforms.SetNum(NumRoles);
 		const TConstArrayView<const UMirrorDataTable*> MirrorDataTables;
 		const TConstArrayView<bool> RelevantRoleIndexes;
-		InteractionAsset->CalculateWarpTransforms(PlayTime, ActorTransforms, FullAlignedActorTransforms, MirrorDataTables, RelevantRoleIndexes, AnimContexts);
+
+		InteractionAsset->CalculateWarpTransforms(PlayTime, 0.f, ActorTransforms, FullAlignedActorTransforms, MirrorDataTables, RelevantRoleIndexes, AnimContexts);
+
+#if ENABLE_ANIM_DEBUG
+		const UWorld* DebugWorld = GetWorld();
+#endif // ENABLE_ANIM_DEBUG
 
 		for (int32 InteractionAssetRoleIndex = 0; InteractionAssetRoleIndex < NumRoles; ++InteractionAssetRoleIndex)
 		{
@@ -373,8 +378,33 @@ void FInteractionAssetViewModel::Tick(float DeltaSeconds)
 			check(PreviewActor.GetRole() == InteractionAssetRole);
 			FTransform DebugActorTransform;
 			DebugActorTransform.Blend(ActorTransforms[InteractionAssetRoleIndex], FullAlignedActorTransforms[InteractionAssetRoleIndex], InteractionAsset->DebugWarpAmount);
+
+#if ENABLE_ANIM_DEBUG
+			// FullAlignedActorTransforms drawing 
+			const FVector FullAlignedActorPosition = FullAlignedActorTransforms[InteractionAssetRoleIndex].GetTranslation();
+			const FQuat FullAlignedActorRotation = FullAlignedActorTransforms[InteractionAssetRoleIndex].GetRotation();
+			DrawDebugCircle(DebugWorld, FMatrix(FVector::ZAxisVector, FVector::XAxisVector, FVector::YAxisVector, FullAlignedActorPosition), 5.f, 16, FColorList::Yellow, false, 0.f, SDPG_Foreground, 0.f, false);
+			DrawDebugLine(DebugWorld, FullAlignedActorPosition, FullAlignedActorPosition + FullAlignedActorRotation.RotateVector(FVector::XAxisVector) * 10, FColorList::Yellow, false, 0.f, SDPG_Foreground, 0.f);
+#endif // ENABLE_ANIM_DEBUG
+
 			PreviewActor.ForceDebugActorTransform(DebugActorTransform);
 		}
+
+#if ENABLE_ANIM_DEBUG
+		if (InteractionAsset->DebugTimeOffset > UE_KINDA_SMALL_NUMBER)
+		{
+			InteractionAsset->CalculateWarpTransforms(PlayTime, InteractionAsset->DebugTimeOffset, ActorTransforms, FullAlignedActorTransforms, MirrorDataTables, RelevantRoleIndexes, AnimContexts);
+
+			for (int32 InteractionAssetRoleIndex = 0; InteractionAssetRoleIndex < NumRoles; ++InteractionAssetRoleIndex)
+			{
+				// FullAlignedActorTransforms drawing 
+				const FVector FullAlignedActorPosition = FullAlignedActorTransforms[InteractionAssetRoleIndex].GetTranslation();
+				const FQuat FullAlignedActorRotation = FullAlignedActorTransforms[InteractionAssetRoleIndex].GetRotation();
+				DrawDebugCircle(DebugWorld, FMatrix(FVector::ZAxisVector, FVector::XAxisVector, FVector::YAxisVector, FullAlignedActorPosition), 5.f, 16, FColorList::Pink, false, 0.f, SDPG_Foreground, 0.f, false);
+				DrawDebugLine(DebugWorld, FullAlignedActorPosition, FullAlignedActorPosition + FullAlignedActorRotation.RotateVector(FVector::XAxisVector) * 10, FColorList::Pink, false, 0.f, SDPG_Foreground, 0.f);
+			}
+		}
+#endif // ENABLE_ANIM_DEBUG
 	}
 #endif // WITH_EDITORONLY_DATA
 
diff --git a/Engine/Plugins/Animation/PoseSearch/Source/Runtime/Private/AnimNode_MotionMatching.cpp b/Engine/Plugins/Animation/PoseSearch/Source/Runtime/Private/AnimNode_MotionMatching.cpp
--- a/Engine/Plugins/Animation/PoseSearch/Source/Runtime/Private/AnimNode_MotionMatching.cpp
+++ b/Engine/Plugins/Animation/PoseSearch/Source/Runtime/Private/AnimNode_MotionMatching.cpp
@@ -357,7 +357,7 @@ void FAnimNode_MotionMatching::Evaluate_AnyThread(FPoseContext& Output)
 
 	TArray<FTransform, TInlineAllocator<PreallocatedRolesNum>> FullAlignedTransforms;
 	FullAlignedTransforms.SetNum(NumRoles);
-	CalculateFullAlignedTransforms(MotionMatchingState.SearchResult, bWarpUsingRootBone, FullAlignedTransforms);
+	CalculateFullAlignedTransforms(MotionMatchingState.SearchResult, MotionMatchingState.SearchResult.SelectedTime, 0.f, bWarpUsingRootBone, FullAlignedTransforms);
 
 	// NoTe: keep in mind DeltaAlignment is relative to the previous execution frame so we still need to extract and and apply the current animation root motion transform to get to the current frame full aligned transform.
 	const FTransform DeltaAlignment = CalculateDeltaAlignment(MeshWithoutOffset, MeshWithOffset, FullAlignedTransforms[CurrentResultRoleIndex], WarpingRotationRatio, WarpingTranslationRatio);
diff --git a/Engine/Plugins/Animation/PoseSearch/Source/Runtime/Private/AnimNode_MotionMatchingInteraction.cpp b/Engine/Plugins/Animation/PoseSearch/Source/Runtime/Private/AnimNode_MotionMatchingInteraction.cpp
--- a/Engine/Plugins/Animation/PoseSearch/Source/Runtime/Private/AnimNode_MotionMatchingInteraction.cpp
+++ b/Engine/Plugins/Animation/PoseSearch/Source/Runtime/Private/AnimNode_MotionMatchingInteraction.cpp
@@ -163,7 +163,7 @@ void FAnimNode_MotionMatchingInteraction::Evaluate_AnyThread(FPoseContext& Outpu
 
 	TArray<FTransform, TInlineAllocator<PreallocatedRolesNum>> FullAlignedTransforms;
 	FullAlignedTransforms.SetNum(NumRoles);
-	CalculateFullAlignedTransforms(CurrentResult, bWarpUsingRootBone, FullAlignedTransforms);
+	CalculateFullAlignedTransforms(CurrentResult, CurrentResult.SelectedTime, 0.f, bWarpUsingRootBone, FullAlignedTransforms);
 
 	// NoTe: keep in mind DeltaAlignment is relative to the previous execution frame so we still need to extract and and apply the current animation root motion transform to get to the current frame full aligned transform.
 	const FTransform DeltaAlignment = CalculateDeltaAlignment(MeshWithoutOffset, MeshWithOffset, FullAlignedTransforms[CurrentResultRoleIndex], WarpingRotationRatio, WarpingTranslationRatio);
diff --git a/Engine/Plugins/Animation/PoseSearch/Source/Runtime/Private/PoseSearchInteractionAsset.cpp b/Engine/Plugins/Animation/PoseSearch/Source/Runtime/Private/PoseSearchInteractionAsset.cpp
--- a/Engine/Plugins/Animation/PoseSearch/Source/Runtime/Private/PoseSearchInteractionAsset.cpp
+++ b/Engine/Plugins/Animation/PoseSearch/Source/Runtime/Private/PoseSearchInteractionAsset.cpp
@@ -253,7 +253,7 @@ USkeletalMesh* UPoseSearchInteractionAsset::GetPreviewMesh(const UE::PoseSearch:
 
 #endif // WITH_EDITOR
 
-void UPoseSearchInteractionAsset::CalculateWarpTransforms(float Time, const TConstArrayView<const FTransform> ActorRootBoneTransforms, TArrayView<FTransform> FullAlignedActorRootBoneTransforms,
+void UPoseSearchInteractionAsset::CalculateWarpTransforms(float Time, float TimeOffset, const TConstArrayView<const FTransform> ActorRootBoneTransforms, TArrayView<FTransform> FullAlignedActorRootBoneTransforms,
 	const TConstArrayView<const UMirrorDataTable*> MirrorDataTables, const TConstArrayView<bool> RelevantRoleIndexes, const TConstArrayView<TObjectPtr<const UObject>> DebugAnimContexts) const
 {
 	using namespace UE::PoseSearch;
@@ -294,6 +294,9 @@ void UPoseSearchInteractionAsset::CalculateWarpTransforms(float Time, const TCon
 		TArray<FTransform, TInlineAllocator<PreallocatedRolesNum>> AssetRootBoneTransforms;
 		AssetRootBoneTransforms.SetNum(ItemsNum);
 
+		TArray<FTransform, TInlineAllocator<PreallocatedRolesNum>> OffsettedAssetRootBoneTransforms;
+		OffsettedAssetRootBoneTransforms.SetNum(ItemsNum);
+
 		// ItemIndex is the RoleIndex and Role = Item.Role
 		for (int32 ItemIndex = 0; ItemIndex < ItemsNum; ++ItemIndex)
 		{
@@ -307,6 +310,15 @@ void UPoseSearchInteractionAsset::CalculateWarpTransforms(float Time, const TCon
 			{
 				const FMirrorDataCache MirrorDataCache(MirrorDataTables[ItemIndex]);
 				AssetRootBoneTransforms[ItemIndex] = MirrorDataCache.MirrorTransform(AssetRootBoneTransforms[ItemIndex]);
+
+				if (TimeOffset > UE_KINDA_SMALL_NUMBER)
+				{
+					OffsettedAssetRootBoneTransforms[ItemIndex] = MirrorDataCache.MirrorTransform(Sampler.ExtractRootTransform(Time + TimeOffset));
+				}
+			}
+			else if (TimeOffset > UE_KINDA_SMALL_NUMBER)
+			{
+				OffsettedAssetRootBoneTransforms[ItemIndex] = Sampler.ExtractRootTransform(Time + TimeOffset);
 			}
 
 #if ENABLE_ANIM_DEBUG && WITH_EDITOR
@@ -601,6 +613,16 @@ void UPoseSearchInteractionAsset::CalculateWarpTransforms(float Time, const TCon
 			DrawDebugLine(DebugWorld, AssetReferencePosition, AssetReferencePosition + ActorsReferenceOrientation.RotateVector(FVector::XAxisVector) * 10.f, FColorList::Blue, false, 0.f, SDPG_Foreground, 0.f);
 			DrawDebugLine(DebugWorld, ActorsReferencePosition, ActorsReferencePosition + ActorsReferenceOrientation.RotateVector(FVector::XAxisVector) * 10.f, FColorList::Blue, false, 0.f, SDPG_Foreground, 0.f);
 #endif // ENABLE_ANIM_DEBUG && WITH_EDITOR
-		}						
+		}
+
+		if (TimeOffset > UE_KINDA_SMALL_NUMBER)
+		{
+			for (int32 ItemIndex = 0; ItemIndex < ItemsNum; ++ItemIndex)
+			{
+				// calculating the delta transform that move the root bone from Time to Time + TimeOffset
+				const FTransform DeltaAssetRootBoneTransform = OffsettedAssetRootBoneTransforms[ItemIndex].GetRelativeTransform(AssetRootBoneTransforms[ItemIndex]);
+				FullAlignedActorRootBoneTransforms[ItemIndex] = DeltaAssetRootBoneTransform * FullAlignedActorRootBoneTransforms[ItemIndex];
+			}
+		}
 	}
 }
diff --git a/Engine/Plugins/Animation/PoseSearch/Source/Runtime/Private/PoseSearchInteractionLibrary.cpp b/Engine/Plugins/Animation/PoseSearch/Source/Runtime/Private/PoseSearchInteractionLibrary.cpp
--- a/Engine/Plugins/Animation/PoseSearch/Source/Runtime/Private/PoseSearchInteractionLibrary.cpp
+++ b/Engine/Plugins/Animation/PoseSearch/Source/Runtime/Private/PoseSearchInteractionLibrary.cpp
@@ -98,12 +98,12 @@ bool UPoseSearchInteractionLibrary::MotionMatchInteraction(FPoseSearchBlueprintR
 	return bIsInteraction;
 }
 
-void UPoseSearchInteractionLibrary::CalculateFullAlignedTransforms(const FPoseSearchBlueprintResult& Result, bool bWarpUsingRootBone, TArray<FTransform>& OutFullAlignedTransforms)
+void UPoseSearchInteractionLibrary::CalculateFullAlignedTransforms(const FPoseSearchBlueprintResult& Result, float TimeOffset, bool bWarpUsingRootBone, TArray<FTransform>& OutFullAlignedTransforms)
 {
 	if (!Result.AnimContexts.IsEmpty() && Result.SelectedAnim && Result.SelectedAnim->IsA<UMultiAnimAsset>())
 	{
 		OutFullAlignedTransforms.SetNum(Result.AnimContexts.Num());
-		UE::PoseSearch::CalculateFullAlignedTransforms(Result, bWarpUsingRootBone, OutFullAlignedTransforms);
+		UE::PoseSearch::CalculateFullAlignedTransforms(Result, Result.SelectedTime, TimeOffset, bWarpUsingRootBone, OutFullAlignedTransforms);
 	}
 	else
 	{
@@ -112,13 +112,13 @@ void UPoseSearchInteractionLibrary::CalculateFullAlignedTransforms(const FPoseSe
 	}
 }
 
-void UPoseSearchInteractionLibrary::CalculateFullAlignedTransform(const FPoseSearchBlueprintResult& Result, bool bWarpUsingRootBone, FTransform& OutFullAlignedTransform)
+void UPoseSearchInteractionLibrary::CalculateFullAlignedTransform(const FPoseSearchBlueprintResult& Result, float TimeOffset, bool bWarpUsingRootBone, FTransform& OutFullAlignedTransform)
 {
 	if (Result.AnimContexts.IsValidIndex(Result.RoleIndex) && Result.SelectedAnim && Result.SelectedAnim->IsA<UMultiAnimAsset>())
 	{
 		TArray<FTransform, TInlineAllocator<UE::PoseSearch::PreallocatedRolesNum>> FullAlignedTransforms;
 		FullAlignedTransforms.SetNum(Result.AnimContexts.Num());
-		UE::PoseSearch::CalculateFullAlignedTransforms(Result, bWarpUsingRootBone, FullAlignedTransforms);
+		UE::PoseSearch::CalculateFullAlignedTransforms(Result, Result.SelectedTime, TimeOffset, bWarpUsingRootBone, FullAlignedTransforms);
 		OutFullAlignedTransform = FullAlignedTransforms[Result.RoleIndex];
 	}
 	else
@@ -128,6 +128,28 @@ void UPoseSearchInteractionLibrary::CalculateFullAlignedTransform(const FPoseSea
 	}
 }
 
+void UPoseSearchInteractionLibrary::FindAnimNotifyStateEndTime(const UAnimationAsset* AnimationAsset, const UClass* NotifyStateClass, float& EndTime, bool& IsValid)
+{
+	using namespace UE::PoseSearch;
+
+	const FAnimationAssetSampler Sampler(AnimationAsset, FTransform::Identity, FVector::ZeroVector, FAnimationAssetSampler::DefaultRootTransformSamplingRate, false, false);
+	for (const FAnimNotifyEvent& AnimNotifyEvent : Sampler.GetAllAnimNotifyEvents())
+	{
+		if (const UClass* ThisNotifyStateClass = AnimNotifyEvent.NotifyStateClass->GetClass())
+		{
+			if (ThisNotifyStateClass->IsChildOf(NotifyStateClass))
+			{
+				IsValid = true;
+				EndTime = AnimNotifyEvent.GetEndTriggerTime();
+				return;
+			}
+		}
+	}
+
+	IsValid = false;
+	EndTime = 0.f;
+}
+
 FPoseSearchContinuingProperties UPoseSearchInteractionLibrary::GetMontageContinuingProperties(UAnimInstance* AnimInstance)
 {
 	FPoseSearchContinuingProperties ContinuingProperties;
diff --git a/Engine/Plugins/Animation/PoseSearch/Source/Runtime/Private/PoseSearchInteractionUtils.cpp b/Engine/Plugins/Animation/PoseSearch/Source/Runtime/Private/PoseSearchInteractionUtils.cpp
--- a/Engine/Plugins/Animation/PoseSearch/Source/Runtime/Private/PoseSearchInteractionUtils.cpp
+++ b/Engine/Plugins/Animation/PoseSearch/Source/Runtime/Private/PoseSearchInteractionUtils.cpp
@@ -37,17 +37,17 @@ FRoleToIndex MakeRoleToIndex(const UMultiAnimAsset* MultiAnimAsset)
 	return RoleToIndex;
 }
 
-void CalculateFullAlignedTransformsAtTime(const FPoseSearchBlueprintResult& CurrentResult, float Time, bool bWarpUsingRootBone, TArrayView<FTransform> OutFullAlignedTransforms)
+void CalculateFullAlignedTransforms(const FPoseSearchBlueprintResult& Result, float Time, float TimeOffset, bool bWarpUsingRootBone, TArrayView<FTransform> OutFullAlignedTransforms)
 {
-	const UMultiAnimAsset* MultiAnimAsset = CastChecked<UMultiAnimAsset>(CurrentResult.SelectedAnim);
+	const UMultiAnimAsset* MultiAnimAsset = CastChecked<UMultiAnimAsset>(Result.SelectedAnim);
 	const int32 NumRoles = MultiAnimAsset->GetNumRoles();
 	
 	check(OutFullAlignedTransforms.Num() == NumRoles);
-	check(CurrentResult.ActorRootTransforms.Num() == CurrentResult.ActorRootBoneTransforms.Num());
-	check(CurrentResult.ActorRootTransforms.Num() == CurrentResult.AnimContexts.Num());
-	check(CurrentResult.SelectedDatabase != nullptr && CurrentResult.SelectedDatabase->Schema != nullptr);
+	check(Result.ActorRootTransforms.Num() == Result.ActorRootBoneTransforms.Num());
+	check(Result.ActorRootTransforms.Num() == Result.AnimContexts.Num());
+	check(Result.SelectedDatabase != nullptr && Result.SelectedDatabase->Schema != nullptr);
 
-	TConstArrayView<FTransform> ActorTransforms = CurrentResult.ActorRootTransforms;
+	TConstArrayView<FTransform> ActorTransforms = Result.ActorRootTransforms;
 	
 	// if bWarpUsingRootBone we set ActorTransforms as the actors root bone world transforms instead of the root transforms
 	TArray<FTransform, TInlineAllocator<PreallocatedRolesNum, TMemStackAllocator<>>> WorldActorRootBoneTransforms;
@@ -56,20 +56,20 @@ void CalculateFullAlignedTransformsAtTime(const FPoseSearchBlueprintResult& Curr
 		WorldActorRootBoneTransforms.SetNum(NumRoles);
 		for (int32 Index = 0; Index < NumRoles; ++Index)
 		{
-			WorldActorRootBoneTransforms[Index] = CurrentResult.ActorRootBoneTransforms[Index] * CurrentResult.ActorRootTransforms[Index];
+			WorldActorRootBoneTransforms[Index] = Result.ActorRootBoneTransforms[Index] * Result.ActorRootTransforms[Index];
 		}
 
 		ActorTransforms = WorldActorRootBoneTransforms;
 	}
 
 	TArray<const UMirrorDataTable*, TInlineAllocator<PreallocatedRolesNum, TMemStackAllocator<>>> MirrorDataTables;
-	if (CurrentResult.bIsMirrored)
+	if (Result.bIsMirrored)
 	{
 		MirrorDataTables.SetNum(NumRoles);
 		for (int32 MultiAnimAssetRoleIndex = 0; MultiAnimAssetRoleIndex < NumRoles; ++MultiAnimAssetRoleIndex)
 		{
 			const FRole Role = MultiAnimAsset->GetRole(MultiAnimAssetRoleIndex);
-			const FPoseSearchRoledSkeleton* RoledSkeleton = CurrentResult.SelectedDatabase->Schema->GetRoledSkeleton(Role);
+			const FPoseSearchRoledSkeleton* RoledSkeleton = Result.SelectedDatabase->Schema->GetRoledSkeleton(Role);
 			check(RoledSkeleton);
 			MirrorDataTables[MultiAnimAssetRoleIndex] = RoledSkeleton->MirrorDataTable;
 		}
@@ -79,14 +79,9 @@ void CalculateFullAlignedTransformsAtTime(const FPoseSearchBlueprintResult& Curr
 	RelevantRoleIndexes.SetNum(NumRoles);
 	for (int32 Index = 0; Index < NumRoles; ++Index)
 	{
-		RelevantRoleIndexes[Index] = CurrentResult.AnimContexts[Index] != nullptr;
+		RelevantRoleIndexes[Index] = Result.AnimContexts[Index] != nullptr;
 	}
-	MultiAnimAsset->CalculateWarpTransforms(Time, ActorTransforms, OutFullAlignedTransforms, MirrorDataTables, RelevantRoleIndexes);
-}
-
-void CalculateFullAlignedTransforms(const FPoseSearchBlueprintResult& CurrentResult, bool bWarpUsingRootBone, TArrayView<FTransform> OutFullAlignedTransforms)
-{
-	CalculateFullAlignedTransformsAtTime(CurrentResult, CurrentResult.SelectedTime, bWarpUsingRootBone, OutFullAlignedTransforms);
+	MultiAnimAsset->CalculateWarpTransforms(Time, TimeOffset, ActorTransforms, OutFullAlignedTransforms, MirrorDataTables, RelevantRoleIndexes);
 }
 
 FTransform CalculateDeltaAlignment(const FTransform& MeshWithoutOffset, const FTransform& MeshWithOffset, const FTransform& FullAlignedTransform, float WarpingRotationRatio, float WarpingTranslationRatio)
diff --git a/Engine/Plugins/Animation/PoseSearch/Source/Runtime/Public/PoseSearch/MultiAnimAsset.h b/Engine/Plugins/Animation/PoseSearch/Source/Runtime/Public/PoseSearch/MultiAnimAsset.h
--- a/Engine/Plugins/Animation/PoseSearch/Source/Runtime/Public/PoseSearch/MultiAnimAsset.h
+++ b/Engine/Plugins/Animation/PoseSearch/Source/Runtime/Public/PoseSearch/MultiAnimAsset.h
@@ -32,7 +32,12 @@ public:
 	[[nodiscard]] virtual UAnimationAsset* GetAnimationAsset(const FName& Role) const PURE_VIRTUAL(UMultiAnimAsset::GetAnimationAsset, return nullptr;);
 	[[nodiscard]] virtual FTransform GetOrigin(const FName& Role) const PURE_VIRTUAL(UMultiAnimAsset::GetOrigin, return FTransform::Identity;);
 
-	virtual void CalculateWarpTransforms(float Time, const TConstArrayView<const FTransform> ActorRootBoneTransforms, TArrayView<FTransform> FullAlignedActorRootBoneTransforms,
+	// calculates the full aligned root bone transforms (FullAlignedActorRootBoneTransforms) at time "Time + TimeOffset" given the current actors root bone transforms (ActorRootBoneTransforms) at "Time"
+	// the mirror data tables, the RelevantRoleIndexes (array of what roles are present or "relevant"). DebugAnimContexts can be used to draw debug lines
+	// For example if the actors are playing the associated GetAnimationAsset(Role) at time Time, by providing the ActorRootBoneTransforms and an TimeOffset of zero, CalculateWarpTransforms will return the
+	// FullAlignedActorRootBoneTransforms should be NOW to be in full alignment. Any TimeOffset > 0 will have the method returning a future FullAlignedActorRootBoneTransforms. This is useful in synergy with
+	// Motion Warping to provide a transform to use as full aligment at the ned of the notify state time
+	virtual void CalculateWarpTransforms(float Time, float TimeOffset, const TConstArrayView<const FTransform> ActorRootBoneTransforms, TArrayView<FTransform> FullAlignedActorRootBoneTransforms,
 		const TConstArrayView<const UMirrorDataTable*> MirrorDataTables = TConstArrayView<const UMirrorDataTable*>(), const TConstArrayView<bool> RelevantRoleIndexes = TConstArrayView<bool>(),
 		const TConstArrayView<TObjectPtr<const UObject>> DebugAnimContexts = TConstArrayView<TObjectPtr<const UObject>>()) const PURE_VIRTUAL(UMultiAnimAsset::CalculateWarpTransforms, );
 
diff --git a/Engine/Plugins/Animation/PoseSearch/Source/Runtime/Public/PoseSearch/PoseSearchInteractionAsset.h b/Engine/Plugins/Animation/PoseSearch/Source/Runtime/Public/PoseSearch/PoseSearchInteractionAsset.h
--- a/Engine/Plugins/Animation/PoseSearch/Source/Runtime/Public/PoseSearch/PoseSearchInteractionAsset.h
+++ b/Engine/Plugins/Animation/PoseSearch/Source/Runtime/Public/PoseSearch/PoseSearchInteractionAsset.h
@@ -73,6 +73,10 @@ public:
 	UPROPERTY(Transient, EditAnywhere, Category = "Debug", meta = (ClampMin = "0", ClampMax = "1", EditCondition=bEnableDebugWarp, EditConditionHides))
 	float DebugWarpAmount = 0.f;
 
+	// useful to test TimeOffset
+	UPROPERTY(Transient, EditAnywhere, Category = "Debug", meta = (ClampMin = "0", EditCondition=bEnableDebugWarp, EditConditionHides))
+	float DebugTimeOffset = 0.f;
+
 	UPROPERTY(Transient, EditAnywhere, Category = "Debug")
 	bool bEnableDebugWarp = false;
 #endif // WITH_EDITORONLY_DATA
@@ -93,7 +97,7 @@ public:
 	UE_API virtual USkeletalMesh* GetPreviewMesh(const UE::PoseSearch::FRole& Role) const override;
 #endif // WITH_EDITOR
 
-	UE_API virtual void CalculateWarpTransforms(float Time, const TConstArrayView<const FTransform> ActorRootBoneTransforms, TArrayView<FTransform> FullAlignedActorRootBoneTransforms,
+	UE_API virtual void CalculateWarpTransforms(float Time, float TimeOffset, const TConstArrayView<const FTransform> ActorRootBoneTransforms, TArrayView<FTransform> FullAlignedActorRootBoneTransforms,
 		const TConstArrayView<const UMirrorDataTable*> MirrorDataTables = TConstArrayView<const UMirrorDataTable*>(), const TConstArrayView<bool> RelevantRoleIndexes = TConstArrayView<bool>(),
 		const TConstArrayView<TObjectPtr<const UObject>> DebugAnimContexts = TConstArrayView<TObjectPtr<const UObject>>()) const override;
 	
diff --git a/Engine/Plugins/Animation/PoseSearch/Source/Runtime/Public/PoseSearch/PoseSearchInteractionLibrary.h b/Engine/Plugins/Animation/PoseSearch/Source/Runtime/Public/PoseSearch/PoseSearchInteractionLibrary.h
--- a/Engine/Plugins/Animation/PoseSearch/Source/Runtime/Public/PoseSearch/PoseSearchInteractionLibrary.h
+++ b/Engine/Plugins/Animation/PoseSearch/Source/Runtime/Public/PoseSearch/PoseSearchInteractionLibrary.h
@@ -34,10 +34,13 @@ public:
 	static UE_API FPoseSearchBlueprintResult MotionMatchInteraction(TArray<FPoseSearchInteractionAvailability> Availabilities, const UObject* AnimContext, FName PoseHistoryName, bool bValidateResultAgainstAvailabilities = true);
 
 	UFUNCTION(Experimental, BlueprintCallable, Category = "Animation|Pose Search", meta = (BlueprintThreadSafe))
-	static UE_API void CalculateFullAlignedTransforms(const FPoseSearchBlueprintResult& CurrentResult, bool bWarpUsingRootBone, TArray<FTransform>& OutFullAlignedTransforms);
+	static UE_API void CalculateFullAlignedTransforms(const FPoseSearchBlueprintResult& CurrentResult, float TimeOffset, bool bWarpUsingRootBone, TArray<FTransform>& OutFullAlignedTransforms);
 
 	UFUNCTION(Experimental, BlueprintCallable, Category = "Animation|Pose Search", meta = (BlueprintThreadSafe))
-	static UE_API void CalculateFullAlignedTransform(const FPoseSearchBlueprintResult& CurrentResult, bool bWarpUsingRootBone, FTransform& OutFullAlignedTransform);
+	static UE_API void CalculateFullAlignedTransform(const FPoseSearchBlueprintResult& CurrentResult, float TimeOffset, bool bWarpUsingRootBone, FTransform& OutFullAlignedTransform);
+
+	UFUNCTION(Experimental, BlueprintPure, Category = "Animation|Pose Search", meta = (BlueprintThreadSafe))
+	static UE_API void FindAnimNotifyStateEndTime(const UAnimationAsset* AnimationAsset, const UClass* NotifyStateClass, float& EndTime, bool& IsValid);
 
 	static UE_API void MotionMatchInteraction(FPoseSearchBlueprintResult& Result, const TArrayView<const FPoseSearchInteractionAvailability> Availabilities, const UObject* AnimContext, FName PoseHistoryName, const UE::PoseSearch::IPoseHistory* PoseHistory, bool bValidateResultAgainstAvailabilities = true);
 	static UE_API bool MotionMatchInteraction(FPoseSearchBlueprintResult& InOutResult, const TArrayView<const FPoseSearchInteractionAvailability> Availabilities, const UObject* AnimContext, FName PoseHistoryName, const UE::PoseSearch::IPoseHistory* PoseHistory, bool bValidateResultAgainstAvailabilities, bool bKeepInteractionAlive, float BlendTime);
diff --git a/Engine/Plugins/Animation/PoseSearch/Source/Runtime/Public/PoseSearch/PoseSearchInteractionUtils.h b/Engine/Plugins/Animation/PoseSearch/Source/Runtime/Public/PoseSearch/PoseSearchInteractionUtils.h
--- a/Engine/Plugins/Animation/PoseSearch/Source/Runtime/Public/PoseSearch/PoseSearchInteractionUtils.h
+++ b/Engine/Plugins/Animation/PoseSearch/Source/Runtime/Public/PoseSearch/PoseSearchInteractionUtils.h
@@ -17,10 +17,7 @@ POSESEARCH_API int32 GetRoleIndex(const UMultiAnimAsset* MultiAnimAsset, const F
 FRoleToIndex MakeRoleToIndex(const UMultiAnimAsset* MultiAnimAsset);
 	
 // Experimental, this feature might be removed without warning, not for production use
-POSESEARCH_API void CalculateFullAlignedTransformsAtTime(const FPoseSearchBlueprintResult& CurrentResult, float Time, bool bWarpUsingRootBone, TArrayView<FTransform> OutFullAlignedTransforms);
-
-// Experimental, this feature might be removed without warning, not for production use
-POSESEARCH_API void CalculateFullAlignedTransforms(const FPoseSearchBlueprintResult& CurrentResult, bool bWarpUsingRootBone, TArrayView<FTransform> OutFullAlignedTransforms);
+POSESEARCH_API void CalculateFullAlignedTransforms(const FPoseSearchBlueprintResult& Result, float Time, float TimeOffset, bool bWarpUsingRootBone, TArrayView<FTransform> OutFullAlignedTransforms);
 
 // Experimental, this feature might be removed without warning, not for production use
 POSESEARCH_API FTransform CalculateDeltaAlignment(const FTransform& MeshWithoutOffset, const FTransform& MeshWithOffset, const FTransform& FullAlignedTransform, float WarpingRotationRatio, float WarpingTranslationRatio);
diff --git a/Engine/Plugins/Experimental/UAF/UAFPoseSearch/Source/UAFPoseSearch/Private/MotionMatchingTrait.cpp b/Engine/Plugins/Experimental/UAF/UAFPoseSearch/Source/UAFPoseSearch/Private/MotionMatchingTrait.cpp
--- a/Engine/Plugins/Experimental/UAF/UAFPoseSearch/Source/UAFPoseSearch/Private/MotionMatchingTrait.cpp
+++ b/Engine/Plugins/Experimental/UAF/UAFPoseSearch/Source/UAFPoseSearch/Private/MotionMatchingTrait.cpp
@@ -399,7 +399,7 @@ PRAGMA_ENABLE_DEPRECATION_WARNINGS
 					TArray<FTransform, TInlineAllocator<PreallocatedRolesNum>> FullAlignedTransforms;
 					FullAlignedTransforms.SetNum(NumRoles);
 
-					CalculateFullAlignedTransformsAtTime(MotionMatchingState.SearchResult, MultiAnimAsset->GetPlayLength(FVector::ZeroVector), SharedData->GetbWarpUsingRootBone(Binding), FullAlignedTransforms);
+					CalculateFullAlignedTransforms(MotionMatchingState.SearchResult, MultiAnimAsset->GetPlayLength(FVector::ZeroVector), 0.f, SharedData->GetbWarpUsingRootBone(Binding), FullAlignedTransforms);
 					
 					FUAFAssetInstance& Instance = Binding.GetTraitPtr().GetNodeInstance()->GetOwner();
 					FName AlignmentTransformName = SharedData->GetAlignmentTransformVariableName(Binding);
@@ -651,7 +651,7 @@ void FAnimNextMotionMatchingTask::Execute(UE::UAF::FEvaluationVM& VM) const
 	TArray<FTransform, TInlineAllocator<PreallocatedRolesNum>> FullAlignedTransforms;
 	FullAlignedTransforms.SetNum(NumRoles);
 
-	CalculateFullAlignedTransforms(SearchResult, bWarpUsingRootBone, FullAlignedTransforms);
+	CalculateFullAlignedTransforms(SearchResult, SearchResult.SelectedTime, 0.f, bWarpUsingRootBone, FullAlignedTransforms);
 
 	// @todo: implement this properly once we have UE::AnimationWarping::FRootOffsetProvider
 	const FTransform& MeshWithoutOffset = ComponentTransform;
