PoseSearch - removing no longer necessary experimental UPoseSearchInteractionLibrary::FindAnimNotifyStateEndTime

diff --git a/Engine/Plugins/Animation/PoseSearch/Source/Runtime/Private/PoseSearchInteractionLibrary.cpp b/Engine/Plugins/Animation/PoseSearch/Source/Runtime/Private/PoseSearchInteractionLibrary.cpp
--- a/Engine/Plugins/Animation/PoseSearch/Source/Runtime/Private/PoseSearchInteractionLibrary.cpp
+++ b/Engine/Plugins/Animation/PoseSearch/Source/Runtime/Private/PoseSearchInteractionLibrary.cpp
@@ -128,28 +128,6 @@ void UPoseSearchInteractionLibrary::CalculateFullAlignedTransform(const FPoseSea
 	}
 }
 
-void UPoseSearchInteractionLibrary::FindAnimNotifyStateEndTime(const UAnimationAsset* AnimationAsset, const UClass* NotifyStateClass, float& EndTime, bool& IsValid)
-{
-	using namespace UE::PoseSearch;
-
-	const FAnimationAssetSampler Sampler(AnimationAsset, FTransform::Identity, FVector::ZeroVector, FAnimationAssetSampler::DefaultRootTransformSamplingRate, false, false);
-	for (const FAnimNotifyEvent& AnimNotifyEvent : Sampler.GetAllAnimNotifyEvents())
-	{
-		if (const UClass* ThisNotifyStateClass = AnimNotifyEvent.NotifyStateClass->GetClass())
-		{
-			if (ThisNotifyStateClass->IsChildOf(NotifyStateClass))
-			{
-				IsValid = true;
-				EndTime = AnimNotifyEvent.GetEndTriggerTime();
-				return;
-			}
-		}
-	}
-
-	IsValid = false;
-	EndTime = 0.f;
-}
-
 FPoseSearchContinuingProperties UPoseSearchInteractionLibrary::GetMontageContinuingProperties(UAnimInstance* AnimInstance)
 {
 	FPoseSearchContinuingProperties ContinuingProperties;
diff --git a/Engine/Plugins/Animation/PoseSearch/Source/Runtime/Public/PoseSearch/PoseSearchInteractionLibrary.h b/Engine/Plugins/Animation/PoseSearch/Source/Runtime/Public/PoseSearch/PoseSearchInteractionLibrary.h
--- a/Engine/Plugins/Animation/PoseSearch/Source/Runtime/Public/PoseSearch/PoseSearchInteractionLibrary.h
+++ b/Engine/Plugins/Animation/PoseSearch/Source/Runtime/Public/PoseSearch/PoseSearchInteractionLibrary.h
@@ -39,9 +39,6 @@ public:
 	UFUNCTION(Experimental, BlueprintCallable, Category = "Animation|Pose Search", meta = (BlueprintThreadSafe))
 	static UE_API void CalculateFullAlignedTransform(const FPoseSearchBlueprintResult& CurrentResult, float TimeOffset, bool bWarpUsingRootBone, FTransform& OutFullAlignedTransform);
 
-	UFUNCTION(Experimental, BlueprintPure, Category = "Animation|Pose Search", meta = (BlueprintThreadSafe))
-	static UE_API void FindAnimNotifyStateEndTime(const UAnimationAsset* AnimationAsset, const UClass* NotifyStateClass, float& EndTime, bool& IsValid);
-
 	static UE_API void MotionMatchInteraction(FPoseSearchBlueprintResult& Result, const TArrayView<const FPoseSearchInteractionAvailability> Availabilities, const UObject* AnimContext, FName PoseHistoryName, const UE::PoseSearch::IPoseHistory* PoseHistory, bool bValidateResultAgainstAvailabilities = true);
 	static UE_API bool MotionMatchInteraction(FPoseSearchBlueprintResult& InOutResult, const TArrayView<const FPoseSearchInteractionAvailability> Availabilities, const UObject* AnimContext, FName PoseHistoryName, const UE::PoseSearch::IPoseHistory* PoseHistory, bool bValidateResultAgainstAvailabilities, bool bKeepInteractionAlive, float BlendTime);
 
