PoseSearch - adding UPoseSearchInteractionLibrary::CalculateNoAlignedTransform to calculate motion warp "recovery"

diff --git a/Engine/Plugins/Animation/PoseSearch/Source/Runtime/Private/PoseSearchInteractionLibrary.cpp b/Engine/Plugins/Animation/PoseSearch/Source/Runtime/Private/PoseSearchInteractionLibrary.cpp
--- a/Engine/Plugins/Animation/PoseSearch/Source/Runtime/Private/PoseSearchInteractionLibrary.cpp
+++ b/Engine/Plugins/Animation/PoseSearch/Source/Runtime/Private/PoseSearchInteractionLibrary.cpp
@@ -128,6 +128,45 @@ void UPoseSearchInteractionLibrary::CalculateFullAlignedTransform(const FPoseSea
 	}
 }
 
+void UPoseSearchInteractionLibrary::CalculateNoAlignedTransform(const FPoseSearchBlueprintResult& Result, float TimeOffset, bool bWarpUsingRootBone, FTransform& OutNoAlignedTransform)
+{
+	using namespace UE::PoseSearch;
+
+	if (!Result.ActorRootTransforms.IsValidIndex(Result.RoleIndex))
+	{
+		UE_LOG(LogPoseSearch, Error, TEXT("UPoseSearchInteractionLibrary::CalculateNoAlignedTransform failed!"));
+		OutNoAlignedTransform.SetIdentity();
+		return;
+	}
+
+	const UMultiAnimAsset* MultiAnimAsset = Cast<UMultiAnimAsset>(Result.SelectedAnim);
+	if (!MultiAnimAsset)
+	{
+		UE_LOG(LogPoseSearch, Error, TEXT("UPoseSearchInteractionLibrary::CalculateNoAlignedTransform failed!"));
+		OutNoAlignedTransform.SetIdentity();
+		return;
+	}
+
+	const UAnimationAsset* AnimationAsset = MultiAnimAsset->GetAnimationAsset(Result.Role);
+	if (!AnimationAsset)
+	{
+		UE_LOG(LogPoseSearch, Error, TEXT("UPoseSearchInteractionLibrary::CalculateNoAlignedTransform failed!"));
+		OutNoAlignedTransform.SetIdentity();
+		return;
+	}
+
+	check(Result.ActorRootBoneTransforms.Num() == Result.ActorRootTransforms.Num());
+
+	const FAnimationAssetSampler Sampler(AnimationAsset, FTransform::Identity, Result.BlendParameters);
+	const FTransform SelectedTimeRootTransform = Sampler.ExtractRootTransform(Result.SelectedTime);
+	const FTransform OffsettedTimeRootTransform = Sampler.ExtractRootTransform(Result.SelectedTime + TimeOffset);
+
+	const FTransform DeltaRootTransform = OffsettedTimeRootTransform.GetRelativeTransform(SelectedTimeRootTransform);
+	const FTransform ActorRootTransform = bWarpUsingRootBone ? Result.ActorRootBoneTransforms[Result.RoleIndex] * Result.ActorRootTransforms[Result.RoleIndex] : Result.ActorRootTransforms[Result.RoleIndex];
+
+	OutNoAlignedTransform = DeltaRootTransform * ActorRootTransform;
+}
+
 FPoseSearchContinuingProperties UPoseSearchInteractionLibrary::GetMontageContinuingProperties(UAnimInstance* AnimInstance)
 {
 	FPoseSearchContinuingProperties ContinuingProperties;
diff --git a/Engine/Plugins/Animation/PoseSearch/Source/Runtime/Public/PoseSearch/PoseSearchInteractionLibrary.h b/Engine/Plugins/Animation/PoseSearch/Source/Runtime/Public/PoseSearch/PoseSearchInteractionLibrary.h
--- a/Engine/Plugins/Animation/PoseSearch/Source/Runtime/Public/PoseSearch/PoseSearchInteractionLibrary.h
+++ b/Engine/Plugins/Animation/PoseSearch/Source/Runtime/Public/PoseSearch/PoseSearchInteractionLibrary.h
@@ -33,12 +33,22 @@ public:
 	UFUNCTION(Experimental, BlueprintCallable, Category = "Animation|Pose Search", meta = (BlueprintThreadSafe))
 	static UE_API FPoseSearchBlueprintResult MotionMatchInteraction(TArray<FPoseSearchInteractionAvailability> Availabilities, const UObject* AnimContext, FName PoseHistoryName, bool bValidateResultAgainstAvailabilities = true);
 
+	// this function calculates the mesh space transforms (OutFullAlignedTransforms) the AnimContext(s) should reach within a time of TimeOffset
+	// to be consistently aligned to the selected UMultiAnimAsset (CurrentResult.SelectedAnim) at the selected time (CurrentResult.SelectedTime), 
+	// given the interacting AnimContext(s) transforms (CurrentResult.ActorRootBoneTransforms, CurrentResult.ActorRootTransforms)
+	// if bWarpUsingRootBone is true the function will use the AnimContext root bone transforms, otherwise the mesh space transforms
 	UFUNCTION(Experimental, BlueprintCallable, Category = "Animation|Pose Search", meta = (BlueprintThreadSafe))
 	static UE_API void CalculateFullAlignedTransforms(const FPoseSearchBlueprintResult& CurrentResult, float TimeOffset, bool bWarpUsingRootBone, TArray<FTransform>& OutFullAlignedTransforms);
 
+	// this function calculates the mesh space transform for the AnimContext assigned the role CurrentResult.Role as CalculateFullAlignedTransforms does for all the AnimContext(s)
 	UFUNCTION(Experimental, BlueprintCallable, Category = "Animation|Pose Search", meta = (BlueprintThreadSafe))
 	static UE_API void CalculateFullAlignedTransform(const FPoseSearchBlueprintResult& CurrentResult, float TimeOffset, bool bWarpUsingRootBone, FTransform& OutFullAlignedTransform);
 
+	// this function calculates the mesh space transform (OutNoAlignedTransform), or root bone transform if bWarpUsingRootBone is true,
+	// the AnimContext will reach within a time of TimeOffset, if it continues playing the CurrentResult.SelectedAnim from CurrentResult.SelectedTime forward
+	UFUNCTION(Experimental, BlueprintCallable, Category = "Animation|Pose Search", meta = (BlueprintThreadSafe))
+	static UE_API void CalculateNoAlignedTransform(const FPoseSearchBlueprintResult& CurrentResult, float TimeOffset, bool bWarpUsingRootBone, FTransform& OutNoAlignedTransform);
+
 	static UE_API void MotionMatchInteraction(FPoseSearchBlueprintResult& Result, const TArrayView<const FPoseSearchInteractionAvailability> Availabilities, const UObject* AnimContext, FName PoseHistoryName, const UE::PoseSearch::IPoseHistory* PoseHistory, bool bValidateResultAgainstAvailabilities = true);
 	static UE_API bool MotionMatchInteraction(FPoseSearchBlueprintResult& InOutResult, const TArrayView<const FPoseSearchInteractionAvailability> Availabilities, const UObject* AnimContext, FName PoseHistoryName, const UE::PoseSearch::IPoseHistory* PoseHistory, bool bValidateResultAgainstAvailabilities, bool bKeepInteractionAlive, float BlendTime);
 
