PoseSearch - improved motion matching schema channels filtering. implemented position channel max position filter as FPoseSearchFilterMaxPosition and added a new added FPoseSearchFilterWedge to facilitate MM interaction pie wedges filtering

diff --git a/Engine/Plugins/Animation/PoseSearch/Source/Editor/Private/PoseSearchDatabaseViewModel.cpp b/Engine/Plugins/Animation/PoseSearch/Source/Editor/Private/PoseSearchDatabaseViewModel.cpp
--- a/Engine/Plugins/Animation/PoseSearch/Source/Editor/Private/PoseSearchDatabaseViewModel.cpp
+++ b/Engine/Plugins/Animation/PoseSearch/Source/Editor/Private/PoseSearchDatabaseViewModel.cpp
@@ -350,7 +350,7 @@ bool FDatabasePreviewActor::DrawPreviewActors(TConstArrayView<FDatabasePreviewAc
 		PoseHistories.Add(&ArchivedPoseHistory);
 	}
 
-	UE::PoseSearch::FDebugDrawParams DrawParams(AnimContexts, PoseHistories, RoleToIndex, PoseSearchDatabase);
+	FDebugDrawParams DrawParams(AnimContexts, PoseHistories, RoleToIndex, PoseSearchDatabase, FDebugDrawParams::EDrawContext::DrawQuery);
 	DrawParams.DrawFeatureVector(CommonCurrentPoseIndex);
 
 	if (!QueryVector.IsEmpty())
diff --git a/Engine/Plugins/Animation/PoseSearch/Source/Editor/Private/PoseSearchDebuggerDatabaseView.cpp b/Engine/Plugins/Animation/PoseSearch/Source/Editor/Private/PoseSearchDebuggerDatabaseView.cpp
--- a/Engine/Plugins/Animation/PoseSearch/Source/Editor/Private/PoseSearchDebuggerDatabaseView.cpp
+++ b/Engine/Plugins/Animation/PoseSearch/Source/Editor/Private/PoseSearchDebuggerDatabaseView.cpp
@@ -739,7 +739,14 @@ void SDebuggerDatabaseView::PopulateViewRows()
 
 				if (RowDataA.DbAssetIdx == RowDataB.DbAssetIdx)
 				{
-					return RowDataA.PoseCost < RowDataB.PoseCost;
+					const bool bIsRowDataAValid = EnumHasAnyFlags(RowDataA.PoseCandidateFlags, EPoseCandidateFlags::AnyValidMask);
+					const bool bIsRowDataBValid = EnumHasAnyFlags(RowDataB.PoseCandidateFlags, EPoseCandidateFlags::AnyValidMask);
+
+					if (bIsRowDataAValid == bIsRowDataBValid)
+					{
+						return RowDataA.PoseCost < RowDataB.PoseCost;
+					}
+					return bIsRowDataAValid;
 				}
 				return RowDataA.DbAssetIdx < RowDataB.DbAssetIdx;
 			});
diff --git a/Engine/Plugins/Animation/PoseSearch/Source/Editor/Private/PoseSearchDebuggerView.cpp b/Engine/Plugins/Animation/PoseSearch/Source/Editor/Private/PoseSearchDebuggerView.cpp
--- a/Engine/Plugins/Animation/PoseSearch/Source/Editor/Private/PoseSearchDebuggerView.cpp
+++ b/Engine/Plugins/Animation/PoseSearch/Source/Editor/Private/PoseSearchDebuggerView.cpp
@@ -106,11 +106,21 @@ FText GenerateSearchNameWithBestScore(const FTraceMotionMatchingStateMessage& Mo
 	{
 		for (const FTraceMotionMatchingStatePoseEntry& PoseEntry : DbEntry.PoseEntries)
 		{
-			BestCost = FMath::Min(PoseEntry.Cost, BestCost);
+			if (EnumHasAnyFlags(PoseEntry.PoseCandidateFlags, EPoseCandidateFlags::AnyValidMask))
+			{
+				BestCost = FMath::Min(PoseEntry.Cost, BestCost);
+			}
 		}
 	}
 
-	StringBuilder.Appendf(TEXT(" [%.3f]"), BestCost);
+	if (FMath::IsNearlyEqual(BestCost, UE_MAX_FLT))
+	{
+		StringBuilder.Appendf(TEXT(" [Invalid]"), BestCost);
+	}
+	else
+	{
+		StringBuilder.Appendf(TEXT(" [%.3f]"), BestCost);
+	}
 	return FText::FromString(StringBuilder.ToString());
 }
 
@@ -403,7 +413,7 @@ void SDebuggerView::Tick(const FGeometry& AllottedGeometry, const double InCurre
 						if (!DrawParams)
 						{
 							DrawParams = &CachedDebugDrawParams.Add(RowDatabase);
-							DrawParams->Init(AnimContexts, PoseHistories, RoleToIndex, RowDatabase);
+							DrawParams->Init(AnimContexts, PoseHistories, RoleToIndex, RowDatabase, FDebugDrawParams::EDrawContext::DrawCandidate);
 						}
 
 						DrawParams->DrawFeatureVector(Row->PoseIdx);
@@ -420,6 +430,9 @@ void SDebuggerView::Tick(const FGeometry& AllottedGeometry, const double InCurre
 		// Draw queries: either the one for the selected row databases or ALL of them
 		if (Model->GetDrawQuery())
 		{
+			// resetting CachedDebugDrawParams since we're now drawing the queries, that have a different drawing context
+			CachedDebugDrawParams.Reset();
+
 			for (const FTraceMotionMatchingStateDatabaseEntry& DbEntry : State->DatabaseEntries)
 			{
 				const UPoseSearchDatabase* Database = ResolveDatabaseFromId(DbEntry.DatabaseId);
@@ -432,7 +445,7 @@ void SDebuggerView::Tick(const FGeometry& AllottedGeometry, const double InCurre
 					if (!DrawParams)
 					{
 						DrawParams = &CachedDebugDrawParams.Add(Database);
-						DrawParams->Init(AnimContexts, PoseHistories, RoleToIndex, Database);
+						DrawParams->Init(AnimContexts, PoseHistories, RoleToIndex, Database, FDebugDrawParams::EDrawContext::DrawQuery);
 					}
 
 					DrawParams->DrawFeatureVector(DbEntry.QueryVector);
diff --git a/Engine/Plugins/Animation/PoseSearch/Source/Runtime/Private/AnimNode_MotionMatching.cpp b/Engine/Plugins/Animation/PoseSearch/Source/Runtime/Private/AnimNode_MotionMatching.cpp
--- a/Engine/Plugins/Animation/PoseSearch/Source/Runtime/Private/AnimNode_MotionMatching.cpp
+++ b/Engine/Plugins/Animation/PoseSearch/Source/Runtime/Private/AnimNode_MotionMatching.cpp
@@ -143,7 +143,7 @@ void FAnimNode_MotionMatching::UpdateAssetPlayer(const FAnimationUpdateContext&
 	const IPoseHistory* PoseHistory = nullptr;
 	if (const FPoseHistoryProvider* PoseHistoryProvider = Context.GetMessage<FPoseHistoryProvider>())
 	{
-		PoseHistory = &PoseHistoryProvider->GetPoseHistory();
+		PoseHistory = PoseHistoryProvider->GetPoseHistoryPtr();
 	}
 
 	if (!PoseHistory)
diff --git a/Engine/Plugins/Animation/PoseSearch/Source/Runtime/Private/AnimNode_MotionMatchingInteraction.cpp b/Engine/Plugins/Animation/PoseSearch/Source/Runtime/Private/AnimNode_MotionMatchingInteraction.cpp
--- a/Engine/Plugins/Animation/PoseSearch/Source/Runtime/Private/AnimNode_MotionMatchingInteraction.cpp
+++ b/Engine/Plugins/Animation/PoseSearch/Source/Runtime/Private/AnimNode_MotionMatchingInteraction.cpp
@@ -65,7 +65,7 @@ void FAnimNode_MotionMatchingInteraction::UpdateAssetPlayer(const FAnimationUpda
 	const float DeltaTime = Context.GetDeltaTime();
 	if (FPoseHistoryProvider* PoseHistoryProvider = Context.GetMessage<FPoseHistoryProvider>())
 	{
-		UPoseSearchInteractionLibrary::MotionMatchInteraction(CurrentResult, Availabilities, Context.AnimInstanceProxy->GetAnimInstanceObject(), FName(), &PoseHistoryProvider->GetPoseHistory(), bValidateResultAgainstAvailabilities);
+		UPoseSearchInteractionLibrary::MotionMatchInteraction(CurrentResult, Availabilities, Context.AnimInstanceProxy->GetAnimInstanceObject(), FName(), PoseHistoryProvider->GetPoseHistoryPtr(), bValidateResultAgainstAvailabilities);
 		check(CurrentResult.ActorRootTransforms.Num() == CurrentResult.ActorRootBoneTransforms.Num());
 
 		if (const UMultiAnimAsset* MultiAnimAsset = Cast<UMultiAnimAsset>(CurrentResult.SelectedAnim))
diff --git a/Engine/Plugins/Animation/PoseSearch/Source/Runtime/Private/AnimNode_PoseSearchHistoryCollector.cpp b/Engine/Plugins/Animation/PoseSearch/Source/Runtime/Private/AnimNode_PoseSearchHistoryCollector.cpp
--- a/Engine/Plugins/Animation/PoseSearch/Source/Runtime/Private/AnimNode_PoseSearchHistoryCollector.cpp
+++ b/Engine/Plugins/Animation/PoseSearch/Source/Runtime/Private/AnimNode_PoseSearchHistoryCollector.cpp
@@ -163,29 +163,32 @@ void FAnimNode_PoseSearchHistoryCollector::Evaluate_AnyThread(FPoseContext& Outp
 		RequiredBones = GetRequiredBones(Output.AnimInstanceProxy);
 	}
 
-	UE::PoseSearch::FComponentSpacePoseProvider ComponentSpacePoseProvider(ComponentSpacePose);
-	GetPoseHistory().EvaluateComponentSpace_AnyThread(Output.AnimInstanceProxy->GetDeltaSeconds(), ComponentSpacePoseProvider, bStoreScales,
-		RootBoneRecoveryTime, RootBoneTranslationRecoveryRatio, RootBoneRotationRecoveryRatio, bNeedsReset, bCacheBones, 
-		RequiredBones, Output.Curve, MakeConstArrayView(CollectedCurves));
-
-	// Fill in the pose history into a custom attribute for access downstream
-	if (FPoseHistoryAnimationAttribute* PoseHistoryAttribute = Output.CustomAttributes.FindOrAdd<FPoseHistoryAnimationAttribute>(UE::PoseSearch::PoseHistoryAttributeId))
+	if (FPoseHistory* PoseHistory = GetPoseHistoryPtr())
 	{
-		PoseHistoryAttribute->PoseHistory = &GetPoseHistory();
-		PoseHistoryAttribute->ScopeObject = Output.GetAnimInstanceObject();
-	}
+		UE::PoseSearch::FComponentSpacePoseProvider ComponentSpacePoseProvider(ComponentSpacePose);
+		PoseHistory->EvaluateComponentSpace_AnyThread(Output.AnimInstanceProxy->GetDeltaSeconds(), ComponentSpacePoseProvider, bStoreScales,
+			RootBoneRecoveryTime, RootBoneTranslationRecoveryRatio, RootBoneRotationRecoveryRatio, bNeedsReset, bCacheBones,
+			RequiredBones, Output.Curve, MakeConstArrayView(CollectedCurves));
 
-	bCacheBones = false;
+		// Fill in the pose history into a custom attribute for access downstream
+		if (FPoseHistoryAnimationAttribute* PoseHistoryAttribute = Output.CustomAttributes.FindOrAdd<FPoseHistoryAnimationAttribute>(UE::PoseSearch::PoseHistoryAttributeId))
+		{
+			PoseHistoryAttribute->PoseHistory = PoseHistory;
+			PoseHistoryAttribute->ScopeObject = Output.GetAnimInstanceObject();
+		}
 
 #if ENABLE_DRAW_DEBUG && ENABLE_ANIM_DEBUG
-	FColor Color;
+		FColor Color;
 #if WITH_EDITORONLY_DATA
-	Color = DebugColor.ToFColor(true);
+		Color = DebugColor.ToFColor(true);
 #else // WITH_EDITORONLY_DATA
-	Color = FLinearColor::Red.ToFColor(true);
+		Color = FColor::Red;
 #endif // WITH_EDITORONLY_DATA
-	GetPoseHistory().DebugDraw(*Output.AnimInstanceProxy, Color);
+		PoseHistory->DebugDraw(*Output.AnimInstanceProxy, Color);
 #endif // ENABLE_DRAW_DEBUG && ENABLE_ANIM_DEBUG
+	}
+
+	bCacheBones = false;
 }
 
 void FAnimNode_PoseSearchHistoryCollector::GatherDebugData(FNodeDebugData& DebugData)
@@ -233,29 +236,32 @@ void FAnimNode_PoseSearchComponentSpaceHistoryCollector::EvaluateComponentSpace_
 		RequiredBones = GetRequiredBones(Output.AnimInstanceProxy);
 	}
 
-	UE::PoseSearch::FComponentSpacePoseProvider ComponentSpacePoseProvider(Output.Pose);
-	GetPoseHistory().EvaluateComponentSpace_AnyThread(Output.AnimInstanceProxy->GetDeltaSeconds(), ComponentSpacePoseProvider, bStoreScales, 
-		RootBoneRecoveryTime, RootBoneTranslationRecoveryRatio, RootBoneRotationRecoveryRatio, bNeedsReset, bCacheBones,
-		RequiredBones, Output.Curve, MakeConstArrayView(CollectedCurves));
-
-	// Fill in the pose history into a custom attribute for access downstream
-	if (FPoseHistoryAnimationAttribute* PoseHistoryAttribute = Output.CustomAttributes.FindOrAdd<FPoseHistoryAnimationAttribute>(UE::PoseSearch::PoseHistoryAttributeId))
+	if (FPoseHistory* PoseHistory = GetPoseHistoryPtr())
 	{
-		PoseHistoryAttribute->PoseHistory = &GetPoseHistory();
-		PoseHistoryAttribute->ScopeObject = Output.GetAnimInstanceObject();
-	}
-	
-	bCacheBones = false;
+		UE::PoseSearch::FComponentSpacePoseProvider ComponentSpacePoseProvider(Output.Pose);
+		PoseHistory->EvaluateComponentSpace_AnyThread(Output.AnimInstanceProxy->GetDeltaSeconds(), ComponentSpacePoseProvider, bStoreScales,
+			RootBoneRecoveryTime, RootBoneTranslationRecoveryRatio, RootBoneRotationRecoveryRatio, bNeedsReset, bCacheBones,
+			RequiredBones, Output.Curve, MakeConstArrayView(CollectedCurves));
+
+		// Fill in the pose history into a custom attribute for access downstream
+		if (FPoseHistoryAnimationAttribute* PoseHistoryAttribute = Output.CustomAttributes.FindOrAdd<FPoseHistoryAnimationAttribute>(UE::PoseSearch::PoseHistoryAttributeId))
+		{
+			PoseHistoryAttribute->PoseHistory = PoseHistory;
+			PoseHistoryAttribute->ScopeObject = Output.GetAnimInstanceObject();
+		}
 
 #if ENABLE_DRAW_DEBUG && ENABLE_ANIM_DEBUG
-	FColor Color;
+		FColor Color;
 #if WITH_EDITORONLY_DATA
-	Color = DebugColor.ToFColor(true);
+		Color = DebugColor.ToFColor(true);
 #else // WITH_EDITORONLY_DATA
-	Color = FLinearColor::Red.ToFColor(true);
+		Color = FColor::Red;
 #endif // WITH_EDITORONLY_DATA
-	GetPoseHistory().DebugDraw(*Output.AnimInstanceProxy, Color);
+		PoseHistory->DebugDraw(*Output.AnimInstanceProxy, Color);
 #endif // ENABLE_DRAW_DEBUG && ENABLE_ANIM_DEBUG
+	}
+
+	bCacheBones = false;
 }
 
 void FAnimNode_PoseSearchHistoryCollector::Update_AnyThread(const FAnimationUpdateContext& Context)
diff --git a/Engine/Plugins/Animation/PoseSearch/Source/Runtime/Private/Chooser/PoseSearchChooserColumn.cpp b/Engine/Plugins/Animation/PoseSearch/Source/Runtime/Private/Chooser/PoseSearchChooserColumn.cpp
--- a/Engine/Plugins/Animation/PoseSearch/Source/Runtime/Private/Chooser/PoseSearchChooserColumn.cpp
+++ b/Engine/Plugins/Animation/PoseSearch/Source/Runtime/Private/Chooser/PoseSearchChooserColumn.cpp
@@ -442,7 +442,7 @@ void FPoseSearchColumn::Filter(FChooserEvaluationContext& Context, const FChoose
 	{
 		if (const FPoseHistoryProvider* PoseHistoryProvider = PoseSearchColumnScratchArea->ChooserPlayerSettings->AnimationUpdateContext->GetMessage<FPoseHistoryProvider>())
 		{
-			PoseHistory = &PoseHistoryProvider->GetPoseHistory();
+			PoseHistory = PoseHistoryProvider->GetPoseHistoryPtr();
 		}
 	}
 
diff --git a/Engine/Plugins/Animation/PoseSearch/Source/Runtime/Private/PoseHistoryProvider.cpp b/Engine/Plugins/Animation/PoseSearch/Source/Runtime/Private/PoseHistoryProvider.cpp
--- a/Engine/Plugins/Animation/PoseSearch/Source/Runtime/Private/PoseHistoryProvider.cpp
+++ b/Engine/Plugins/Animation/PoseSearch/Source/Runtime/Private/PoseHistoryProvider.cpp
@@ -9,7 +9,15 @@ namespace UE::PoseSearch
 const IPoseHistory& FPoseHistoryProvider::GetPoseHistory() const
 {
 	check(HistoryCollector);
+	PRAGMA_DISABLE_DEPRECATION_WARNINGS
 	return HistoryCollector->GetPoseHistory();
+	PRAGMA_ENABLE_DEPRECATION_WARNINGS
+}
+
+const IPoseHistory* FPoseHistoryProvider::GetPoseHistoryPtr() const
+{
+	check(HistoryCollector);
+	return HistoryCollector->GetPoseHistoryPtr();
 }
 
 } // namespace UE::PoseSearch
\ No newline at end of file
diff --git a/Engine/Plugins/Animation/PoseSearch/Source/Runtime/Private/PoseSearchContext.cpp b/Engine/Plugins/Animation/PoseSearch/Source/Runtime/Private/PoseSearchContext.cpp
--- a/Engine/Plugins/Animation/PoseSearch/Source/Runtime/Private/PoseSearchContext.cpp
+++ b/Engine/Plugins/Animation/PoseSearch/Source/Runtime/Private/PoseSearchContext.cpp
@@ -189,10 +189,10 @@ static const UWorld* GetWorld(TConstArrayView<FChooserEvaluationContext*> AnimCo
 
 //////////////////////////////////////////////////////////////////////////
 // FDebugDrawParams
-FDebugDrawParams::FDebugDrawParams(TConstArrayView<FChooserEvaluationContext*> InAnimContexts, TConstArrayView<const IPoseHistory*> InPoseHistories, const FRoleToIndex& InRoleToIndex, const UPoseSearchDatabase* InDatabase)
+FDebugDrawParams::FDebugDrawParams(TConstArrayView<FChooserEvaluationContext*> InAnimContexts, TConstArrayView<const IPoseHistory*> InPoseHistories, const FRoleToIndex& InRoleToIndex, const UPoseSearchDatabase* InDatabase, EDrawContext InDrawContext)
 : FDebugDrawParams()
 {
-	Init(InAnimContexts, InPoseHistories, InRoleToIndex, InDatabase);
+	Init(InAnimContexts, InPoseHistories, InRoleToIndex, InDatabase, InDrawContext);
 }
 
 FDebugDrawParams::FDebugDrawParams(const FDebugDrawParams& Other)
@@ -200,6 +200,7 @@ FDebugDrawParams::FDebugDrawParams(const FDebugDrawParams& Other)
 , PoseHistories(Other.PoseHistories)
 , RoleToIndex(Other.RoleToIndex)
 , Database(Other.Database)
+, DrawContext(Other.DrawContext)
 , DynamicWeightsSqrtBuffer(Other.DynamicWeightsSqrtBuffer)
 {
 	// pointing DynamicWeightsSqrt this object internal DynamicWeightsSqrtBuffer memory buffer!
@@ -211,6 +212,7 @@ FDebugDrawParams::FDebugDrawParams(FDebugDrawParams&& Other)
 , PoseHistories(MoveTemp(Other.PoseHistories))
 , RoleToIndex(MoveTemp(Other.RoleToIndex))
 , Database(Other.Database)
+, DrawContext(Other.DrawContext)
 , DynamicWeightsSqrtBuffer(MoveTemp(Other.DynamicWeightsSqrtBuffer))
 {
 	// pointing DynamicWeightsSqrt this object internal DynamicWeightsSqrtBuffer memory buffer!
@@ -239,12 +241,13 @@ FDebugDrawParams& FDebugDrawParams::operator=(FDebugDrawParams&& Other)
 	return *this;
 }
 
-void FDebugDrawParams::Init(TConstArrayView<FChooserEvaluationContext*> InAnimContexts, TConstArrayView<const IPoseHistory*> InPoseHistories, const FRoleToIndex& InRoleToIndex, const UPoseSearchDatabase* InDatabase)
+void FDebugDrawParams::Init(TConstArrayView<FChooserEvaluationContext*> InAnimContexts, TConstArrayView<const IPoseHistory*> InPoseHistories, const FRoleToIndex& InRoleToIndex, const UPoseSearchDatabase* InDatabase, EDrawContext InDrawContext)
 {
 	AnimContexts = InAnimContexts;
 	PoseHistories = InPoseHistories;
 	RoleToIndex = InRoleToIndex;
 	Database = InDatabase;
+	DrawContext = InDrawContext;
 
 	check(RoleToIndex.Num() == PoseHistories.Num());
 	check(IsValid(RoleToIndex));
@@ -573,13 +576,13 @@ void FDebugDrawParams::DrawCircle(const FMatrix& TransformMatrix, float Radius,
 	}
 }
 
-void FDebugDrawParams::DrawWedge(const FVector& Origin, const FVector& Direction, float InnerRadius, float OuterRadius, float Width, int32 Segments, const FColor& Color, float Thickness) const
+void FDebugDrawParams::DrawWedge(const FVector& Origin, const FVector& Direction, float InnerRadius, float OuterRadius, float WidthInDegrees, int32 Segments, const FColor& Color, float Thickness) const
 {
 	FVector NormalizedDirection = Direction;
 	if (Color.A > 0 && Segments > 1 && NormalizedDirection.Normalize())
 	{
-		float AngleDeg = -Width * 0.5f;
-		const float AngleDegIncrement = Width / (Segments - 1);
+		float AngleDeg = -WidthInDegrees * 0.5f;
+		const float AngleDegIncrement = WidthInDegrees / (Segments - 1);
 
 		FVector PrevDirection = NormalizedDirection.RotateAngleAxis(AngleDeg, FVector::ZAxisVector);
 		DrawLine(PrevDirection * OuterRadius + Origin, PrevDirection * InnerRadius + Origin, Color);
diff --git a/Engine/Plugins/Animation/PoseSearch/Source/Runtime/Private/PoseSearchDatabase.cpp b/Engine/Plugins/Animation/PoseSearch/Source/Runtime/Private/PoseSearchDatabase.cpp
--- a/Engine/Plugins/Animation/PoseSearch/Source/Runtime/Private/PoseSearchDatabase.cpp
+++ b/Engine/Plugins/Animation/PoseSearch/Source/Runtime/Private/PoseSearchDatabase.cpp
@@ -147,138 +147,10 @@ FPoseSearchDatabaseAnimationAsset UpgradeToDatabaseAnimationAsset(FInstancedStru
 	return PoseSearchDatabaseAnimationAsset;
 }
 
-struct FSearchFilters
-{
-	FSearchFilters(const UPoseSearchSchema* Schema, TConstArrayView<int32> NonSelectableIdx, TConstArrayView<int32> SelectableAssetIdx, bool bAddBlockTransitionFilter)
-	{
-		if (bAddBlockTransitionFilter)
-		{
-			Filters.Add(&BlockTransitionFilter);
-		}
-
-		if (NonSelectableIdxFilter.Init(NonSelectableIdx).IsFilterActive())
-		{
-			Filters.Add(&NonSelectableIdxFilter);
-		}
-
-		if (SelectableAssetIdxFilter.Init(SelectableAssetIdx).IsFilterActive())
-		{
-			Filters.Add(&SelectableAssetIdxFilter);
-		}
-
-		for (const IPoseSearchFilter* Filter : Schema->GetChannels())
-		{
-			if (Filter->IsFilterActive())
-			{
-				Filters.Add(Filter);
-			}
-		}
-	}
-
-	// @todo: template this with bAlignedAndPadded to be able to use faster ComparePoses
-	bool AreFiltersValid(const FSearchIndex& SearchIndex, TConstArrayView<float> PoseValues, TConstArrayView<float> QueryValues, TConstArrayView<float> DynamicWeightsSqrt, int32 PoseIdx
-#if UE_POSE_SEARCH_TRACE_ENABLED
-		, float ContinuingPoseCostAddend, float ContinuingInteractionCostAddend, UE::PoseSearch::FSearchContext& SearchContext, const UPoseSearchDatabase* Database
-#endif // UE_POSE_SEARCH_TRACE_ENABLED
-	) const
-	{
-		for (const IPoseSearchFilter* Filter : Filters)
-		{
-			if (!Filter->IsFilterValid(PoseValues, QueryValues, PoseIdx, SearchIndex.PoseMetadata[PoseIdx]))
-			{
-#if UE_POSE_SEARCH_TRACE_ENABLED
-				if (Filter == &NonSelectableIdxFilter)
-				{
-					// candidate already added to SearchContext.BestCandidates by PopulateNonSelectableIdx
-				}
-				else if (Filter == &SelectableAssetIdxFilter)
-				{
-					SearchContext.Track(Database, PoseIdx, EPoseCandidateFlags::DiscardedBy_AssetIdxFilter, DynamicWeightsSqrt, ContinuingPoseCostAddend, ContinuingInteractionCostAddend);
-				}
-				else if (Filter == &BlockTransitionFilter)
-				{
-					SearchContext.Track(Database, PoseIdx, EPoseCandidateFlags::DiscardedBy_BlockTransition, DynamicWeightsSqrt, ContinuingPoseCostAddend, ContinuingInteractionCostAddend);
-				}
-				else
-				{
-					SearchContext.Track(Database, PoseIdx, EPoseCandidateFlags::DiscardedBy_PoseFilter, DynamicWeightsSqrt, ContinuingPoseCostAddend, ContinuingInteractionCostAddend);
-				}
-#endif // UE_POSE_SEARCH_TRACE_ENABLED
-				return false;
-			}
-		}
-		return true;
-	};
-
-private:
-	struct FNonSelectableIdxFilter : public IPoseSearchFilter
-	{
-		const FNonSelectableIdxFilter& Init(TConstArrayView<int32> InNonSelectableIdx)
-		{
-			check(Algo::IsSorted(InNonSelectableIdx));
-			NonSelectableIdx = InNonSelectableIdx;
-			return *this;
-		}
-
-		virtual bool IsFilterActive() const override
-		{
-			return !NonSelectableIdx.IsEmpty();
-		}
-
-		virtual bool IsFilterValid(TConstArrayView<float> PoseValues, TConstArrayView<float> QueryValues, int32 PoseIdx, const FPoseMetadata& Metadata) const override
-		{
-			return Algo::BinarySearch(NonSelectableIdx, PoseIdx) == INDEX_NONE;
-		}
-
-		TConstArrayView<int32> NonSelectableIdx;
-	};
-
-	struct FSelectableAssetIdxFilter : public IPoseSearchFilter
-	{
-		const FSelectableAssetIdxFilter& Init(TConstArrayView<int32> InSelectableAssetIdxFilter)
-		{
-			check(Algo::IsSorted(InSelectableAssetIdxFilter));
-			SelectableAssetIdxFilter = InSelectableAssetIdxFilter;
-			return *this;
-		}
-
-		virtual bool IsFilterActive() const override
-		{
-			return !SelectableAssetIdxFilter.IsEmpty();
-		}
-
-		virtual bool IsFilterValid(TConstArrayView<float> PoseValues, TConstArrayView<float> QueryValues, int32 PoseIdx, const FPoseMetadata& Metadata) const override
-		{
-			return Algo::BinarySearch(SelectableAssetIdxFilter, int32(Metadata.GetAssetIndex())) != INDEX_NONE;
-		}
-
-		TConstArrayView<int32> SelectableAssetIdxFilter;
-	};
-
-	struct FBlockTransitionFilter : public IPoseSearchFilter
-	{
-		virtual bool IsFilterActive() const override
-		{
-			return true;
-		}
-
-		virtual bool IsFilterValid(TConstArrayView<float> PoseValues, TConstArrayView<float> QueryValues, int32 PoseIdx, const FPoseMetadata& Metadata) const override
-		{
-			return !Metadata.IsBlockTransition();
-		}
-	};
-
-	FNonSelectableIdxFilter NonSelectableIdxFilter;
-	FSelectableAssetIdxFilter SelectableAssetIdxFilter;
-	FBlockTransitionFilter BlockTransitionFilter;
-
-	TArray<const IPoseSearchFilter*, TInlineAllocator<64, TMemStackAllocator<>>> Filters;
-};
-
 template<bool bReconstructPoseValues, bool bAlignedAndPadded>
 static inline void EvaluatePoseKernel(FSearchResults& SearchResults, const FSearchIndex& SearchIndex, TConstArrayView<float> QueryValues, TArrayView<float> ReconstructedPoseValuesBuffer,
 	int32 PoseIdx, int32 EventPoseIdx, const FSearchFilters& SearchFilters, float ContinuingPoseCostAddend, float ContinuingInteractionCostAddend,
-	FSearchContext& SearchContext, const UPoseSearchDatabase* Database, TConstArrayView<float> DynamicWeightsSqrt
+	const UPoseSearchDatabase* Database, TConstArrayView<float> DynamicWeightsSqrt
 #if UE_POSE_SEARCH_TRACE_ENABLED
 	, bool bTrackResult = true
 #endif // UE_POSE_SEARCH_TRACE_ENABLED
@@ -293,7 +165,7 @@ static inline void EvaluatePoseKernel(FSearchResults& SearchResults, const FSear
 
 	if (SearchFilters.AreFiltersValid(SearchIndex, PoseValues, QueryValues, DynamicWeightsSqrt, PoseIdx
 #if UE_POSE_SEARCH_TRACE_ENABLED
-		, ContinuingPoseCostAddend, ContinuingInteractionCostAddend, SearchContext, Database
+		, ContinuingPoseCostAddend, ContinuingInteractionCostAddend, Database
 #endif // UE_POSE_SEARCH_TRACE_ENABLED
 	))
 	{
@@ -314,7 +186,7 @@ static inline void EvaluatePoseKernel(FSearchResults& SearchResults, const FSear
 		if (bTrackResult)
 		{
 			// @todo: add tracking for EventPoseIdx
-			SearchContext.Track(Database, PoseIdx, EPoseCandidateFlags::Valid_Pose, Result.PoseCost);
+			SearchFilters.GetSearchContext().Track(Database, PoseIdx, EPoseCandidateFlags::Valid_Pose, Result.PoseCost);
 		}
 #endif // UE_POSE_SEARCH_TRACE_ENABLED
 	}
@@ -1971,7 +1843,7 @@ void UPoseSearchDatabase::SearchPCAKDTree(UE::PoseSearch::FSearchContext& Search
 
 		// NonSelectableIdx are already filtered out inside the kdtree search.
 		// Also kdtrees don't contain block transition poses by construction, so FSearchFilters input bAddBlockTransitionFilter can be set to false
-		const FSearchFilters SearchFilters(Schema, bRunNonSelectableIdxPostKDTree ? NonSelectableIdx : TConstArrayView<int32>(), SelectableAssetIdx, false);
+		const FSearchFilters SearchFilters(Schema, bRunNonSelectableIdxPostKDTree ? NonSelectableIdx : TConstArrayView<int32>(), SelectableAssetIdx, false, QueryValues, SearchContext);
 		
 		// are the PCAValues pruned out of duplicates (multiple poses are associated with the same PCAValuesVectorIdx)
 		if (bArePCAValuesPruned)
@@ -1992,7 +1864,7 @@ void UPoseSearchDatabase::SearchPCAKDTree(UE::PoseSearch::FSearchContext& Search
 						const TConstArrayView<int32> PoseIndexes = SearchIndex.PCAValuesVectorToPoseIndexes[Results[ResultIndex].Index];
 						for (int32 Index = 0; Index < PoseIndexes.Num() && NumEvaluatePoseKernelCalls < MaxNumEvaluatePoseKernelCalls; ++Index, ++NumEvaluatePoseKernelCalls)
 						{
-							EvaluatePoseKernel<true, true>(SearchResults, SearchIndex, QueryValues, ReconstructedPoseValuesBuffer, PoseIndexes[Index], INDEX_NONE, SearchFilters, ContinuingPoseCostAddend, ContinuingInteractionCostAddend, SearchContext, this, DynamicWeightsSqrt
+							EvaluatePoseKernel<true, true>(SearchResults, SearchIndex, QueryValues, ReconstructedPoseValuesBuffer, PoseIndexes[Index], INDEX_NONE, SearchFilters, ContinuingPoseCostAddend, ContinuingInteractionCostAddend, this, DynamicWeightsSqrt
 #if UE_POSE_SEARCH_TRACE_ENABLED
 								, true
 #endif // UE_POSE_SEARCH_TRACE_ENABLED
@@ -2011,7 +1883,7 @@ void UPoseSearchDatabase::SearchPCAKDTree(UE::PoseSearch::FSearchContext& Search
 						const TConstArrayView<int32> PoseIndexes = SearchIndex.PCAValuesVectorToPoseIndexes[Results[ResultIndex].Index];
 						for (int32 Index = 0; Index < PoseIndexes.Num() && NumEvaluatePoseKernelCalls < MaxNumEvaluatePoseKernelCalls; ++Index, ++NumEvaluatePoseKernelCalls)
 						{
-							EvaluatePoseKernel<true, false>(SearchResults, SearchIndex, QueryValues, ReconstructedPoseValuesBuffer, PoseIndexes[Index], INDEX_NONE, SearchFilters, ContinuingPoseCostAddend, ContinuingInteractionCostAddend, SearchContext, this, DynamicWeightsSqrt
+							EvaluatePoseKernel<true, false>(SearchResults, SearchIndex, QueryValues, ReconstructedPoseValuesBuffer, PoseIndexes[Index], INDEX_NONE, SearchFilters, ContinuingPoseCostAddend, ContinuingInteractionCostAddend, this, DynamicWeightsSqrt
 #if UE_POSE_SEARCH_TRACE_ENABLED
 								, true
 #endif // UE_POSE_SEARCH_TRACE_ENABLED
@@ -2034,7 +1906,7 @@ void UPoseSearchDatabase::SearchPCAKDTree(UE::PoseSearch::FSearchContext& Search
 						const TConstArrayView<int32> PoseIndexes = SearchIndex.PCAValuesVectorToPoseIndexes[Results[ResultIndex].Index];
 						for (int32 Index = 0; Index < PoseIndexes.Num() && NumEvaluatePoseKernelCalls < MaxNumEvaluatePoseKernelCalls; ++Index, ++NumEvaluatePoseKernelCalls)
 						{
-							EvaluatePoseKernel<false, true>(SearchResults, SearchIndex, QueryValues, TArrayView<float>(), PoseIndexes[Index], INDEX_NONE, SearchFilters, ContinuingPoseCostAddend, ContinuingInteractionCostAddend, SearchContext, this, DynamicWeightsSqrt
+							EvaluatePoseKernel<false, true>(SearchResults, SearchIndex, QueryValues, TArrayView<float>(), PoseIndexes[Index], INDEX_NONE, SearchFilters, ContinuingPoseCostAddend, ContinuingInteractionCostAddend, this, DynamicWeightsSqrt
 #if UE_POSE_SEARCH_TRACE_ENABLED
 								, true
 #endif // UE_POSE_SEARCH_TRACE_ENABLED
@@ -2053,7 +1925,7 @@ void UPoseSearchDatabase::SearchPCAKDTree(UE::PoseSearch::FSearchContext& Search
 						const TConstArrayView<int32> PoseIndexes = SearchIndex.PCAValuesVectorToPoseIndexes[Results[ResultIndex].Index];
 						for (int32 Index = 0; Index < PoseIndexes.Num() && NumEvaluatePoseKernelCalls < MaxNumEvaluatePoseKernelCalls; ++Index, ++NumEvaluatePoseKernelCalls)
 						{
-							EvaluatePoseKernel<false, false>(SearchResults, SearchIndex, QueryValues, TArrayView<float>(), PoseIndexes[Index], INDEX_NONE, SearchFilters, ContinuingPoseCostAddend, ContinuingInteractionCostAddend, SearchContext, this, DynamicWeightsSqrt
+							EvaluatePoseKernel<false, false>(SearchResults, SearchIndex, QueryValues, TArrayView<float>(), PoseIndexes[Index], INDEX_NONE, SearchFilters, ContinuingPoseCostAddend, ContinuingInteractionCostAddend, this, DynamicWeightsSqrt
 #if UE_POSE_SEARCH_TRACE_ENABLED
 								, true
 #endif // UE_POSE_SEARCH_TRACE_ENABLED
@@ -2078,7 +1950,7 @@ void UPoseSearchDatabase::SearchPCAKDTree(UE::PoseSearch::FSearchContext& Search
 			{
 				for (int32 ResultIndex = 0; ResultIndex < NumResults; ++ResultIndex)
 				{
-					EvaluatePoseKernel<true, true>(SearchResults, SearchIndex, QueryValues, ReconstructedPoseValuesBuffer, Results[ResultIndex].Index, INDEX_NONE, SearchFilters, ContinuingPoseCostAddend, ContinuingInteractionCostAddend, SearchContext, this, DynamicWeightsSqrt
+					EvaluatePoseKernel<true, true>(SearchResults, SearchIndex, QueryValues, ReconstructedPoseValuesBuffer, Results[ResultIndex].Index, INDEX_NONE, SearchFilters, ContinuingPoseCostAddend, ContinuingInteractionCostAddend, this, DynamicWeightsSqrt
 	#if UE_POSE_SEARCH_TRACE_ENABLED
 						, true
 	#endif // UE_POSE_SEARCH_TRACE_ENABLED
@@ -2092,7 +1964,7 @@ void UPoseSearchDatabase::SearchPCAKDTree(UE::PoseSearch::FSearchContext& Search
 			{
 				for (int32 ResultIndex = 0; ResultIndex < NumResults; ++ResultIndex)
 				{
-					EvaluatePoseKernel<true, false>(SearchResults, SearchIndex, QueryValues, ReconstructedPoseValuesBuffer, Results[ResultIndex].Index, INDEX_NONE, SearchFilters, ContinuingPoseCostAddend, ContinuingInteractionCostAddend, SearchContext, this, DynamicWeightsSqrt
+					EvaluatePoseKernel<true, false>(SearchResults, SearchIndex, QueryValues, ReconstructedPoseValuesBuffer, Results[ResultIndex].Index, INDEX_NONE, SearchFilters, ContinuingPoseCostAddend, ContinuingInteractionCostAddend, this, DynamicWeightsSqrt
 #if UE_POSE_SEARCH_TRACE_ENABLED
 						, true
 #endif // UE_POSE_SEARCH_TRACE_ENABLED
@@ -2110,7 +1982,7 @@ void UPoseSearchDatabase::SearchPCAKDTree(UE::PoseSearch::FSearchContext& Search
 			{
 				for (int32 ResultIndex = 0; ResultIndex < NumResults; ++ResultIndex)
 				{
-					EvaluatePoseKernel<false, true>(SearchResults, SearchIndex, QueryValues, TArrayView<float>(), Results[ResultIndex].Index, INDEX_NONE, SearchFilters, ContinuingPoseCostAddend, ContinuingInteractionCostAddend, SearchContext, this, DynamicWeightsSqrt
+					EvaluatePoseKernel<false, true>(SearchResults, SearchIndex, QueryValues, TArrayView<float>(), Results[ResultIndex].Index, INDEX_NONE, SearchFilters, ContinuingPoseCostAddend, ContinuingInteractionCostAddend, this, DynamicWeightsSqrt
 #if UE_POSE_SEARCH_TRACE_ENABLED
 						, true
 #endif // UE_POSE_SEARCH_TRACE_ENABLED
@@ -2125,7 +1997,7 @@ void UPoseSearchDatabase::SearchPCAKDTree(UE::PoseSearch::FSearchContext& Search
 			{
 				for (int32 ResultIndex = 0; ResultIndex < NumResults; ++ResultIndex)
 				{
-					EvaluatePoseKernel<false, false>(SearchResults, SearchIndex, QueryValues, TArrayView<float>(), Results[ResultIndex].Index, INDEX_NONE, SearchFilters, ContinuingPoseCostAddend, ContinuingInteractionCostAddend, SearchContext, this, DynamicWeightsSqrt
+					EvaluatePoseKernel<false, false>(SearchResults, SearchIndex, QueryValues, TArrayView<float>(), Results[ResultIndex].Index, INDEX_NONE, SearchFilters, ContinuingPoseCostAddend, ContinuingInteractionCostAddend, this, DynamicWeightsSqrt
 #if UE_POSE_SEARCH_TRACE_ENABLED
 						, true
 #endif // UE_POSE_SEARCH_TRACE_ENABLED
@@ -2210,7 +2082,7 @@ void UPoseSearchDatabase::SearchVPTree(UE::PoseSearch::FSearchContext& SearchCon
 #endif // UE_POSE_SEARCH_TRACE_ENABLED
 		);
 		
-		const FSearchFilters SearchFilters(Schema, NonSelectableIdx, SelectableAssetIdx, SearchIndex.bAnyBlockTransition);
+		const FSearchFilters SearchFilters(Schema, NonSelectableIdx, SelectableAssetIdx, SearchIndex.bAnyBlockTransition, QueryValues, SearchContext);
 
 		// @todo: implement a FVPTreeDataSource for aligned and padded features vector like CompareAlignedPoses does 
 		FVPTreeDataSource DataSource(SearchIndex);
@@ -2232,7 +2104,7 @@ void UPoseSearchDatabase::SearchVPTree(UE::PoseSearch::FSearchContext& SearchCon
 				const TConstArrayView<int32> PoseIndexes = SearchIndex.ValuesVectorToPoseIndexes[IndexDistance.Index];
 				for (int32 Index = 0; Index < PoseIndexes.Num() && NumEvaluatePoseKernelCalls < MaxNumEvaluatePoseKernelCalls; ++Index, ++NumEvaluatePoseKernelCalls)
 				{
-					EvaluatePoseKernel<false, false>(SearchResults, SearchIndex, QueryValues, TArrayView<float>(), PoseIndexes[Index], INDEX_NONE, SearchFilters, ContinuingPoseCostAddend, ContinuingInteractionCostAddend, SearchContext, this, DynamicWeightsSqrt
+					EvaluatePoseKernel<false, false>(SearchResults, SearchIndex, QueryValues, TArrayView<float>(), PoseIndexes[Index], INDEX_NONE, SearchFilters, ContinuingPoseCostAddend, ContinuingInteractionCostAddend, this, DynamicWeightsSqrt
 #if UE_POSE_SEARCH_TRACE_ENABLED
 						, true
 #endif // UE_POSE_SEARCH_TRACE_ENABLED
@@ -2250,7 +2122,7 @@ void UPoseSearchDatabase::SearchVPTree(UE::PoseSearch::FSearchContext& SearchCon
 				const FIndexDistance& IndexDistance = UnsortedResults[ResultIndex];
 
 				// @todo: IndexDistance.Distance is the Sqrt(DissimilarityCost), so there's no need to calculate it again in SearchIndex.ComparePoses
-				EvaluatePoseKernel<false, false>(SearchResults, SearchIndex, QueryValues, TArrayView<float>(), IndexDistance.Index, INDEX_NONE, SearchFilters, ContinuingPoseCostAddend, ContinuingInteractionCostAddend, SearchContext, this, DynamicWeightsSqrt
+				EvaluatePoseKernel<false, false>(SearchResults, SearchIndex, QueryValues, TArrayView<float>(), IndexDistance.Index, INDEX_NONE, SearchFilters, ContinuingPoseCostAddend, ContinuingInteractionCostAddend, this, DynamicWeightsSqrt
 #if UE_POSE_SEARCH_TRACE_ENABLED
 					, true
 #endif // UE_POSE_SEARCH_TRACE_ENABLED
@@ -2324,7 +2196,7 @@ void UPoseSearchDatabase::SearchBruteForce(UE::PoseSearch::FSearchContext& Searc
 #endif // UE_POSE_SEARCH_TRACE_ENABLED
 		);
 
-		const FSearchFilters SearchFilters(Schema, NonSelectableIdx, FSelectableAssetIdx(), SearchIndex.bAnyBlockTransition);
+		const FSearchFilters SearchFilters(Schema, NonSelectableIdx, FSelectableAssetIdx(), SearchIndex.bAnyBlockTransition, QueryValues, SearchContext);
 
 		if (SelectableAssetIdx.IsEmpty())
 		{
@@ -2336,7 +2208,7 @@ void UPoseSearchDatabase::SearchBruteForce(UE::PoseSearch::FSearchContext& Searc
 				check(IsAligned(ReconstructedPoseValuesBuffer.GetData(), alignof(VectorRegister4Float)));
 				for (int32 PoseIdx = 0; PoseIdx < SearchIndex.GetNumPoses(); ++PoseIdx)
 				{
-					EvaluatePoseKernel<true, false>(SearchResults, SearchIndex, QueryValues, ReconstructedPoseValuesBuffer, PoseIdx, INDEX_NONE, SearchFilters, ContinuingPoseCostAddend, ContinuingInteractionCostAddend, SearchContext, this, DynamicWeightsSqrt
+					EvaluatePoseKernel<true, false>(SearchResults, SearchIndex, QueryValues, ReconstructedPoseValuesBuffer, PoseIdx, INDEX_NONE, SearchFilters, ContinuingPoseCostAddend, ContinuingInteractionCostAddend, this, DynamicWeightsSqrt
 #if UE_POSE_SEARCH_TRACE_ENABLED
 						, bTrackResults
 #endif // UE_POSE_SEARCH_TRACE_ENABLED
@@ -2351,7 +2223,7 @@ void UPoseSearchDatabase::SearchBruteForce(UE::PoseSearch::FSearchContext& Searc
 			{
 				for (int32 PoseIdx = 0; PoseIdx < SearchIndex.GetNumPoses(); ++PoseIdx)
 				{
-					EvaluatePoseKernel<false, true>(SearchResults, SearchIndex, QueryValues, TArrayView<float>(), PoseIdx, INDEX_NONE, SearchFilters, ContinuingPoseCostAddend, ContinuingInteractionCostAddend, SearchContext, this, DynamicWeightsSqrt
+					EvaluatePoseKernel<false, true>(SearchResults, SearchIndex, QueryValues, TArrayView<float>(), PoseIdx, INDEX_NONE, SearchFilters, ContinuingPoseCostAddend, ContinuingInteractionCostAddend, this, DynamicWeightsSqrt
 #if UE_POSE_SEARCH_TRACE_ENABLED
 						, bTrackResults
 #endif // UE_POSE_SEARCH_TRACE_ENABLED
@@ -2366,7 +2238,7 @@ void UPoseSearchDatabase::SearchBruteForce(UE::PoseSearch::FSearchContext& Searc
 			{
 				for (int32 PoseIdx = 0; PoseIdx < SearchIndex.GetNumPoses(); ++PoseIdx)
 				{
-					EvaluatePoseKernel<false, false>(SearchResults, SearchIndex, QueryValues, TArrayView<float>(), PoseIdx, INDEX_NONE, SearchFilters, ContinuingPoseCostAddend, ContinuingInteractionCostAddend, SearchContext, this, DynamicWeightsSqrt
+					EvaluatePoseKernel<false, false>(SearchResults, SearchIndex, QueryValues, TArrayView<float>(), PoseIdx, INDEX_NONE, SearchFilters, ContinuingPoseCostAddend, ContinuingInteractionCostAddend, this, DynamicWeightsSqrt
 #if UE_POSE_SEARCH_TRACE_ENABLED
 						, bTrackResults
 #endif // UE_POSE_SEARCH_TRACE_ENABLED
@@ -2395,7 +2267,7 @@ void UPoseSearchDatabase::SearchBruteForce(UE::PoseSearch::FSearchContext& Searc
 					const int32 LastPoseIdx = FirstPoseIdx + SearchIndexAsset.GetNumPoses();
 					for (int32 PoseIdx = FirstPoseIdx; PoseIdx < LastPoseIdx; ++PoseIdx)
 					{
-						EvaluatePoseKernel<true, false>(SearchResults, SearchIndex, QueryValues, ReconstructedPoseValuesBuffer, PoseIdx, INDEX_NONE, SearchFilters, ContinuingPoseCostAddend, ContinuingInteractionCostAddend, SearchContext, this, DynamicWeightsSqrt
+						EvaluatePoseKernel<true, false>(SearchResults, SearchIndex, QueryValues, ReconstructedPoseValuesBuffer, PoseIdx, INDEX_NONE, SearchFilters, ContinuingPoseCostAddend, ContinuingInteractionCostAddend, this, DynamicWeightsSqrt
 #if UE_POSE_SEARCH_TRACE_ENABLED
 							, bTrackResults
 #endif // UE_POSE_SEARCH_TRACE_ENABLED
@@ -2416,7 +2288,7 @@ void UPoseSearchDatabase::SearchBruteForce(UE::PoseSearch::FSearchContext& Searc
 					const int32 LastPoseIdx = FirstPoseIdx + SearchIndexAsset.GetNumPoses();
 					for (int32 PoseIdx = FirstPoseIdx; PoseIdx < LastPoseIdx; ++PoseIdx)
 					{
-						EvaluatePoseKernel<false, true>(SearchResults, SearchIndex, QueryValues, TArrayView<float>(), PoseIdx, INDEX_NONE, SearchFilters, ContinuingPoseCostAddend, ContinuingInteractionCostAddend, SearchContext, this, DynamicWeightsSqrt
+						EvaluatePoseKernel<false, true>(SearchResults, SearchIndex, QueryValues, TArrayView<float>(), PoseIdx, INDEX_NONE, SearchFilters, ContinuingPoseCostAddend, ContinuingInteractionCostAddend, this, DynamicWeightsSqrt
 #if UE_POSE_SEARCH_TRACE_ENABLED
 							, bTrackResults
 #endif // UE_POSE_SEARCH_TRACE_ENABLED
@@ -2437,7 +2309,7 @@ void UPoseSearchDatabase::SearchBruteForce(UE::PoseSearch::FSearchContext& Searc
 					const int32 LastPoseIdx = FirstPoseIdx + SearchIndexAsset.GetNumPoses();
 					for (int32 PoseIdx = FirstPoseIdx; PoseIdx < LastPoseIdx; ++PoseIdx)
 					{
-						EvaluatePoseKernel<false, false>(SearchResults, SearchIndex, QueryValues, TArrayView<float>(), PoseIdx, INDEX_NONE, SearchFilters, ContinuingPoseCostAddend, ContinuingInteractionCostAddend, SearchContext, this, DynamicWeightsSqrt
+						EvaluatePoseKernel<false, false>(SearchResults, SearchIndex, QueryValues, TArrayView<float>(), PoseIdx, INDEX_NONE, SearchFilters, ContinuingPoseCostAddend, ContinuingInteractionCostAddend, this, DynamicWeightsSqrt
 #if UE_POSE_SEARCH_TRACE_ENABLED
 							, bTrackResults
 #endif // UE_POSE_SEARCH_TRACE_ENABLED
@@ -2523,13 +2395,13 @@ void UPoseSearchDatabase::SearchEvent(UE::PoseSearch::FSearchContext& SearchCont
 				);
 			}
 
-			FSearchFilters SearchFilters(Schema, NonSelectableIdx, SelectableAssetIdx, SearchIndex.bAnyBlockTransition);
+			FSearchFilters SearchFilters(Schema, NonSelectableIdx, SelectableAssetIdx, SearchIndex.bAnyBlockTransition, QueryValues, SearchContext);
 
 			if (FMath::IsNearlyZero(EventToSearch.TimeToEvent))
 			{
 				for (int32 EventPoseIdx : PosesWithEvent)
 				{
-					EvaluatePoseKernel<false, false>(SearchResults, SearchIndex, QueryValues, TArrayView<float>(), EventPoseIdx, EventPoseIdx, SearchFilters, ContinuingPoseCostAddend, ContinuingInteractionCostAddend, SearchContext, this, DynamicWeightsSqrt
+					EvaluatePoseKernel<false, false>(SearchResults, SearchIndex, QueryValues, TArrayView<float>(), EventPoseIdx, EventPoseIdx, SearchFilters, ContinuingPoseCostAddend, ContinuingInteractionCostAddend, this, DynamicWeightsSqrt
 #if UE_POSE_SEARCH_TRACE_ENABLED
 						, true
 #endif // UE_POSE_SEARCH_TRACE_ENABLED
@@ -2551,7 +2423,7 @@ void UPoseSearchDatabase::SearchEvent(UE::PoseSearch::FSearchContext& SearchCont
 					SearchIndexAsset.IteratePoseIndexesOverTime(EventTime - TimeToEventEnd, EventTime - TimeToEventStart, Schema->SampleRate,
 						[this, &SearchResults, &SearchIndex, &QueryValues, EventPoseIdx, &SearchFilters, ContinuingPoseCostAddend, ContinuingInteractionCostAddend, &SearchContext, &DynamicWeightsSqrt](int32 PoseIdx)
 						{
-							EvaluatePoseKernel<false, false>(SearchResults, SearchIndex, QueryValues, TArrayView<float>(), PoseIdx, EventPoseIdx, SearchFilters, ContinuingPoseCostAddend, ContinuingInteractionCostAddend, SearchContext, this, DynamicWeightsSqrt
+							EvaluatePoseKernel<false, false>(SearchResults, SearchIndex, QueryValues, TArrayView<float>(), PoseIdx, EventPoseIdx, SearchFilters, ContinuingPoseCostAddend, ContinuingInteractionCostAddend, this, DynamicWeightsSqrt
 #if UE_POSE_SEARCH_TRACE_ENABLED
 								, true
 #endif // UE_POSE_SEARCH_TRACE_ENABLED
diff --git a/Engine/Plugins/Animation/PoseSearch/Source/Runtime/Private/PoseSearchFeatureChannel.cpp b/Engine/Plugins/Animation/PoseSearch/Source/Runtime/Private/PoseSearchFeatureChannel.cpp
--- a/Engine/Plugins/Animation/PoseSearch/Source/Runtime/Private/PoseSearchFeatureChannel.cpp
+++ b/Engine/Plugins/Animation/PoseSearch/Source/Runtime/Private/PoseSearchFeatureChannel.cpp
@@ -7,150 +7,6 @@
 
 #include UE_INLINE_GENERATED_CPP_BY_NAME(PoseSearchFeatureChannel)
 
-namespace UE::PoseSearch
-{
-
-//////////////////////////////////////////////////////////////////////////
-// FFeatureVectorHelper
-int32 FFeatureVectorHelper::GetVectorCardinality(EComponentStrippingVector ComponentStrippingVector)
-{
-	switch (ComponentStrippingVector)
-	{
-	case EComponentStrippingVector::None:
-		return 3;
-	case EComponentStrippingVector::StripXY:
-		return 1;
-	case EComponentStrippingVector::StripZ:
-		return 2;
-	default:
-		checkNoEntry();
-		return 0;
-	}
-}
-
-void FFeatureVectorHelper::EncodeVector(TArrayView<float> Values, int32 DataOffset, const FVector& Vector, EComponentStrippingVector ComponentStrippingVector, bool bNormalize)
-{
-	check(Values.GetData());
-
-	if (bNormalize)
-	{
-		FVector NormalizedVector = Vector;
-		switch (ComponentStrippingVector)
-		{
-		case EComponentStrippingVector::None:
-			check(Values.Num() >= DataOffset + 2);
-			NormalizedVector.Normalize();
-			Values[DataOffset + 0] = NormalizedVector.X;
-			Values[DataOffset + 1] = NormalizedVector.Y;
-			Values[DataOffset + 2] = NormalizedVector.Z;
-			break;
-		case EComponentStrippingVector::StripXY:
-			check(Values.Num() >= DataOffset + 0);			
-			NormalizedVector.X = 0.f;
-			NormalizedVector.Y = 0.f;
-			NormalizedVector.Normalize();
-			Values[DataOffset + 0] = NormalizedVector.Z;
-			break;
-		case EComponentStrippingVector::StripZ:
-			check(Values.Num() >= DataOffset + 1);
-			NormalizedVector.Z = 0.f;
-			NormalizedVector.Normalize();
-			Values[DataOffset + 0] = NormalizedVector.X;
-			Values[DataOffset + 1] = NormalizedVector.Y;
-			break;
-		default:
-			checkNoEntry();
-			break;
-		}
-	}
-	else
-	{
-		switch (ComponentStrippingVector)
-		{
-		case EComponentStrippingVector::None:
-			check(Values.Num() >= DataOffset + 2);
-			Values[DataOffset + 0] = Vector.X;
-			Values[DataOffset + 1] = Vector.Y;
-			Values[DataOffset + 2] = Vector.Z;
-			break;
-		case EComponentStrippingVector::StripXY:
-			check(Values.Num() >= DataOffset + 0);
-			Values[DataOffset + 0] = Vector.Z;
-			break;
-		case EComponentStrippingVector::StripZ:
-			check(Values.Num() >= DataOffset + 1);
-			Values[DataOffset + 0] = Vector.X;
-			Values[DataOffset + 1] = Vector.Y;
-			break;
-		default:
-			checkNoEntry();
-			break;
-		}
-	}
-}
-
-FVector FFeatureVectorHelper::DecodeVector(TConstArrayView<float> Values, int32 DataOffset, EComponentStrippingVector ComponentStrippingVector)
-{
-	switch (ComponentStrippingVector)
-	{
-	case EComponentStrippingVector::None:
-		check(Values.Num() >= DataOffset + 2);
-		return FVector(Values[DataOffset + 0], Values[DataOffset + 1], Values[DataOffset + 2]);
-	case EComponentStrippingVector::StripXY:
-		check(Values.Num() >= DataOffset + 0);
-		return FVector(0, 0, Values[DataOffset + 0]);
-	case EComponentStrippingVector::StripZ:
-		check(Values.Num() >= DataOffset + 1);
-		return FVector(Values[DataOffset + 0], Values[DataOffset + 1], 0);
-	default:
-		checkNoEntry();
-		return FVector::Zero();
-	}
-}
-
-void FFeatureVectorHelper::EncodeVector2D(TArrayView<float> Values, int32 DataOffset, const FVector2D& Vector2D)
-{
-	check(Values.Num() >= DataOffset + 1);
-	Values[DataOffset + 0] = Vector2D.X;
-	Values[DataOffset + 1] = Vector2D.Y;
-}
-
-FVector2D FFeatureVectorHelper::DecodeVector2D(TConstArrayView<float> Values, int32 DataOffset)
-{
-	check(Values.Num() >= DataOffset + 1);
-	return FVector2D(Values[DataOffset + 0], Values[DataOffset + 1]);
-}
-
-void FFeatureVectorHelper::EncodeFloat(TArrayView<float> Values, int32 DataOffset, const float Value)
-{
-	check(Values.Num() >= DataOffset + 0);
-	Values[DataOffset + 0] = Value;
-}
-
-float FFeatureVectorHelper::DecodeFloat(TConstArrayView<float> Values, int32 DataOffset)
-{
-	check(Values.Num() >= DataOffset + 0);
-	return Values[DataOffset];
-}
-
-void FFeatureVectorHelper::Copy(TArrayView<float> Values, int32 DataOffset, int32 DataCardinality, TConstArrayView<float> OriginValues)
-{
-	check(Values.GetData() && OriginValues.GetData());
-	check(Values.Num() == OriginValues.Num() && DataOffset + DataCardinality <= Values.Num());
-
-	float* RESTRICT ValuesData = Values.GetData() + DataOffset;
-	const float* RESTRICT OriginValuesData = OriginValues.GetData() + DataOffset;
-
-	check(ValuesData != OriginValuesData);
-
-	for (int32 i = 0; i < DataCardinality; ++i, ++ValuesData, ++OriginValuesData)
-	{
-		*ValuesData = *OriginValuesData;
-	}
-}
-
-} // namespace UE::PoseSearch
-
 //////////////////////////////////////////////////////////////////////////
 // UPoseSearchFeatureChannel
 void UPoseSearchFeatureChannel::GetPermutationTimeOffsets(EPermutationTimeType PermutationTimeType, float DesiredPermutationTimeOffset, float& OutPermutationSampleTimeOffset, float& OutPermutationOriginTimeOffset)
@@ -177,6 +33,19 @@ void UPoseSearchFeatureChannel::GetPermutationTimeOffsets(EPermutationTimeType P
 	}
 }
 
+#if ENABLE_DRAW_DEBUG
+void UPoseSearchFeatureChannel::DebugDraw(const UE::PoseSearch::FDebugDrawParams& DrawParams, TConstArrayView<float> PoseVector) const
+{
+	for (int32 FilterIndex = 0; FilterIndex < GetNumFilters(); ++FilterIndex)
+	{
+		if (const FPoseSearchFilter* Filter = GetFilter(FilterIndex))
+		{
+			Filter->DebugDraw(DrawParams, PoseVector, this);
+		}
+	}
+}
+#endif //ENABLE_DRAW_DEBUG
+
 #if WITH_EDITOR
 void UPoseSearchFeatureChannel::GetOuterLabel(UE::PoseSearch::TLabelBuilder& LabelBuilder, UE::PoseSearch::ELabelFormat LabelFormat) const
 {
@@ -277,7 +146,7 @@ const UE::PoseSearch::FRole UPoseSearchFeatureChannel::GetDefaultRole() const
 {
 	if (const UPoseSearchSchema* Schema = GetSchema())
 	{
-		Schema->GetDefaultRole();
+		return Schema->GetDefaultRole();
 	}
 	return UE::PoseSearch::DefaultRole;
 }
diff --git a/Engine/Plugins/Animation/PoseSearch/Source/Runtime/Private/PoseSearchFeatureChannel_Distance.cpp b/Engine/Plugins/Animation/PoseSearch/Source/Runtime/Private/PoseSearchFeatureChannel_Distance.cpp
--- a/Engine/Plugins/Animation/PoseSearch/Source/Runtime/Private/PoseSearchFeatureChannel_Distance.cpp
+++ b/Engine/Plugins/Animation/PoseSearch/Source/Runtime/Private/PoseSearchFeatureChannel_Distance.cpp
@@ -154,11 +154,13 @@ void UPoseSearchFeatureChannel_Distance::DebugDraw(const UE::PoseSearch::FDebugD
 
 	if (bDrawInjectAdditionalDebugChannels || DrawParams.IsAnyWeightRelevant(this))
 	{
+		Super::DebugDraw(DrawParams, PoseVector);
+
 		FColor Color;
 #if WITH_EDITORONLY_DATA
 		Color = DebugColor.ToFColor(true);
 #else // WITH_EDITORONLY_DATA
-		Color = FLinearColor::Blue.ToFColor(true);
+		Color = FColor::Blue;
 #endif // WITH_EDITORONLY_DATA
 
 		float PermutationSampleTimeOffset = 0.f;
diff --git a/Engine/Plugins/Animation/PoseSearch/Source/Runtime/Private/PoseSearchFeatureChannel_FilterCrashingLegs.cpp b/Engine/Plugins/Animation/PoseSearch/Source/Runtime/Private/PoseSearchFeatureChannel_FilterCrashingLegs.cpp
--- a/Engine/Plugins/Animation/PoseSearch/Source/Runtime/Private/PoseSearchFeatureChannel_FilterCrashingLegs.cpp
+++ b/Engine/Plugins/Animation/PoseSearch/Source/Runtime/Private/PoseSearchFeatureChannel_FilterCrashingLegs.cpp
@@ -126,6 +126,8 @@ void UPoseSearchFeatureChannel_FilterCrashingLegs::DebugDraw(const UE::PoseSearc
 
 	if (DrawParams.IsAnyWeightRelevant(this))
 	{
+		Super::DebugDraw(DrawParams, PoseVector);
+
 		const float CrashingLegsValue = FFeatureVectorHelper::DecodeFloat(PoseVector, ChannelDataOffset);
 
 		const FVector LeftThighPosition = DrawParams.ExtractPosition(PoseVector, 0.f, LeftThighIdx, SampleRole);
@@ -159,6 +161,7 @@ void UPoseSearchFeatureChannel_FilterCrashingLegs::DebugDraw(const UE::PoseSearc
 }
 #endif // ENABLE_DRAW_DEBUG
 
+// @todo: reimplement this using the new FPoseSearchFilter
 // IPoseSearchFilter interface
 bool UPoseSearchFeatureChannel_FilterCrashingLegs::IsFilterActive() const
 {
@@ -212,4 +215,4 @@ bool UPoseSearchFeatureChannel_FilterCrashingLegs::IndexAsset(UE::PoseSearch::FA
 	return true;
 }
 
-#endif
+#endif
\ No newline at end of file
diff --git a/Engine/Plugins/Animation/PoseSearch/Source/Runtime/Private/PoseSearchFeatureChannel_Group.cpp b/Engine/Plugins/Animation/PoseSearch/Source/Runtime/Private/PoseSearchFeatureChannel_Group.cpp
--- a/Engine/Plugins/Animation/PoseSearch/Source/Runtime/Private/PoseSearchFeatureChannel_Group.cpp
+++ b/Engine/Plugins/Animation/PoseSearch/Source/Runtime/Private/PoseSearchFeatureChannel_Group.cpp
@@ -75,6 +75,8 @@ void UPoseSearchFeatureChannel_GroupBase::BuildQuery(UE::PoseSearch::FSearchCont
 #if ENABLE_DRAW_DEBUG
 void UPoseSearchFeatureChannel_GroupBase::DebugDraw(const UE::PoseSearch::FDebugDrawParams& DrawParams, TConstArrayView<float> PoseVector) const
 {
+	Super::DebugDraw(DrawParams, PoseVector);
+
 	for (const TObjectPtr<UPoseSearchFeatureChannel>& SubChannelPtr : GetSubChannels())
 	{
 		if (const UPoseSearchFeatureChannel* SubChannel = SubChannelPtr.Get())
@@ -92,7 +94,9 @@ bool UPoseSearchFeatureChannel_GroupBase::IsFilterActive() const
 	{
 		if (const UPoseSearchFeatureChannel* SubChannel = SubChannelPtr.Get())
 		{
+			PRAGMA_DISABLE_DEPRECATION_WARNINGS
 			if (SubChannel->IsFilterActive())
+			PRAGMA_ENABLE_DEPRECATION_WARNINGS
 			{
 				return true;
 			}
@@ -108,7 +112,9 @@ bool UPoseSearchFeatureChannel_GroupBase::IsFilterValid(TConstArrayView<float> P
 	{
 		if (const UPoseSearchFeatureChannel* SubChannel = SubChannelPtr.Get())
 		{
+			PRAGMA_DISABLE_DEPRECATION_WARNINGS
 			if (!SubChannel->IsFilterValid(PoseValues, QueryValues, PoseIdx, Metadata))
+			PRAGMA_ENABLE_DEPRECATION_WARNINGS
 			{
 				return false;
 			}
diff --git a/Engine/Plugins/Animation/PoseSearch/Source/Runtime/Private/PoseSearchFeatureChannel_Heading.cpp b/Engine/Plugins/Animation/PoseSearch/Source/Runtime/Private/PoseSearchFeatureChannel_Heading.cpp
--- a/Engine/Plugins/Animation/PoseSearch/Source/Runtime/Private/PoseSearchFeatureChannel_Heading.cpp
+++ b/Engine/Plugins/Animation/PoseSearch/Source/Runtime/Private/PoseSearchFeatureChannel_Heading.cpp
@@ -199,11 +199,13 @@ void UPoseSearchFeatureChannel_Heading::DebugDraw(const UE::PoseSearch::FDebugDr
 
 	if (bDrawInjectAdditionalDebugChannels || DrawParams.IsAnyWeightRelevant(this))
 	{
+		Super::DebugDraw(DrawParams, PoseVector);
+
 		FColor Color;
 #if WITH_EDITORONLY_DATA
 		Color = DebugColor.ToFColor(true);
 #else // WITH_EDITORONLY_DATA
-		Color = FLinearColor::White.ToFColor(true);
+		Color = FColor::White;
 #endif // WITH_EDITORONLY_DATA
 
 		float PermutationSampleTimeOffset = 0.f;
diff --git a/Engine/Plugins/Animation/PoseSearch/Source/Runtime/Private/PoseSearchFeatureChannel_Phase.cpp b/Engine/Plugins/Animation/PoseSearch/Source/Runtime/Private/PoseSearchFeatureChannel_Phase.cpp
--- a/Engine/Plugins/Animation/PoseSearch/Source/Runtime/Private/PoseSearchFeatureChannel_Phase.cpp
+++ b/Engine/Plugins/Animation/PoseSearch/Source/Runtime/Private/PoseSearchFeatureChannel_Phase.cpp
@@ -363,11 +363,13 @@ void UPoseSearchFeatureChannel_Phase::DebugDraw(const UE::PoseSearch::FDebugDraw
 
 	if (DrawParams.IsAnyWeightRelevant(this))
 	{
+		Super::DebugDraw(DrawParams, PoseVector);
+
 		FColor Color;
 #if WITH_EDITORONLY_DATA
 		Color = DebugColor.ToFColor(true);
 #else // WITH_EDITORONLY_DATA
-		Color = FLinearColor::Yellow.ToFColor(true);
+		Color = FColor::Yellow;
 #endif // WITH_EDITORONLY_DATA
 
 		const FVector2D Phase = FFeatureVectorHelper::DecodeVector2D(PoseVector, ChannelDataOffset);
diff --git a/Engine/Plugins/Animation/PoseSearch/Source/Runtime/Private/PoseSearchFeatureChannel_Pose.cpp b/Engine/Plugins/Animation/PoseSearch/Source/Runtime/Private/PoseSearchFeatureChannel_Pose.cpp
--- a/Engine/Plugins/Animation/PoseSearch/Source/Runtime/Private/PoseSearchFeatureChannel_Pose.cpp
+++ b/Engine/Plugins/Animation/PoseSearch/Source/Runtime/Private/PoseSearchFeatureChannel_Pose.cpp
@@ -135,10 +135,19 @@ bool UPoseSearchFeatureChannel_Pose::Finalize(UPoseSearchSchema* Schema)
 #if WITH_EDITOR
 UE::PoseSearch::TLabelBuilder& UPoseSearchFeatureChannel_Pose::GetLabel(UE::PoseSearch::TLabelBuilder& LabelBuilder, UE::PoseSearch::ELabelFormat LabelFormat) const
 {
+	using namespace UE::PoseSearch;
+
 	GetOuterLabel(LabelBuilder, LabelFormat);
 	AppendLabelSeparator(LabelBuilder, UE::PoseSearch::ELabelFormat::Full_Horizontal);
 	LabelBuilder.Append(TEXT("Pose"));
 
+	if (SampleRole != DefaultRole)
+	{
+		LabelBuilder.Append(TEXT("["));
+		LabelBuilder.Append(SampleRole.ToString());
+		LabelBuilder.Append(TEXT("]"));
+	}
+
 	if (PermutationTimeType == EPermutationTimeType::UsePermutationTime)
 	{
 		LabelBuilder.Append(TEXT("_PT"));
diff --git a/Engine/Plugins/Animation/PoseSearch/Source/Runtime/Private/PoseSearchFeatureChannel_Position.cpp b/Engine/Plugins/Animation/PoseSearch/Source/Runtime/Private/PoseSearchFeatureChannel_Position.cpp
--- a/Engine/Plugins/Animation/PoseSearch/Source/Runtime/Private/PoseSearchFeatureChannel_Position.cpp
+++ b/Engine/Plugins/Animation/PoseSearch/Source/Runtime/Private/PoseSearchFeatureChannel_Position.cpp
@@ -91,20 +91,6 @@ void UPoseSearchFeatureChannel_Position::AddDependentChannels(UPoseSearchSchema*
 	}
 }
 
-bool UPoseSearchFeatureChannel_Position::IsFilterValid(TConstArrayView<float> PoseValues, TConstArrayView<float> QueryValues, int32 PoseIdx, const UE::PoseSearch::FPoseMetadata& Metadata) const
-{
-	using namespace UE::PoseSearch;
-	const FVector Pose = FFeatureVectorHelper::DecodeVector(PoseValues, ChannelDataOffset, ComponentStripping);
-	const FVector Query = FFeatureVectorHelper::DecodeVector(QueryValues, ChannelDataOffset, ComponentStripping);
-
-	const float SquaredLength = (Pose - Query).SquaredLength();
-
-	check(MaxPositionDistanceSquared > 0.f);
-
-	return SquaredLength <= MaxPositionDistanceSquared;
-}
-
-
 void UPoseSearchFeatureChannel_Position::BuildQuery(UE::PoseSearch::FSearchContext& SearchContext) const
 {
 	using namespace UE::PoseSearch;
@@ -206,11 +192,13 @@ void UPoseSearchFeatureChannel_Position::DebugDraw(const UE::PoseSearch::FDebugD
 
 	if (bDrawInjectAdditionalDebugChannels || DrawParams.IsAnyWeightRelevant(this))
 	{
+		Super::DebugDraw(DrawParams, PoseVector);
+
 		FColor Color;
 #if WITH_EDITORONLY_DATA
 		Color = DebugColor.ToFColor(true);
 #else // WITH_EDITORONLY_DATA
-		Color = FLinearColor::Blue.ToFColor(true);
+		Color = FColor::Blue;
 #endif // WITH_EDITORONLY_DATA
 
 		float PermutationSampleTimeOffset = 0.f;
@@ -223,21 +211,6 @@ void UPoseSearchFeatureChannel_Position::DebugDraw(const UE::PoseSearch::FDebugD
 		const FVector DeltaPos = DrawParams.ExtractRotation(PoseVector, OriginTimeOffset, RootSchemaBoneIdx, OriginRole, OriginPermutationTimeType, INDEX_NONE, PermutationOriginTimeOffset).RotateVector(FeaturesVector);
 		const FVector BonePos = OriginBonePos + DeltaPos;
 
-		if (MaxPositionDistanceSquared > 0.f)
-		{
-			static constexpr int32 Segments = 32;
-			const float Radius = FMath::Sqrt(MaxPositionDistanceSquared);
-			if (ComponentStripping == EComponentStrippingVector::StripZ)
-			{
-				const FMatrix CircleTransform(FVector::ZAxisVector, FVector::XAxisVector, FVector::YAxisVector, BonePos);
-				DrawParams.DrawCircle(CircleTransform, Radius, Segments, Color);
-			}
-			else
-			{
-				DrawParams.DrawSphere(BonePos, Radius, Segments, Color);
-			}
-		}
-
 		if (bNormalizeDisplacement)
 		{
 			static const float NormalizeDisplacementLength = 100.f;
@@ -391,4 +364,4 @@ USkeleton* UPoseSearchFeatureChannel_Position::GetSkeleton(bool& bInvalidSkeleto
 
 	return Super::GetSkeleton(bInvalidSkeletonIsError, PropertyHandle);
 }
-#endif
+#endif
\ No newline at end of file
diff --git a/Engine/Plugins/Animation/PoseSearch/Source/Runtime/Private/PoseSearchFeatureChannel_Trajectory.cpp b/Engine/Plugins/Animation/PoseSearch/Source/Runtime/Private/PoseSearchFeatureChannel_Trajectory.cpp
--- a/Engine/Plugins/Animation/PoseSearch/Source/Runtime/Private/PoseSearchFeatureChannel_Trajectory.cpp
+++ b/Engine/Plugins/Animation/PoseSearch/Source/Runtime/Private/PoseSearchFeatureChannel_Trajectory.cpp
@@ -143,6 +143,8 @@ void UPoseSearchFeatureChannel_Trajectory::DebugDraw(const UE::PoseSearch::FDebu
 
 	if (DrawParams.IsAnyWeightRelevant(this))
 	{
+		Super::DebugDraw(DrawParams, PoseVector);
+
 		TArray<const UPoseSearchFeatureChannel_Position*, TInlineAllocator<32>> Positions;
 		for (const TObjectPtr<UPoseSearchFeatureChannel>& SubChannelPtr : GetSubChannels())
 		{
@@ -170,7 +172,7 @@ void UPoseSearchFeatureChannel_Trajectory::DebugDraw(const UE::PoseSearch::FDebu
 #if WITH_EDITORONLY_DATA
 				Color = Positions[i]->DebugColor.ToFColor(true);
 #else // WITH_EDITORONLY_DATA
-				Color = FLinearColor::Blue.ToFColor(true);
+				Color = FColor::Blue;
 #endif // WITH_EDITORONLY_DATA
 
 				if (Positions.Num() == 1 || PrevTimeOffset * CurrTimeOffset < UE_KINDA_SMALL_NUMBER)
diff --git a/Engine/Plugins/Animation/PoseSearch/Source/Runtime/Private/PoseSearchFeatureChannel_Velocity.cpp b/Engine/Plugins/Animation/PoseSearch/Source/Runtime/Private/PoseSearchFeatureChannel_Velocity.cpp
--- a/Engine/Plugins/Animation/PoseSearch/Source/Runtime/Private/PoseSearchFeatureChannel_Velocity.cpp
+++ b/Engine/Plugins/Animation/PoseSearch/Source/Runtime/Private/PoseSearchFeatureChannel_Velocity.cpp
@@ -144,11 +144,13 @@ void UPoseSearchFeatureChannel_Velocity::DebugDraw(const UE::PoseSearch::FDebugD
 
 	if (DrawParams.IsAnyWeightRelevant(this))
 	{
+		Super::DebugDraw(DrawParams, PoseVector);
+
 		FColor Color;
 #if WITH_EDITORONLY_DATA
 		Color = DebugColor.ToFColor(true);
 #else // WITH_EDITORONLY_DATA
-		Color = FLinearColor::Green.ToFColor(true);
+		Color = FColor::Green;
 #endif // WITH_EDITORONLY_DATA
 
 		const float LinearVelocityScale = bNormalize ? 15.f : 0.08f;
diff --git a/Engine/Plugins/Animation/PoseSearch/Source/Runtime/Private/PoseSearchFeatureVectorHelper.cpp b/Engine/Plugins/Animation/PoseSearch/Source/Runtime/Private/PoseSearchFeatureVectorHelper.cpp
new file mode 100644
--- /dev/null
+++ b/Engine/Plugins/Animation/PoseSearch/Source/Runtime/Private/PoseSearchFeatureVectorHelper.cpp
@@ -0,0 +1,146 @@
+// Copyright Epic Games, Inc. All Rights Reserved.
+
+#include "PoseSearch/PoseSearchFeatureVectorHelper.h"
+#include "CoreMinimal.h"
+
+namespace UE::PoseSearch
+{
+
+int32 FFeatureVectorHelper::GetVectorCardinality(EComponentStrippingVector ComponentStrippingVector)
+{
+	switch (ComponentStrippingVector)
+	{
+	case EComponentStrippingVector::None:
+		return 3;
+	case EComponentStrippingVector::StripXY:
+		return 1;
+	case EComponentStrippingVector::StripZ:
+		return 2;
+	default:
+		checkNoEntry();
+		return 0;
+	}
+}
+
+void FFeatureVectorHelper::EncodeVector(TArrayView<float> Values, int32 DataOffset, const FVector& Vector, EComponentStrippingVector ComponentStrippingVector, bool bNormalize)
+{
+	check(Values.GetData());
+
+	if (bNormalize)
+	{
+		FVector NormalizedVector = Vector;
+		switch (ComponentStrippingVector)
+		{
+		case EComponentStrippingVector::None:
+			check(Values.Num() > DataOffset + 2);
+			NormalizedVector.Normalize();
+			Values[DataOffset + 0] = NormalizedVector.X;
+			Values[DataOffset + 1] = NormalizedVector.Y;
+			Values[DataOffset + 2] = NormalizedVector.Z;
+			break;
+		case EComponentStrippingVector::StripXY:
+			check(Values.Num() > DataOffset + 0);			
+			NormalizedVector.X = 0.f;
+			NormalizedVector.Y = 0.f;
+			NormalizedVector.Normalize();
+			Values[DataOffset + 0] = NormalizedVector.Z;
+			break;
+		case EComponentStrippingVector::StripZ:
+			check(Values.Num() > DataOffset + 1);
+			NormalizedVector.Z = 0.f;
+			NormalizedVector.Normalize();
+			Values[DataOffset + 0] = NormalizedVector.X;
+			Values[DataOffset + 1] = NormalizedVector.Y;
+			break;
+		default:
+			checkNoEntry();
+			break;
+		}
+	}
+	else
+	{
+		switch (ComponentStrippingVector)
+		{
+		case EComponentStrippingVector::None:
+			check(Values.Num() > DataOffset + 2);
+			Values[DataOffset + 0] = Vector.X;
+			Values[DataOffset + 1] = Vector.Y;
+			Values[DataOffset + 2] = Vector.Z;
+			break;
+		case EComponentStrippingVector::StripXY:
+			check(Values.Num() > DataOffset + 0);
+			Values[DataOffset + 0] = Vector.Z;
+			break;
+		case EComponentStrippingVector::StripZ:
+			check(Values.Num() > DataOffset + 1);
+			Values[DataOffset + 0] = Vector.X;
+			Values[DataOffset + 1] = Vector.Y;
+			break;
+		default:
+			checkNoEntry();
+			break;
+		}
+	}
+}
+
+FVector FFeatureVectorHelper::DecodeVector(TConstArrayView<float> Values, int32 DataOffset, EComponentStrippingVector ComponentStrippingVector)
+{
+	switch (ComponentStrippingVector)
+	{
+	case EComponentStrippingVector::None:
+		check(Values.Num() > DataOffset + 2);
+		return FVector(Values[DataOffset + 0], Values[DataOffset + 1], Values[DataOffset + 2]);
+	case EComponentStrippingVector::StripXY:
+		check(Values.Num() > DataOffset + 0);
+		return FVector(0, 0, Values[DataOffset + 0]);
+	case EComponentStrippingVector::StripZ:
+		check(Values.Num() > DataOffset + 1);
+		return FVector(Values[DataOffset + 0], Values[DataOffset + 1], 0);
+	default:
+		checkNoEntry();
+		return FVector::Zero();
+	}
+}
+
+void FFeatureVectorHelper::EncodeVector2D(TArrayView<float> Values, int32 DataOffset, const FVector2D& Vector2D)
+{
+	check(Values.Num() > DataOffset + 1);
+	Values[DataOffset + 0] = Vector2D.X;
+	Values[DataOffset + 1] = Vector2D.Y;
+}
+
+FVector2D FFeatureVectorHelper::DecodeVector2D(TConstArrayView<float> Values, int32 DataOffset)
+{
+	check(Values.Num() > DataOffset + 1);
+	return FVector2D(Values[DataOffset + 0], Values[DataOffset + 1]);
+}
+
+void FFeatureVectorHelper::EncodeFloat(TArrayView<float> Values, int32 DataOffset, const float Value)
+{
+	check(Values.Num() > DataOffset + 0);
+	Values[DataOffset + 0] = Value;
+}
+
+float FFeatureVectorHelper::DecodeFloat(TConstArrayView<float> Values, int32 DataOffset)
+{
+	check(Values.Num() > DataOffset + 0);
+	return Values[DataOffset];
+}
+
+void FFeatureVectorHelper::Copy(TArrayView<float> Values, int32 DataOffset, int32 DataCardinality, TConstArrayView<float> OriginValues)
+{
+	check(Values.GetData() && OriginValues.GetData());
+	check(Values.Num() == OriginValues.Num() && DataOffset + DataCardinality <= Values.Num());
+
+	float* RESTRICT ValuesData = Values.GetData() + DataOffset;
+	const float* RESTRICT OriginValuesData = OriginValues.GetData() + DataOffset;
+
+	check(ValuesData != OriginValuesData);
+
+	for (int32 i = 0; i < DataCardinality; ++i, ++ValuesData, ++OriginValuesData)
+	{
+		*ValuesData = *OriginValuesData;
+	}
+}
+
+} // namespace UE::PoseSearch
diff --git a/Engine/Plugins/Animation/PoseSearch/Source/Runtime/Private/PoseSearchFilter.cpp b/Engine/Plugins/Animation/PoseSearch/Source/Runtime/Private/PoseSearchFilter.cpp
new file mode 100644
--- /dev/null
+++ b/Engine/Plugins/Animation/PoseSearch/Source/Runtime/Private/PoseSearchFilter.cpp
@@ -0,0 +1,422 @@
+// Copyright Epic Games, Inc. All Rights Reserved.
+
+#include "PoseSearch/PoseSearchFilter.h"
+#include "PoseSearch/PoseSearchContext.h"
+#include "PoseSearch/PoseSearchSchema.h"
+
+// @todo: add support for generic vector (3 float of data) channel
+#include "PoseSearch/PoseSearchFeatureChannel_Position.h"
+
+///////////////////////////////////////////////////
+// FPoseSearchFilterPadding 
+int32 FPoseSearchFilterPadding::GetRequiredDataSize(const UPoseSearchFeatureChannel* Channel) const
+{
+	return 1;
+}
+
+///////////////////////////////////////////////////
+// FPoseSearchFilterMaxPosition 
+int32 FPoseSearchFilterMaxPosition::GetRequiredDataSize(const UPoseSearchFeatureChannel* Channel) const
+{
+	// @todo: add support for generic vector (3 float of data) channel
+	if (MaxPositionDistanceSquared > UE_KINDA_SMALL_NUMBER &&
+		Cast<UPoseSearchFeatureChannel_Position>(Channel))
+	{
+		return sizeof(FPoseSearchFilterMaxPosition::FData);
+	}
+	// @todo: log some warnings?
+	return InvalidFilterDataSize;
+}
+
+
+bool FPoseSearchFilterMaxPosition::InitializeData(TArrayView<uint8> Data, const UPoseSearchFeatureChannel* Channel, TConstArrayView<float> QueryValues, UE::PoseSearch::FSearchContext& SearchContext) const
+{
+	using namespace UE::PoseSearch;
+
+	check(IsAligned(Data.GetData(), alignof(FPoseSearchFilterMaxPosition::FData)));
+	check(Data.Num() == sizeof(FPoseSearchFilterMaxPosition::FData));
+
+	// @todo: add support for generic vector (3 float of data) channel
+	check(nullptr != Cast<UPoseSearchFeatureChannel_Position>(Channel));
+	const UPoseSearchFeatureChannel_Position* PositionChannel = static_cast<const UPoseSearchFeatureChannel_Position*>(Channel);
+
+	const int32 ChannelDataOffset = PositionChannel->GetChannelDataOffset();
+	const EComponentStrippingVector ComponentStripping = PositionChannel->ComponentStripping;
+
+	FPoseSearchFilterMaxPosition::FData& FilterMaxPositionData = reinterpret_cast<FPoseSearchFilterMaxPosition::FData&>(*Data.GetData());
+	FilterMaxPositionData.Query = FFeatureVectorHelper::DecodeVector(QueryValues, ChannelDataOffset, ComponentStripping);
+
+	return true;
+}
+
+void FPoseSearchFilterMaxPosition::DeinitializeData(TArrayView<uint8> Data, const UPoseSearchFeatureChannel* Channel, UE::PoseSearch::FSearchContext& SearchContext) const
+{
+	check(IsAligned(Data.GetData(), alignof(FPoseSearchFilterMaxPosition::FData)));
+	check(Data.Num() == sizeof(FPoseSearchFilterMaxPosition::FData));
+}
+	
+bool FPoseSearchFilterMaxPosition::IsFilterValid(TConstArrayView<float> PoseValues, TConstArrayView<float> QueryValues, int32 PoseIdx, const UE::PoseSearch::FPoseMetadata& Metadata, TConstArrayView<uint8> Data, const UPoseSearchFeatureChannel* Channel) const
+{
+	using namespace UE::PoseSearch;
+
+	check(IsAligned(Data.GetData(), alignof(FPoseSearchFilterMaxPosition::FData)));
+	check(Data.Num() == sizeof(FPoseSearchFilterMaxPosition::FData));
+
+	const FPoseSearchFilterMaxPosition::FData& FilterMaxPositionData = reinterpret_cast<const FPoseSearchFilterMaxPosition::FData&>(*Data.GetData());
+
+	// @todo: add support for generic vector (3 float of data) channel
+	check(nullptr != Cast<UPoseSearchFeatureChannel_Position>(Channel));
+
+	const UPoseSearchFeatureChannel_Position* PositionChannel = static_cast<const UPoseSearchFeatureChannel_Position*>(Channel);
+	const int32 ChannelDataOffset = PositionChannel->GetChannelDataOffset();
+	const EComponentStrippingVector ComponentStripping = PositionChannel->ComponentStripping;
+
+	const FVector Pose = FFeatureVectorHelper::DecodeVector(PoseValues, ChannelDataOffset, ComponentStripping);
+
+	const float SquaredLength = (Pose - FilterMaxPositionData.Query).SquaredLength();
+
+	return SquaredLength <= MaxPositionDistanceSquared;
+}
+
+#if ENABLE_DRAW_DEBUG
+void FPoseSearchFilterMaxPosition::DebugDraw(const UE::PoseSearch::FDebugDrawParams& DrawParams, TConstArrayView<float> PoseValues, const UPoseSearchFeatureChannel* Channel) const
+{
+	using namespace UE::PoseSearch;
+
+	if (DrawParams.GetDrawContext() == FDebugDrawParams::EDrawContext::DrawQuery)
+	{
+		// @todo: add support for generic vector (3 float of data) channel
+		if (const UPoseSearchFeatureChannel_Position* PositionChannel = Cast<UPoseSearchFeatureChannel_Position>(Channel))
+		{
+			const TConstArrayView<float> PoseVector = PoseValues;
+			const int32 ChannelDataOffset = PositionChannel->GetChannelDataOffset();
+			const EComponentStrippingVector ComponentStripping = PositionChannel->ComponentStripping;
+			const float OriginTimeOffset = PositionChannel->OriginTimeOffset;
+			const int8 SchemaOriginBoneIdx = PositionChannel->SchemaOriginBoneIdx;
+			const FRole OriginRole = PositionChannel->OriginRole;
+			const EPermutationTimeType PermutationTimeType = PositionChannel->PermutationTimeType;
+
+			float PermutationSampleTimeOffset = 0.f;
+			float PermutationOriginTimeOffset = 0.f;
+			UPoseSearchFeatureChannel::GetPermutationTimeOffsets(PermutationTimeType, DrawParams.ExtractPermutationTime(PoseVector), PermutationSampleTimeOffset, PermutationOriginTimeOffset);
+			const EPermutationTimeType OriginPermutationTimeType = PermutationTimeType == EPermutationTimeType::UsePermutationTime ? EPermutationTimeType::UseSampleToPermutationTime : EPermutationTimeType::UseSampleTime;
+
+			const FVector FeaturesVector = FFeatureVectorHelper::DecodeVector(PoseVector, ChannelDataOffset, ComponentStripping);
+			const FVector OriginBonePos = DrawParams.ExtractPosition(PoseVector, OriginTimeOffset, SchemaOriginBoneIdx, OriginRole, OriginPermutationTimeType, INDEX_NONE, PermutationOriginTimeOffset);
+			const FVector DeltaPos = DrawParams.ExtractRotation(PoseVector, OriginTimeOffset, RootSchemaBoneIdx, OriginRole, OriginPermutationTimeType, INDEX_NONE, PermutationOriginTimeOffset).RotateVector(FeaturesVector);
+			const FVector BonePos = OriginBonePos + DeltaPos;
+
+			static constexpr int32 Segments = 32;
+			const float Radius = FMath::Sqrt(MaxPositionDistanceSquared);
+
+#if WITH_EDITORONLY_DATA
+			const FColor Color = PositionChannel->GetDebugColor().ToFColor(true);
+#else // WITH_EDITORONLY_DATA
+			const FColor Color = FColor::White;
+#endif // WITH_EDITORONLY_DATA
+
+			if (PositionChannel->ComponentStripping == EComponentStrippingVector::StripZ)
+			{
+				DrawParams.DrawCircle(FMatrix(FVector::ZAxisVector, FVector::XAxisVector, FVector::YAxisVector, BonePos), Radius, Segments, Color);
+			}
+			else
+			{
+				DrawParams.DrawSphere(BonePos, Radius, Segments, Color);
+			}
+		}
+	}
+}
+
+#endif // ENABLE_DRAW_DEBUG
+
+///////////////////////////////////////////////////
+// FPoseSearchFilterWedge 
+int32 FPoseSearchFilterWedge::GetRequiredDataSize(const UPoseSearchFeatureChannel* Channel) const
+{
+	// @todo: add support for generic vector (3 float of data) channel
+	if (WedgeWidthInDegrees > UE_KINDA_SMALL_NUMBER && 
+		!FMath::IsNearlyEqual(RadiusOffsets.Min, RadiusOffsets.Max,UE_KINDA_SMALL_NUMBER) && 
+		Cast<UPoseSearchFeatureChannel_Position>(Channel))
+	{
+		return sizeof(FPoseSearchFilterWedge::FData);
+	}
+	// @todo: log some warnings?
+	return InvalidFilterDataSize;
+}
+
+bool FPoseSearchFilterWedge::InitializeData(TArrayView<uint8> Data, const UPoseSearchFeatureChannel* Channel, TConstArrayView<float> QueryValues, UE::PoseSearch::FSearchContext& SearchContext) const
+{
+	using namespace UE::PoseSearch;
+
+	check(IsAligned(Data.GetData(), alignof(FPoseSearchFilterWedge::FData)));
+	check(Data.Num() == sizeof(FPoseSearchFilterWedge::FData));
+
+	// @todo: add support for generic vector (3 float of data) channel
+	check(nullptr != Cast<UPoseSearchFeatureChannel_Position>(Channel));
+	const UPoseSearchFeatureChannel_Position* PositionChannel = static_cast<const UPoseSearchFeatureChannel_Position*>(Channel);
+
+	const int32 ChannelDataOffset = PositionChannel->GetChannelDataOffset();
+	const EComponentStrippingVector ComponentStripping = PositionChannel->ComponentStripping;
+
+	FPoseSearchFilterWedge::FData& FilterWedgeData = reinterpret_cast<FPoseSearchFilterWedge::FData&>(*Data.GetData());
+	FilterWedgeData.NormalizedQuery = FFeatureVectorHelper::DecodeVector(QueryValues, ChannelDataOffset, ComponentStripping);
+	const float QueryLength = FilterWedgeData.NormalizedQuery.Length();
+	
+	const bool bShouldFilter = QueryLength > UE_KINDA_SMALL_NUMBER;
+	if (bShouldFilter)
+	{
+		FilterWedgeData.NormalizedQuery /= QueryLength;
+		FilterWedgeData.InnerRadius = FMath::Max(0.f, QueryLength + RadiusOffsets.Min);
+		FilterWedgeData.OuterRadius = FMath::Max(0.f, QueryLength + RadiusOffsets.Max);
+		FilterWedgeData.DotWedgeWidth = FMath::Cos(FMath::DegreesToRadians(WedgeWidthInDegrees * 0.5f));
+		return true;
+	}
+
+	return false;
+}
+
+void FPoseSearchFilterWedge::DeinitializeData(TArrayView<uint8> Data, const UPoseSearchFeatureChannel* Channel, UE::PoseSearch::FSearchContext& SearchContext) const
+{
+	check(IsAligned(Data.GetData(), alignof(FPoseSearchFilterWedge::FData)));
+	check(Data.Num() == sizeof(FPoseSearchFilterWedge::FData));
+}
+
+// @todo: fix the math to be on the plane with normal ZAxisVector. or fix the drawing to support a 3d conical wedge
+bool FPoseSearchFilterWedge::IsFilterValid(TConstArrayView<float> PoseValues, TConstArrayView<float> QueryValues, int32 PoseIdx, const UE::PoseSearch::FPoseMetadata& Metadata, TConstArrayView<uint8> Data, const UPoseSearchFeatureChannel* Channel) const
+{
+	using namespace UE::PoseSearch;
+
+	check(IsAligned(Data.GetData(), alignof(FPoseSearchFilterWedge::FData)));
+	check(Data.Num() == sizeof(FPoseSearchFilterWedge::FData));
+
+	const FPoseSearchFilterWedge::FData& FilterWedgeData = reinterpret_cast<const FPoseSearchFilterWedge::FData&>(*Data.GetData());
+
+	// @todo: add support for generic vector (3 float of data) channel
+	check(nullptr != Cast<UPoseSearchFeatureChannel_Position>(Channel));
+	const UPoseSearchFeatureChannel_Position* PositionChannel = static_cast<const UPoseSearchFeatureChannel_Position*>(Channel);
+	
+	const int32 ChannelDataOffset = PositionChannel->GetChannelDataOffset();
+	const EComponentStrippingVector ComponentStripping = PositionChannel->ComponentStripping;
+
+	const FVector Pose = FFeatureVectorHelper::DecodeVector(PoseValues, ChannelDataOffset, ComponentStripping);
+	const float PoseLength = Pose.Length();
+
+	if (PoseLength > UE_KINDA_SMALL_NUMBER)
+	{
+		if (PoseLength >= FilterWedgeData.InnerRadius && PoseLength <= FilterWedgeData.OuterRadius)
+		{
+			const FVector NormalizedPose = Pose / PoseLength;
+			const float Dot = NormalizedPose | FilterWedgeData.NormalizedQuery;
+
+			if (Dot > FilterWedgeData.DotWedgeWidth)
+			{
+				return true;
+			}
+		}
+	}
+	
+	return false;
+}
+
+#if ENABLE_DRAW_DEBUG
+void FPoseSearchFilterWedge::DebugDraw(const UE::PoseSearch::FDebugDrawParams& DrawParams, TConstArrayView<float> PoseValues, const UPoseSearchFeatureChannel* Channel) const
+{
+	using namespace UE::PoseSearch;
+
+	if (DrawParams.GetDrawContext() == FDebugDrawParams::EDrawContext::DrawQuery)
+	{
+		// @todo: add support for generic vector (3 float of data) channel
+		if (const UPoseSearchFeatureChannel_Position* PositionChannel = Cast<UPoseSearchFeatureChannel_Position>(Channel))
+		{
+			// @todo: avoid this duplicate calculation already done in UPoseSearchFeatureChannel_Position::DebugDraw
+
+			const TConstArrayView<float> PoseVector = PoseValues;
+			const int32 ChannelDataOffset = PositionChannel->GetChannelDataOffset();
+			const EComponentStrippingVector ComponentStripping = PositionChannel->ComponentStripping;
+			const float OriginTimeOffset = PositionChannel->OriginTimeOffset;
+			const int8 SchemaOriginBoneIdx = PositionChannel->SchemaOriginBoneIdx;
+			const FRole OriginRole = PositionChannel->OriginRole;
+			const EPermutationTimeType PermutationTimeType = PositionChannel->PermutationTimeType;
+
+			float PermutationSampleTimeOffset = 0.f;
+			float PermutationOriginTimeOffset = 0.f;
+			UPoseSearchFeatureChannel::GetPermutationTimeOffsets(PermutationTimeType, DrawParams.ExtractPermutationTime(PoseVector), PermutationSampleTimeOffset, PermutationOriginTimeOffset);
+			const EPermutationTimeType OriginPermutationTimeType = PermutationTimeType == EPermutationTimeType::UsePermutationTime ? EPermutationTimeType::UseSampleToPermutationTime : EPermutationTimeType::UseSampleTime;
+
+			const FVector FeaturesVector = FFeatureVectorHelper::DecodeVector(PoseVector, ChannelDataOffset, ComponentStripping);
+			const FVector OriginBonePos = DrawParams.ExtractPosition(PoseVector, OriginTimeOffset, SchemaOriginBoneIdx, OriginRole, OriginPermutationTimeType, INDEX_NONE, PermutationOriginTimeOffset);
+			const FVector DeltaPos = DrawParams.ExtractRotation(PoseVector, OriginTimeOffset, RootSchemaBoneIdx, OriginRole, OriginPermutationTimeType, INDEX_NONE, PermutationOriginTimeOffset).RotateVector(FeaturesVector);
+			const FVector BonePos = OriginBonePos + DeltaPos;
+
+			//const FVector RootBonePos = DrawParams.GetRootBoneTransform(OriginRole, SamplingRootTime).GetLocation();
+			const FVector WedgeDirection = BonePos - OriginBonePos;
+
+			static constexpr int32 Segments = 32;
+
+#if WITH_EDITORONLY_DATA
+			const FColor Color = PositionChannel->GetDebugColor().ToFColor(true);
+#else // WITH_EDITORONLY_DATA
+			const FColor Color = FColor::White;
+#endif // WITH_EDITORONLY_DATA
+
+			const float WedgeLength = WedgeDirection.Length();
+			const float InnerRadius = FMath::Max(0.f, WedgeLength + RadiusOffsets.Min);
+			const float OuterRadius = FMath::Max(0.f, WedgeLength + RadiusOffsets.Max);
+			DrawParams.DrawWedge(OriginBonePos, WedgeDirection, InnerRadius, OuterRadius, WedgeWidthInDegrees, Segments, Color);
+		}
+	}
+}
+#endif // ENABLE_DRAW_DEBUG
+
+namespace UE::PoseSearch
+{
+	
+FSearchFilters::FSearchFilters(const UPoseSearchSchema* Schema, TConstArrayView<int32> NonSelectableIdx, TConstArrayView<int32> SelectableAssetIdx,
+	bool bAddBlockTransitionFilter, TConstArrayView<float> QueryValues, FSearchContext& InSearchContext)
+	: SearchContext(InSearchContext)
+{
+	struct FFilterPreContext
+	{
+		FFilterPreContext(const FPoseSearchFilter* InFilter, const UPoseSearchFeatureChannel* InChannel, int32 InDataSize)
+			: Filter(InFilter)
+			, Channel(InChannel)
+			, DataSize(InDataSize)
+		{
+		}
+
+		const FPoseSearchFilter* Filter;
+		const UPoseSearchFeatureChannel* Channel;
+		int32 DataSize;
+	};
+	TArray<FFilterPreContext, TInlineAllocator<64, TMemStackAllocator<>>> FilterPreContexts;
+
+	int32 TotalFiltersDataSize = 0;
+	if (bAddBlockTransitionFilter)
+	{
+		const int32 BlockTransitionFilterDataSize = BlockTransitionFilter.GetRequiredDataSize(nullptr);
+		TotalFiltersDataSize += Align(BlockTransitionFilterDataSize, FPoseSearchFilter::DataAlignment);
+		FilterPreContexts.Emplace(&BlockTransitionFilter, nullptr, BlockTransitionFilterDataSize);
+	}
+
+	const int32 NonSelectableIdxFilterDataSize = NonSelectableIdxFilter.Init(NonSelectableIdx).GetRequiredDataSize(nullptr);
+	if (NonSelectableIdxFilterDataSize != FPoseSearchFilter::InvalidFilterDataSize)
+	{
+		TotalFiltersDataSize += Align(NonSelectableIdxFilterDataSize, FPoseSearchFilter::DataAlignment);
+		FilterPreContexts.Emplace(&NonSelectableIdxFilter, nullptr, NonSelectableIdxFilterDataSize);
+	}
+
+	const int32 SelectableAssetIdxFilterDataSize = SelectableAssetIdxFilter.Init(SelectableAssetIdx).GetRequiredDataSize(nullptr);
+	if (SelectableAssetIdxFilterDataSize != FPoseSearchFilter::InvalidFilterDataSize)
+	{
+		TotalFiltersDataSize += Align(SelectableAssetIdxFilterDataSize, FPoseSearchFilter::DataAlignment);
+		FilterPreContexts.Emplace(&SelectableAssetIdxFilter, nullptr, SelectableAssetIdxFilterDataSize);
+	}
+
+	// initialing Filters from all the channels FPoseSearchFilter(s) by exploring the schema recursively
+	// to collect FilterPreContexts and sum up the TotalFiltersDataSize
+	Schema->IterateChannels([&FilterPreContexts, &TotalFiltersDataSize](const UPoseSearchFeatureChannel* Channel)
+		{
+			const int32 NumFilters = Channel->GetNumFilters();
+			for (int32 FilterIndex = 0; FilterIndex < NumFilters; ++FilterIndex)
+			{
+				const FPoseSearchFilter* Filter = Channel->GetFilter(FilterIndex);
+				check(Filter);
+				const int32 ChannelFilterDataSize = Filter->GetRequiredDataSize(Channel);
+				if (ChannelFilterDataSize != FPoseSearchFilter::InvalidFilterDataSize)
+				{
+					TotalFiltersDataSize += Align(ChannelFilterDataSize, FPoseSearchFilter::DataAlignment);
+					FilterPreContexts.Emplace(Filter, Channel, ChannelFilterDataSize);
+				}
+			}
+		});
+
+	// initializing the data required by the filters
+	FiltersData.SetNumUninitialized(TotalFiltersDataSize);
+	check(IsAligned(FiltersData.GetData(), FPoseSearchFilter::DataAlignment));
+
+	// converting FFilterPreContext into FFilterContext
+	int32 CurrentFilterDataOffset = 0;
+	for (const FFilterPreContext& FilterPreContext : FilterPreContexts)
+	{
+		TArrayView<uint8> FilterData = MakeArrayView(FiltersData.GetData() + CurrentFilterDataOffset, FilterPreContext.DataSize);
+		if (FilterPreContext.Filter->InitializeData(FilterData, FilterPreContext.Channel, QueryValues, SearchContext))
+		{
+			CurrentFilterDataOffset += Align(FilterPreContext.DataSize, FPoseSearchFilter::DataAlignment);
+			FilterContexts.Emplace(FilterPreContext.Filter, FilterPreContext.Channel, FilterData);
+		}
+		else
+		{
+			// initialization failed. Deinitializing the filter and continue
+			FilterPreContext.Filter->DeinitializeData(FilterData, FilterPreContext.Channel, SearchContext);
+		}
+	}
+
+	PRAGMA_DISABLE_DEPRECATION_WARNINGS
+	// adding the deprecated filters that don't require any data
+	for (const IPoseSearchFilter* Filter : Schema->GetChannels())
+	{
+		if (Filter->IsFilterActive())
+		{
+			Filters_DEPRECATED.Add(Filter);
+		}
+	}
+	PRAGMA_ENABLE_DEPRECATION_WARNINGS
+}
+
+FSearchFilters::~FSearchFilters()
+{
+	for (FFilterContext& FilterContext : FilterContexts)
+	{
+		FilterContext.Filter->DeinitializeData(FilterContext.Data, FilterContext.Channel, SearchContext);
+	}
+}
+
+bool FSearchFilters::AreFiltersValid(const FSearchIndex& SearchIndex, TConstArrayView<float> PoseValues, TConstArrayView<float> QueryValues, TConstArrayView<float> DynamicWeightsSqrt, int32 PoseIdx
+#if UE_POSE_SEARCH_TRACE_ENABLED
+	, float ContinuingPoseCostAddend, float ContinuingInteractionCostAddend, const UPoseSearchDatabase* Database
+#endif // UE_POSE_SEARCH_TRACE_ENABLED
+) const
+{
+	for (const FFilterContext& FilterContext : FilterContexts)
+	{
+		if (!FilterContext.Filter->IsFilterValid(PoseValues, QueryValues, PoseIdx, SearchIndex.PoseMetadata[PoseIdx], FilterContext.Data, FilterContext.Channel))
+		{
+#if UE_POSE_SEARCH_TRACE_ENABLED
+			if (FilterContext.Filter == &NonSelectableIdxFilter)
+			{
+				// candidate already added to SearchContext.BestCandidates by PopulateNonSelectableIdx
+			}
+			else if (FilterContext.Filter == &SelectableAssetIdxFilter)
+			{
+				SearchContext.Track(Database, PoseIdx, EPoseCandidateFlags::DiscardedBy_AssetIdxFilter, DynamicWeightsSqrt, ContinuingPoseCostAddend, ContinuingInteractionCostAddend);
+			}
+			else if (FilterContext.Filter == &BlockTransitionFilter)
+			{
+				SearchContext.Track(Database, PoseIdx, EPoseCandidateFlags::DiscardedBy_BlockTransition, DynamicWeightsSqrt, ContinuingPoseCostAddend, ContinuingInteractionCostAddend);
+			}
+			else
+			{
+				SearchContext.Track(Database, PoseIdx, EPoseCandidateFlags::DiscardedBy_PoseFilter, DynamicWeightsSqrt, ContinuingPoseCostAddend, ContinuingInteractionCostAddend);
+			}
+#endif // UE_POSE_SEARCH_TRACE_ENABLED
+			return false;
+		}
+	}
+
+	PRAGMA_DISABLE_DEPRECATION_WARNINGS
+	for (const IPoseSearchFilter* Filter : Filters_DEPRECATED)
+	{
+		if (!Filter->IsFilterValid(PoseValues, QueryValues, PoseIdx, SearchIndex.PoseMetadata[PoseIdx]))
+		{
+#if UE_POSE_SEARCH_TRACE_ENABLED
+			SearchContext.Track(Database, PoseIdx, EPoseCandidateFlags::DiscardedBy_PoseFilter, DynamicWeightsSqrt, ContinuingPoseCostAddend, ContinuingInteractionCostAddend);
+#endif // UE_POSE_SEARCH_TRACE_ENABLED
+			return false;
+		}
+	}
+	PRAGMA_ENABLE_DEPRECATION_WARNINGS
+		
+	return true;
+}
+
+} // namespace UE::PoseSearch
\ No newline at end of file
diff --git a/Engine/Plugins/Animation/PoseSearch/Source/Runtime/Private/PoseSearchHistoryCollectorAnimNodeLibrary.cpp b/Engine/Plugins/Animation/PoseSearch/Source/Runtime/Private/PoseSearchHistoryCollectorAnimNodeLibrary.cpp
--- a/Engine/Plugins/Animation/PoseSearch/Source/Runtime/Private/PoseSearchHistoryCollectorAnimNodeLibrary.cpp
+++ b/Engine/Plugins/Animation/PoseSearch/Source/Runtime/Private/PoseSearchHistoryCollectorAnimNodeLibrary.cpp
@@ -32,7 +32,10 @@ void UPoseSearchHistoryCollectorAnimNodeLibrary::GetPoseHistoryNodeTransformTraj
 {
 	if (FAnimNode_PoseSearchHistoryCollector* PoseSearchHistoryCollectorNodePtr = PoseSearchHistoryCollectorNode.GetAnimNodePtr<FAnimNode_PoseSearchHistoryCollector>())
 	{
-		Trajectory = PoseSearchHistoryCollectorNodePtr->GetPoseHistory().GetTrajectory();
+		if (const UE::PoseSearch::FPoseHistory* PoseHistory = PoseSearchHistoryCollectorNodePtr->GetPoseHistoryPtr())
+		{
+			Trajectory = PoseHistory->GetTrajectory();
+		}
 	}
 }
 
@@ -40,24 +43,9 @@ void UPoseSearchHistoryCollectorAnimNodeLibrary::SetPoseHistoryNodeTransformTraj
 {
 	if (FAnimNode_PoseSearchHistoryCollector* PoseSearchHistoryCollectorNodePtr = PoseSearchHistoryCollectorNode.GetAnimNodePtr<FAnimNode_PoseSearchHistoryCollector>())
 	{
-		PoseSearchHistoryCollectorNodePtr->GetPoseHistory().SetTrajectory(Trajectory);
+		if (UE::PoseSearch::FPoseHistory* PoseHistory = PoseSearchHistoryCollectorNodePtr->GetPoseHistoryPtr())
+		{
+			PoseHistory->SetTrajectory(Trajectory);
+		}
 	}
 }
-
-PRAGMA_DISABLE_DEPRECATION_WARNINGS
-void UPoseSearchHistoryCollectorAnimNodeLibrary::GetPoseHistoryNodeTrajectory(const FPoseSearchHistoryCollectorAnimNodeReference& PoseSearchHistoryCollectorNode, FPoseSearchQueryTrajectory& Trajectory)
-{
-	if (FAnimNode_PoseSearchHistoryCollector* PoseSearchHistoryCollectorNodePtr = PoseSearchHistoryCollectorNode.GetAnimNodePtr<FAnimNode_PoseSearchHistoryCollector>())
-	{
-		Trajectory = PoseSearchHistoryCollectorNodePtr->GetPoseHistory().GetTrajectory();
-	}
-}
-
-void UPoseSearchHistoryCollectorAnimNodeLibrary::SetPoseHistoryNodeTrajectory(const FPoseSearchHistoryCollectorAnimNodeReference& PoseSearchHistoryCollectorNode, const FPoseSearchQueryTrajectory& Trajectory)
-{
-	if (FAnimNode_PoseSearchHistoryCollector* PoseSearchHistoryCollectorNodePtr = PoseSearchHistoryCollectorNode.GetAnimNodePtr<FAnimNode_PoseSearchHistoryCollector>())
-	{
-		PoseSearchHistoryCollectorNodePtr->GetPoseHistory().SetTrajectory(Trajectory);
-	}
-}
-PRAGMA_ENABLE_DEPRECATION_WARNINGS
diff --git a/Engine/Plugins/Animation/PoseSearch/Source/Runtime/Private/PoseSearchInteractionSubsystem.cpp b/Engine/Plugins/Animation/PoseSearch/Source/Runtime/Private/PoseSearchInteractionSubsystem.cpp
--- a/Engine/Plugins/Animation/PoseSearch/Source/Runtime/Private/PoseSearchInteractionSubsystem.cpp
+++ b/Engine/Plugins/Animation/PoseSearch/Source/Runtime/Private/PoseSearchInteractionSubsystem.cpp
@@ -138,6 +138,8 @@ FString FInteractionAvailability::GetPoseHistoryName() const
 
 const UE::PoseSearch::IPoseHistory* FInteractionAvailability::GetPoseHistory(const UObject* AnimContext) const
 {
+	using namespace UE::PoseSearch;
+
 	if (PoseHistory)
 	{
 		return PoseHistory;
@@ -147,7 +149,10 @@ const UE::PoseSearch::IPoseHistory* FInteractionAvailability::GetPoseHistory(con
 	{
 		if (const FAnimNode_PoseSearchHistoryCollector_Base* PoseSearchHistoryCollector = UPoseSearchLibrary::FindPoseHistoryNode(PoseHistoryName, AnimInstance))
 		{
-			return &PoseSearchHistoryCollector->GetPoseHistory();
+			if (const FPoseHistory* PoseHistoryFromCollector = PoseSearchHistoryCollector->GetPoseHistoryPtr())
+			{
+				return PoseHistoryFromCollector;
+			}
 		}
 	}
 
diff --git a/Engine/Plugins/Animation/PoseSearch/Source/Runtime/Private/PoseSearchLibrary.cpp b/Engine/Plugins/Animation/PoseSearch/Source/Runtime/Private/PoseSearchLibrary.cpp
--- a/Engine/Plugins/Animation/PoseSearch/Source/Runtime/Private/PoseSearchLibrary.cpp
+++ b/Engine/Plugins/Animation/PoseSearch/Source/Runtime/Private/PoseSearchLibrary.cpp
@@ -1025,14 +1025,16 @@ void UPoseSearchLibrary::UpdateMotionMatchingState(
 		if (EAsyncBuildIndexResult::Success == FAsyncPoseSearchDatabasesManagement::RequestAsyncBuildIndex(CurResultDatabase, ERequestAsyncBuildFlag::ContinueRequest))
 #endif // WITH_EDITOR
 		{
-			FDebugDrawParams DrawParams(SearchContext.GetContexts(), SearchContext.GetPoseHistories(), SearchContext.GetRoleToIndex(), CurResultDatabase);
+			FDebugDrawParams DrawParams;
 			if (bDebugDrawCurResult)
 			{
+				DrawParams.Init(SearchContext.GetContexts(), SearchContext.GetPoseHistories(), SearchContext.GetRoleToIndex(), CurResultDatabase, FDebugDrawParams::EDrawContext::DrawCandidate);
 				DrawParams.DrawFeatureVector(SearchResult.PoseIdx);
 			}
 
 			if (bDebugDrawQuery)
 			{
+				DrawParams.Init(SearchContext.GetContexts(), SearchContext.GetPoseHistories(), SearchContext.GetRoleToIndex(), CurResultDatabase, FDebugDrawParams::EDrawContext::DrawQuery);
 				DrawParams.DrawFeatureVector(SearchContext.GetOrBuildQuery(CurResultDatabase->Schema));
 			}
 		}
@@ -1102,11 +1104,19 @@ void UPoseSearchLibrary::MotionMatch(
 	TArray<const IPoseHistory*, TInlineAllocator<PreallocatedRolesNum, TMemStackAllocator<>>> PoseHistories;
 	if (const FAnimNode_PoseSearchHistoryCollector_Base* PoseHistoryNode = FindPoseHistoryNode(PoseHistoryName, AnimInstance))
 	{
-		PoseHistories.Add(&PoseHistoryNode->GetPoseHistory());
+		if (const FPoseHistory* PoseHistory = PoseHistoryNode->GetPoseHistoryPtr())
+		{
+			PoseHistories.Add(PoseHistory);
+		}
+		else
+		{
+			UE_LOG(LogPoseSearch, Error, TEXT("UPoseSearchLibrary::MotionMatch - Couldn't find a valid pose history in the pose history node '%s'"), *PoseHistoryName.ToString());
+			return;
+		}
 	}
 	else
 	{
-		UE_LOG(LogPoseSearch, Error, TEXT("UPoseSearchLibrary::MotionMatch - Couldn't find pose history with name '%s'"), *PoseHistoryName.ToString());
+		UE_LOG(LogPoseSearch, Error, TEXT("UPoseSearchLibrary::MotionMatch - Couldn't find pose history node with name '%s'"), *PoseHistoryName.ToString());
 		return;
 	}
 
@@ -1377,21 +1387,17 @@ void UPoseSearchLibrary::MotionMatch(
 #if ENABLE_DRAW_DEBUG && ENABLE_ANIM_DEBUG
 	if (SearchResult.IsValid())
 	{
-		const bool bDrawMatch = GVarAnimMotionMatchDrawMatchEnable;
-		const bool bDrawquery = GVarAnimMotionMatchDrawQueryEnable;
-
-		if (bDrawMatch || bDrawquery)
+		FDebugDrawParams DrawParams;
+		if (GVarAnimMotionMatchDrawMatchEnable)
 		{
-			FDebugDrawParams DrawParams(SearchContext.GetContexts(), SearchContext.GetPoseHistories(), SearchContext.GetRoleToIndex(), SearchResult.Database.Get());
-			if (bDrawMatch)
-			{
-				DrawParams.DrawFeatureVector(SearchResult.PoseIdx);
-			}
+			DrawParams.Init(SearchContext.GetContexts(), SearchContext.GetPoseHistories(), SearchContext.GetRoleToIndex(), SearchResult.Database.Get(), FDebugDrawParams::EDrawContext::DrawCandidate);
+			DrawParams.DrawFeatureVector(SearchResult.PoseIdx);
+		}
 
-			if (bDrawquery)
-			{
-				DrawParams.DrawFeatureVector(SearchContext.GetOrBuildQuery(SearchResult.Database->Schema));
-			}
+		if (GVarAnimMotionMatchDrawQueryEnable)
+		{
+			DrawParams.Init(SearchContext.GetContexts(), SearchContext.GetPoseHistories(), SearchContext.GetRoleToIndex(), SearchResult.Database.Get(), FDebugDrawParams::EDrawContext::DrawQuery);
+			DrawParams.DrawFeatureVector(SearchContext.GetOrBuildQuery(SearchResult.Database->Schema));
 		}
 	}
 #endif // ENABLE_DRAW_DEBUG && ENABLE_ANIM_DEBUG
diff --git a/Engine/Plugins/Animation/PoseSearch/Source/Runtime/Private/PoseSearchModule.cpp b/Engine/Plugins/Animation/PoseSearch/Source/Runtime/Private/PoseSearchModule.cpp
--- a/Engine/Plugins/Animation/PoseSearch/Source/Runtime/Private/PoseSearchModule.cpp
+++ b/Engine/Plugins/Animation/PoseSearch/Source/Runtime/Private/PoseSearchModule.cpp
@@ -69,7 +69,7 @@ public:
 		const IPoseHistory* PoseHistory = nullptr;
 		if (FPoseHistoryProvider* PoseHistoryProvider = GraphContext.GetMessage<FPoseHistoryProvider>())
 		{
-			PoseHistory = &PoseHistoryProvider->GetPoseHistory();
+			PoseHistory = PoseHistoryProvider->GetPoseHistoryPtr();
 		}
 
 		UE::Anim::IPoseSearchProvider::FSearchResult ProviderResult;
diff --git a/Engine/Plugins/Animation/PoseSearch/Source/Runtime/Public/PoseSearch/AnimNode_PoseSearchHistoryCollector.h b/Engine/Plugins/Animation/PoseSearch/Source/Runtime/Public/PoseSearch/AnimNode_PoseSearchHistoryCollector.h
--- a/Engine/Plugins/Animation/PoseSearch/Source/Runtime/Public/PoseSearch/AnimNode_PoseSearchHistoryCollector.h
+++ b/Engine/Plugins/Animation/PoseSearch/Source/Runtime/Public/PoseSearch/AnimNode_PoseSearchHistoryCollector.h
@@ -113,7 +113,15 @@ public:
 	UE_API virtual void Update_AnyThread(const FAnimationUpdateContext& Context) override;
 	// End of FAnimNode_Base interface
 
+	// GetPoseHistory got deprecated in favor of GetPoseHistoryPtr, since FAnimNode_PoseSearchHistoryCollector_Base::PoseHistoryPtr can be invalid during blue print
+	// recompilation when the pose history node exists in the graph but on an inactive branch, for instance on a never activated branch of a blend by bool node
+	const UE::PoseSearch::FPoseHistory* GetPoseHistoryPtr() const { return PoseHistoryPtr.Get(); }
+	UE::PoseSearch::FPoseHistory* GetPoseHistoryPtr() { return PoseHistoryPtr.Get(); }
+
+	UE_DEPRECATED(5.8, "Use GetPoseHistoryPtr instead")
 	const UE::PoseSearch::FPoseHistory& GetPoseHistory() const { check(PoseHistoryPtr.IsValid()); return *PoseHistoryPtr.Get(); }
+	
+	UE_DEPRECATED(5.8, "Use GetPoseHistoryPtr instead")
 	UE::PoseSearch::FPoseHistory& GetPoseHistory() { check(PoseHistoryPtr.IsValid()); return *PoseHistoryPtr.Get(); }
 	
 	FPoseHistoryReference GetPoseHistoryReference() const { check(PoseHistoryPtr.IsValid()); return FPoseHistoryReference {PoseHistoryPtr}; }
@@ -136,12 +144,6 @@ protected:
 	// returns an array of skeleton bone indexes (encoded as FBoneIndexType)
 	TArray<FBoneIndexType> GetRequiredBones(const FAnimInstanceProxy* AnimInstanceProxy) const;
 
-	UE_DEPRECATED(5.6, "Use GetPoseHistory() instead")
-	UE::PoseSearch::FPoseHistory PoseHistory;
-
-	UE_DEPRECATED(5.6, "Use PoseHistoryPtr->bIsTrajectoryGeneratedBeforePreUpdate instead")
-	bool bIsTrajectoryGeneratedBeforePreUpdate = false;
-	
 private:
 	TSharedPtr<UE::PoseSearch::FGenerateTrajectoryPoseHistory, ESPMode::ThreadSafe> PoseHistoryPtr;
 };
diff --git a/Engine/Plugins/Animation/PoseSearch/Source/Runtime/Public/PoseSearch/PoseHistoryProvider.h b/Engine/Plugins/Animation/PoseSearch/Source/Runtime/Public/PoseSearch/PoseHistoryProvider.h
--- a/Engine/Plugins/Animation/PoseSearch/Source/Runtime/Public/PoseSearch/PoseHistoryProvider.h
+++ b/Engine/Plugins/Animation/PoseSearch/Source/Runtime/Public/PoseSearch/PoseHistoryProvider.h
@@ -19,7 +19,11 @@ class FPoseHistoryProvider : public UE::Anim::IGraphMessage
 
 public:
 	FPoseHistoryProvider(const FAnimNode_PoseSearchHistoryCollector_Base* InHistoryCollector) : HistoryCollector(InHistoryCollector) { check(HistoryCollector); }
+	
+	UE_DEPRECATED(5.8, "Use GetPoseHistoryPtr instead")
 	UE_API const IPoseHistory& GetPoseHistory() const;
+
+	UE_API const IPoseHistory* GetPoseHistoryPtr() const;
 	
 private:
 	const FAnimNode_PoseSearchHistoryCollector_Base* HistoryCollector = nullptr;
diff --git a/Engine/Plugins/Animation/PoseSearch/Source/Runtime/Public/PoseSearch/PoseSearchContext.h b/Engine/Plugins/Animation/PoseSearch/Source/Runtime/Public/PoseSearch/PoseSearchContext.h
--- a/Engine/Plugins/Animation/PoseSearch/Source/Runtime/Public/PoseSearch/PoseSearchContext.h
+++ b/Engine/Plugins/Animation/PoseSearch/Source/Runtime/Public/PoseSearch/PoseSearchContext.h
@@ -62,7 +62,13 @@ ENUM_CLASS_FLAGS(EPoseCandidateFlags);
 #if ENABLE_DRAW_DEBUG
 struct FDebugDrawParams
 {
-	UE_API FDebugDrawParams(TConstArrayView<FChooserEvaluationContext*> InAnimContexts, TConstArrayView<const IPoseHistory*> InPoseHistories, const FRoleToIndex& InRoleToIndex, const UPoseSearchDatabase* InDatabase);
+	enum EDrawContext
+	{
+		DrawQuery,
+		DrawCandidate
+	};
+
+	UE_API FDebugDrawParams(TConstArrayView<FChooserEvaluationContext*> InAnimContexts, TConstArrayView<const IPoseHistory*> InPoseHistories, const FRoleToIndex& InRoleToIndex, const UPoseSearchDatabase* InDatabase, EDrawContext InDrawContext = EDrawContext::DrawQuery);
 	FDebugDrawParams() = default;
 
 	UE_API FDebugDrawParams(const FDebugDrawParams& Other);
@@ -70,7 +76,7 @@ struct FDebugDrawParams
 	UE_API FDebugDrawParams& operator=(const FDebugDrawParams& Other);
 	UE_API FDebugDrawParams& operator=(FDebugDrawParams&& Other);
 
-	UE_API void Init(TConstArrayView<FChooserEvaluationContext*> InAnimContexts, TConstArrayView<const IPoseHistory*> InPoseHistories, const FRoleToIndex& InRoleToIndex, const UPoseSearchDatabase* InDatabase);
+	UE_API void Init(TConstArrayView<FChooserEvaluationContext*> InAnimContexts, TConstArrayView<const IPoseHistory*> InPoseHistories, const FRoleToIndex& InRoleToIndex, const UPoseSearchDatabase* InDatabase, EDrawContext InDrawContext = EDrawContext::DrawQuery);
 
 	UE_API const FSearchIndex* GetSearchIndex() const;
 	UE_API const UPoseSearchSchema* GetSchema() const;
@@ -84,7 +90,7 @@ struct FDebugDrawParams
 	UE_API void DrawPoint(const FVector& Position, const FColor& Color, float Thickness = 6.f) const;
 	UE_API void DrawCircle(const FVector& Center, const FVector& UpVector, float Radius, int32 Segments, const FColor& Color, float Thickness = 0.f) const;
 	UE_API void DrawCircle(const FMatrix& TransformMatrix, float Radius, int32 Segments, const FColor& Color, float Thickness = 0.f) const;
-	UE_API void DrawWedge(const FVector& Origin, const FVector& Direction, float InnerRadius, float OuterRadius, float Width, int32 Segments, const FColor& Color, float Thickness = 0.f) const;
+	UE_API void DrawWedge(const FVector& Origin, const FVector& Direction, float InnerRadius, float OuterRadius, float WidthInDegrees, int32 Segments, const FColor& Color, float Thickness = 0.f) const;
 	UE_API void DrawSphere(const FVector& Center, float Radius, int32 Segments, const FColor& Color, float Thickness = 0.f) const;
 	UE_API void DrawCentripetalCatmullRomSpline(TConstArrayView<FVector> Points, TConstArrayView<FColor> Colors, float Alpha, int32 NumSamplesPerSegment, float Thickness = 1.f) const;
 	
@@ -93,6 +99,8 @@ struct FDebugDrawParams
 
 	UE_API bool IsAnyWeightRelevant(const UPoseSearchFeatureChannel* Channel) const;
 
+	EDrawContext GetDrawContext() const { return DrawContext; }
+
 private:
 	bool CanDraw() const;
 
@@ -104,6 +112,8 @@ private:
 	FRoleToIndex RoleToIndex;
 
 	const UPoseSearchDatabase* Database = nullptr;
+	EDrawContext DrawContext = EDrawContext::DrawCandidate;
+
 	TAlignedArray<float> DynamicWeightsSqrtBuffer;
 	TConstArrayView<float> DynamicWeightsSqrt;
 };
diff --git a/Engine/Plugins/Animation/PoseSearch/Source/Runtime/Public/PoseSearch/PoseSearchFeatureChannel.h b/Engine/Plugins/Animation/PoseSearch/Source/Runtime/Public/PoseSearch/PoseSearchFeatureChannel.h
--- a/Engine/Plugins/Animation/PoseSearch/Source/Runtime/Public/PoseSearch/PoseSearchFeatureChannel.h
+++ b/Engine/Plugins/Animation/PoseSearch/Source/Runtime/Public/PoseSearch/PoseSearchFeatureChannel.h
@@ -2,29 +2,18 @@
 
 #pragma once
 
-#include "IO/IoHash.h"
-#include "Interfaces/Interface_BoneReferenceSkeletonProvider.h"
 #include "DrawDebugHelpers.h"
+#include "Interfaces/Interface_BoneReferenceSkeletonProvider.h"
 #include "PoseSearch/PoseSearchRole.h"
+#include "PoseSearchFeatureVectorHelper.h"
+#include "PoseSearchFilter.h"
+#include "StructUtils/InstancedStruct.h"
 #include "PoseSearchFeatureChannel.generated.h"
 
 #define UE_API POSESEARCH_API
 
 class UPoseSearchSchema;
 
-UENUM()
-enum class EComponentStrippingVector : uint8
-{
-	// No component stripping.
-	None,
-
-	// Stripping X and Y components (matching only on the horizontal plane).
-	StripXY,
-
-	// Stripping Z (matching only vertically - caring only about the height of the feature).
-	StripZ,
-};
-
 UENUM()
 enum class EInputQueryPose : uint8
 {
@@ -80,43 +69,16 @@ typedef TStringBuilder<256> TLabelBuilder;
 
 #endif // WITH_EDITOR
 
-/** Helper class for extracting and encoding features into a float buffer */
-class FFeatureVectorHelper
-{
-public:
-	static UE_API int32 GetVectorCardinality(EComponentStrippingVector ComponentStrippingVector);
-	static UE_API void EncodeVector(TArrayView<float> Values, int32 DataOffset, const FVector& Vector, EComponentStrippingVector ComponentStrippingVector, bool bNormalize);
-	static UE_API FVector DecodeVector(TConstArrayView<float> Values, int32 DataOffset, EComponentStrippingVector ComponentStrippingVector);
-
-	static UE_API void EncodeVector2D(TArrayView<float> Values, int32 DataOffset, const FVector2D& Vector2D);
-	static UE_API FVector2D DecodeVector2D(TConstArrayView<float> Values, int32 DataOffset);
-
-	static UE_API void EncodeFloat(TArrayView<float> Values, int32 DataOffset, const float Value);
-	static UE_API float DecodeFloat(TConstArrayView<float> Values, int32 DataOffset);
-
-	static UE_API void Copy(TArrayView<float> Values, int32 DataOffset, int32 DataCardinality, TConstArrayView<float> OriginValues);
-};
-
 } // namespace UE::PoseSearch
 
-class IPoseSearchFilter
-{
-public:
-	virtual ~IPoseSearchFilter() {}
-
-	// if true this filter will be evaluated
-	virtual bool IsFilterActive() const { return false; }
-	
-	// if it returns false the pose candidate will be discarded
-	virtual bool IsFilterValid(TConstArrayView<float> PoseValues, TConstArrayView<float> QueryValues, int32 PoseIdx, const UE::PoseSearch::FPoseMetadata& Metadata) const { return true; }
-};
-
 //////////////////////////////////////////////////////////////////////////
 // Feature channels interface
+PRAGMA_DISABLE_DEPRECATION_WARNINGS
 UCLASS(MinimalAPI, Abstract, BlueprintType, EditInlineNew)
 class UPoseSearchFeatureChannel : public UObject, public IBoneReferenceSkeletonProvider, public IPoseSearchFilter
 {
 	GENERATED_BODY()
+PRAGMA_ENABLE_DEPRECATION_WARNINGS
 
 public:
 	int32 GetChannelCardinality() const { checkSlow(ChannelCardinality >= 0); return ChannelCardinality; }
@@ -134,12 +96,15 @@ public:
 
 	virtual void AddDependentChannels(UPoseSearchSchema* Schema) const {}
 
+	int32 GetNumFilters() const { return Filters.Num(); }
+	const FPoseSearchFilter* GetFilter(int32 FilterIndex) const { return Filters[FilterIndex].GetPtr<FPoseSearchFilter>(); }
+
 	virtual EPermutationTimeType GetPermutationTimeType() const { return EPermutationTimeType::UseSampleTime; }
 	static UE_API void GetPermutationTimeOffsets(EPermutationTimeType PermutationTimeType, float DesiredPermutationTimeOffset, float& OutPermutationSampleTimeOffset, float& OutPermutationOriginTimeOffset);
 
 #if ENABLE_DRAW_DEBUG
 	// Draw this channel's data for the given pose vector
-	virtual void DebugDraw(const UE::PoseSearch::FDebugDrawParams& DrawParams, TConstArrayView<float> PoseVector) const {}
+	virtual void DebugDraw(const UE::PoseSearch::FDebugDrawParams& DrawParams, TConstArrayView<float> PoseVector) const;
 #endif //ENABLE_DRAW_DEBUG
 
 #if WITH_EDITOR
@@ -168,6 +133,10 @@ public:
 	UE_API virtual const UE::PoseSearch::FRole GetDefaultRole() const;
 #endif
 
+#if WITH_EDITORONLY_DATA
+	virtual FLinearColor GetDebugColor() const { return FLinearColor::White; }
+#endif // WITH_EDITORONLY_DATA
+
 	// IBoneReferenceSkeletonProvider interface
 	// Note this function is exclusively for FBoneReference details customization
 	UE_API USkeleton* GetSkeleton(bool& bInvalidSkeletonIsError, const IPropertyHandle* PropertyHandle) override;
@@ -184,6 +153,8 @@ protected:
 
 	UPROPERTY(Transient)
 	int32 ChannelCardinality = INDEX_NONE;
-};
 
+	UPROPERTY(EditAnywhere, Category = "Filters", meta = (ExcludeBaseStruct, BaseStruct = "/Script/PoseSearch.PoseSearchFilter"))
+	TArray<FInstancedStruct> Filters;
+};
 #undef UE_API
diff --git a/Engine/Plugins/Animation/PoseSearch/Source/Runtime/Public/PoseSearch/PoseSearchFeatureChannel_Curve.h b/Engine/Plugins/Animation/PoseSearch/Source/Runtime/Public/PoseSearch/PoseSearchFeatureChannel_Curve.h
--- a/Engine/Plugins/Animation/PoseSearch/Source/Runtime/Public/PoseSearch/PoseSearchFeatureChannel_Curve.h
+++ b/Engine/Plugins/Animation/PoseSearch/Source/Runtime/Public/PoseSearch/PoseSearchFeatureChannel_Curve.h
@@ -66,6 +66,10 @@ public:
 	UE_API virtual UE::PoseSearch::TLabelBuilder& GetLabel(UE::PoseSearch::TLabelBuilder& LabelBuilder, UE::PoseSearch::ELabelFormat LabelFormat = UE::PoseSearch::ELabelFormat::Full_Horizontal) const override;
 	virtual FName GetNormalizationGroup() const override { return NormalizationGroup; }
 #endif
+
+#if WITH_EDITORONLY_DATA
+	virtual FLinearColor GetDebugColor() const override { return DebugColor; }
+#endif // WITH_EDITORONLY_DATA
 };
 
 #undef UE_API
diff --git a/Engine/Plugins/Animation/PoseSearch/Source/Runtime/Public/PoseSearch/PoseSearchFeatureChannel_Distance.h b/Engine/Plugins/Animation/PoseSearch/Source/Runtime/Public/PoseSearch/PoseSearchFeatureChannel_Distance.h
--- a/Engine/Plugins/Animation/PoseSearch/Source/Runtime/Public/PoseSearch/PoseSearchFeatureChannel_Distance.h
+++ b/Engine/Plugins/Animation/PoseSearch/Source/Runtime/Public/PoseSearch/PoseSearchFeatureChannel_Distance.h
@@ -121,6 +121,10 @@ public:
 	// IBoneReferenceSkeletonProvider interface
 	UE_API USkeleton* GetSkeleton(bool& bInvalidSkeletonIsError, const IPropertyHandle* PropertyHandle) override;
 #endif
+
+#if WITH_EDITORONLY_DATA
+	virtual FLinearColor GetDebugColor() const override { return DebugColor; }
+#endif // WITH_EDITORONLY_DATA
 };
 
 #undef UE_API
diff --git a/Engine/Plugins/Animation/PoseSearch/Source/Runtime/Public/PoseSearch/PoseSearchFeatureChannel_FilterCrashingLegs.h b/Engine/Plugins/Animation/PoseSearch/Source/Runtime/Public/PoseSearch/PoseSearchFeatureChannel_FilterCrashingLegs.h
--- a/Engine/Plugins/Animation/PoseSearch/Source/Runtime/Public/PoseSearch/PoseSearchFeatureChannel_FilterCrashingLegs.h
+++ b/Engine/Plugins/Animation/PoseSearch/Source/Runtime/Public/PoseSearch/PoseSearchFeatureChannel_FilterCrashingLegs.h
@@ -74,6 +74,7 @@ public:
 	virtual const UE::PoseSearch::FRole GetDefaultRole() const override { return SampleRole; }
 #endif
 
+	// @todo: reimplement this using the new FPoseSearchFilter
 	// IPoseSearchFilter interface
 	UE_API virtual bool IsFilterActive() const override;
 	UE_API virtual bool IsFilterValid(TConstArrayView<float> PoseValues, TConstArrayView<float> QueryValues, int32 PoseIdx, const UE::PoseSearch::FPoseMetadata& Metadata) const override;
diff --git a/Engine/Plugins/Animation/PoseSearch/Source/Runtime/Public/PoseSearch/PoseSearchFeatureChannel_Heading.h b/Engine/Plugins/Animation/PoseSearch/Source/Runtime/Public/PoseSearch/PoseSearchFeatureChannel_Heading.h
--- a/Engine/Plugins/Animation/PoseSearch/Source/Runtime/Public/PoseSearch/PoseSearchFeatureChannel_Heading.h
+++ b/Engine/Plugins/Animation/PoseSearch/Source/Runtime/Public/PoseSearch/PoseSearchFeatureChannel_Heading.h
@@ -124,6 +124,10 @@ public:
 	UE_API USkeleton* GetSkeleton(bool& bInvalidSkeletonIsError, const IPropertyHandle* PropertyHandle) override;
 #endif
 
+#if WITH_EDITORONLY_DATA
+	virtual FLinearColor GetDebugColor() const override { return DebugColor; }
+#endif // WITH_EDITORONLY_DATA
+
 	UE_API FVector GetAxis(const FQuat& Rotation) const;
 	static UE_API void FindOrAddToSchema(UPoseSearchSchema* Schema, float SampleTimeOffset, const FName& BoneName, const UE::PoseSearch::FRole& Role, EHeadingAxis HeadingAxis = EHeadingAxis::X, EPermutationTimeType PermutationTimeType = EPermutationTimeType::UseSampleTime);
 };
diff --git a/Engine/Plugins/Animation/PoseSearch/Source/Runtime/Public/PoseSearch/PoseSearchFeatureChannel_Phase.h b/Engine/Plugins/Animation/PoseSearch/Source/Runtime/Public/PoseSearch/PoseSearchFeatureChannel_Phase.h
--- a/Engine/Plugins/Animation/PoseSearch/Source/Runtime/Public/PoseSearch/PoseSearchFeatureChannel_Phase.h
+++ b/Engine/Plugins/Animation/PoseSearch/Source/Runtime/Public/PoseSearch/PoseSearchFeatureChannel_Phase.h
@@ -62,6 +62,10 @@ public:
 	virtual FName GetNormalizationGroup() const override { return NormalizationGroup; }
 	virtual const UE::PoseSearch::FRole GetDefaultRole() const override { return SampleRole; }
 #endif
+
+#if WITH_EDITORONLY_DATA
+	virtual FLinearColor GetDebugColor() const override { return DebugColor; }
+#endif // WITH_EDITORONLY_DATA
 };
 
 #undef UE_API
diff --git a/Engine/Plugins/Animation/PoseSearch/Source/Runtime/Public/PoseSearch/PoseSearchFeatureChannel_Position.h b/Engine/Plugins/Animation/PoseSearch/Source/Runtime/Public/PoseSearch/PoseSearchFeatureChannel_Position.h
--- a/Engine/Plugins/Animation/PoseSearch/Source/Runtime/Public/PoseSearch/PoseSearchFeatureChannel_Position.h
+++ b/Engine/Plugins/Animation/PoseSearch/Source/Runtime/Public/PoseSearch/PoseSearchFeatureChannel_Position.h
@@ -75,13 +75,6 @@ public:
 	UPROPERTY(EditAnywhere, Category = "Settings")
 	EPermutationTimeType PermutationTimeType = EPermutationTimeType::UseSampleTime;
 
-	// Experimental, this feature might be removed without warning, not for production use
-	// During pose selection if the squared distance between query versus candidate poses for this 
-	// position channel is greater than MaxPositionDistanceSquared the candidate will be discarded.
-	// the filtering will be enabled only for MaxPositionDistanceSquared > 0
-	UPROPERTY(EditAnywhere, Category = "Experimental|Filter")
-	float MaxPositionDistanceSquared = 0.f;
-
 	// Experimental, this feature might be removed without warning, not for production use
 	// UPoseSearchFeatureChannel_Position represents / calculates the position (displacement) of Bone relative to OriginBone
 	// if bNormalizeDisplacement is true, the displacement will be normalized. if zero it'll stay zero, so it'll be like matching the 
@@ -112,12 +105,6 @@ public:
 	virtual EPermutationTimeType GetPermutationTimeType() const override { return PermutationTimeType; }
 	UE_API virtual void AddDependentChannels(UPoseSearchSchema* Schema) const override;
 
-	// IPoseSearchFilter interface
-	// Experimental, this feature might be removed without warning, not for production use
-	virtual bool IsFilterActive() const override { return MaxPositionDistanceSquared > 0.f; }
-	// Experimental, this feature might be removed without warning, not for production use
-	UE_API virtual bool IsFilterValid(TConstArrayView<float> PoseValues, TConstArrayView<float> QueryValues, int32 PoseIdx, const UE::PoseSearch::FPoseMetadata& Metadata) const override;
-
 #if ENABLE_DRAW_DEBUG
 	UE_API virtual void DebugDraw(const UE::PoseSearch::FDebugDrawParams& DrawParams, TConstArrayView<float> PoseVector) const override;
 #endif // ENABLE_DRAW_DEBUG
@@ -132,6 +119,10 @@ public:
 	UE_API USkeleton* GetSkeleton(bool& bInvalidSkeletonIsError, const IPropertyHandle* PropertyHandle) override;
 #endif
 
+#if WITH_EDITORONLY_DATA
+	virtual FLinearColor GetDebugColor() const override { return DebugColor; }
+#endif // WITH_EDITORONLY_DATA
+
 	static UE_API void FindOrAddToSchema(UPoseSearchSchema* Schema, float SampleTimeOffset, const FName& BoneName, const UE::PoseSearch::FRole& Role, EPermutationTimeType PermutationTimeType = EPermutationTimeType::UseSampleTime);
 };
 
diff --git a/Engine/Plugins/Animation/PoseSearch/Source/Runtime/Public/PoseSearch/PoseSearchFeatureChannel_Velocity.h b/Engine/Plugins/Animation/PoseSearch/Source/Runtime/Public/PoseSearch/PoseSearchFeatureChannel_Velocity.h
--- a/Engine/Plugins/Animation/PoseSearch/Source/Runtime/Public/PoseSearch/PoseSearchFeatureChannel_Velocity.h
+++ b/Engine/Plugins/Animation/PoseSearch/Source/Runtime/Public/PoseSearch/PoseSearchFeatureChannel_Velocity.h
@@ -118,6 +118,10 @@ public:
 	// IBoneReferenceSkeletonProvider interface
 	UE_API USkeleton* GetSkeleton(bool& bInvalidSkeletonIsError, const IPropertyHandle* PropertyHandle) override;
 #endif
+
+#if WITH_EDITORONLY_DATA
+	virtual FLinearColor GetDebugColor() const override { return DebugColor; }
+#endif // WITH_EDITORONLY_DATA
 };
 
 #undef UE_API
diff --git a/Engine/Plugins/Animation/PoseSearch/Source/Runtime/Public/PoseSearch/PoseSearchFeatureVectorHelper.h b/Engine/Plugins/Animation/PoseSearch/Source/Runtime/Public/PoseSearch/PoseSearchFeatureVectorHelper.h
new file mode 100644
--- /dev/null
+++ b/Engine/Plugins/Animation/PoseSearch/Source/Runtime/Public/PoseSearch/PoseSearchFeatureVectorHelper.h
@@ -0,0 +1,44 @@
+// Copyright Epic Games, Inc. All Rights Reserved.
+
+#pragma once
+
+#include "PoseSearchFeatureVectorHelper.generated.h"
+
+#define UE_API POSESEARCH_API
+
+UENUM()
+enum class EComponentStrippingVector : uint8
+{
+	// No component stripping.
+	None,
+
+	// Stripping X and Y components (matching only on the horizontal plane).
+	StripXY,
+
+	// Stripping Z (matching only vertically - caring only about the height of the feature).
+	StripZ,
+};
+
+namespace UE::PoseSearch
+{
+
+/** Helper class for extracting and encoding features into a float buffer */
+class FFeatureVectorHelper
+{
+public:
+	static UE_API int32 GetVectorCardinality(EComponentStrippingVector ComponentStrippingVector);
+	static UE_API void EncodeVector(TArrayView<float> Values, int32 DataOffset, const FVector& Vector, EComponentStrippingVector ComponentStrippingVector, bool bNormalize);
+	static UE_API FVector DecodeVector(TConstArrayView<float> Values, int32 DataOffset, EComponentStrippingVector ComponentStrippingVector);
+
+	static UE_API void EncodeVector2D(TArrayView<float> Values, int32 DataOffset, const FVector2D& Vector2D);
+	static UE_API FVector2D DecodeVector2D(TConstArrayView<float> Values, int32 DataOffset);
+
+	static UE_API void EncodeFloat(TArrayView<float> Values, int32 DataOffset, const float Value);
+	static UE_API float DecodeFloat(TConstArrayView<float> Values, int32 DataOffset);
+
+	static UE_API void Copy(TArrayView<float> Values, int32 DataOffset, int32 DataCardinality, TConstArrayView<float> OriginValues);
+};
+
+} // namespace UE::PoseSearch
+
+#undef UE_API
diff --git a/Engine/Plugins/Animation/PoseSearch/Source/Runtime/Public/PoseSearch/PoseSearchFilter.h b/Engine/Plugins/Animation/PoseSearch/Source/Runtime/Public/PoseSearch/PoseSearchFilter.h
new file mode 100644
--- /dev/null
+++ b/Engine/Plugins/Animation/PoseSearch/Source/Runtime/Public/PoseSearch/PoseSearchFilter.h
@@ -0,0 +1,245 @@
+// Copyright Epic Games, Inc. All Rights Reserved.
+
+#pragma once
+
+#include "CoreMinimal.h"
+#include "DrawDebugHelpers.h"
+#include "PoseSearch/PoseSearchIndex.h"
+#include "PoseSearchFilter.generated.h"
+
+#define UE_API POSESEARCH_API
+
+class UPoseSearchDatabase;
+class UPoseSearchFeatureChannel;
+class UPoseSearchSchema;
+
+namespace UE::PoseSearch
+{
+	struct FDebugDrawParams;
+	struct FPoseMetadata;
+	struct FSearchContext;
+} // namespace UE::PoseSearch
+
+class POSESEARCH_API UE_DEPRECATED(5.8, "Use FPoseSearchFilter instead") IPoseSearchFilter
+{
+public:
+	virtual ~IPoseSearchFilter() {}
+
+	// if true this filter will be evaluated
+	UE_DEPRECATED(5.8, "Use FPoseSearchFilter::GetRequiredDataSize instead")
+	virtual bool IsFilterActive() const { return false; }
+	
+	// if it returns false the pose candidate will be discarded
+	UE_DEPRECATED(5.8, "Use FPoseSearchFilter::IsFilterValid instead")
+	virtual bool IsFilterValid(TConstArrayView<float> PoseValues, TConstArrayView<float> QueryValues, int32 PoseIdx, const UE::PoseSearch::FPoseMetadata& Metadata) const { return true; }
+};
+
+USTRUCT(Experimental)
+struct FPoseSearchFilter
+{
+	GENERATED_BODY()
+
+	enum { InvalidFilterDataSize = -1 };
+	enum { DataAlignment = 16 };
+
+	virtual ~FPoseSearchFilter() {}
+
+	// returns require data size in bytes for this filter. InvalidFilterDataSize if the filter is NOT active
+	virtual int32 GetRequiredDataSize(const UPoseSearchFeatureChannel* Channel) const { return InvalidFilterDataSize; }
+
+	// called to initialize Data with this FPoseSearchFilter data  only if GetRequiredDataSize returns >= 0. Data.Num() == GetRequiredDataSize()
+	// returns successful it returns true, or else DeinitializeData will be called right away and IsFilterValid will NOT be called
+	// it is safe to use Data to perform a new inplace in InitializeData and call the destructor in DeinitializeData.. but not recommended
+	virtual bool InitializeData(TArrayView<uint8> Data, const UPoseSearchFeatureChannel* Channel, TConstArrayView<float> QueryValues, UE::PoseSearch::FSearchContext& SearchContext) const { return true; }
+
+	// called to deinitialize Data with this FPoseSearchFilter data  only if GetRequiredDataSize returns >= 0. Data.Num() == GetRequiredDataSize()
+	// it is safe to use Data to perform a new inplace in InitializeData and call the destructor in DeinitializeData.. but not recommended
+	virtual void DeinitializeData(TArrayView<uint8> Data, const UPoseSearchFeatureChannel* Channel, UE::PoseSearch::FSearchContext& SearchContext) const {}
+	
+	// called for EVERY evaluated pose (be sure it's as fast as possible). if it returns false the pose candidate will be discarded
+	virtual bool IsFilterValid(TConstArrayView<float> PoseValues, TConstArrayView<float> QueryValues, int32 PoseIdx, const UE::PoseSearch::FPoseMetadata& Metadata, TConstArrayView<uint8> Data, const UPoseSearchFeatureChannel* Channel) const { return true; }
+
+#if ENABLE_DRAW_DEBUG
+	virtual void DebugDraw(const UE::PoseSearch::FDebugDrawParams& DrawParams, TConstArrayView<float> PoseVector, const UPoseSearchFeatureChannel* Channel) const {}
+#endif // ENABLE_DRAW_DEBUG
+};
+
+// debug FPoseSearchFilter to test InitializeData / DeinitializeData by padding the data by 1 byte
+USTRUCT(Experimental)
+struct POSESEARCH_API FPoseSearchFilterPadding : public FPoseSearchFilter
+{
+	GENERATED_BODY()
+
+	virtual int32 GetRequiredDataSize(const UPoseSearchFeatureChannel* Channel) const override;
+};
+
+// FPoseSearchFilterMaxPosition is invalid if the query for the associated UPoseSearchFeatureChannel_Position is too far from the candidate pose
+USTRUCT(Experimental)
+struct POSESEARCH_API FPoseSearchFilterMaxPosition : public FPoseSearchFilter
+{
+	GENERATED_BODY()
+
+	virtual int32 GetRequiredDataSize(const UPoseSearchFeatureChannel* Channel) const override;
+	virtual bool InitializeData(TArrayView<uint8> Data, const UPoseSearchFeatureChannel* Channel, TConstArrayView<float> QueryValues, UE::PoseSearch::FSearchContext& SearchContext) const override;
+	virtual void DeinitializeData(TArrayView<uint8> Data, const UPoseSearchFeatureChannel* Channel, UE::PoseSearch::FSearchContext& SearchContext) const override;
+	virtual bool IsFilterValid(TConstArrayView<float> PoseValues, TConstArrayView<float> QueryValues, int32 PoseIdx, const UE::PoseSearch::FPoseMetadata& Metadata, TConstArrayView<uint8> Data, const UPoseSearchFeatureChannel* Channel) const override;
+	
+#if ENABLE_DRAW_DEBUG
+	virtual void DebugDraw(const UE::PoseSearch::FDebugDrawParams& DrawParams, TConstArrayView<float> PoseValues, const UPoseSearchFeatureChannel* Channel) const override;
+#endif // ENABLE_DRAW_DEBUG
+
+	// During pose selection if the squared distance between query versus candidate poses for the associated 
+	// UPoseSearchFeatureChannel_Position is greater than MaxPositionDistanceSquared the candidate will be discarded.
+	// The filtering will be enabled only for MaxPositionDistanceSquared != 0
+	UPROPERTY(EditAnywhere, Category = "Experimental|Filter", meta = (ClampMin = "0", UIMin = "0"))
+	float MaxPositionDistanceSquared = 0.f;
+
+private:
+	struct FData
+	{
+		FVector Query;
+	};
+	static_assert(alignof(FData) <= DataAlignment);
+};
+
+// FPoseSearchFilterWedge is invalid if the query for the associated UPoseSearchFeatureChannel_Position doesn't land in the defined wedge relative to the trajectory
+USTRUCT(Experimental)
+struct POSESEARCH_API FPoseSearchFilterWedge : public FPoseSearchFilter
+{
+	GENERATED_BODY()
+
+	virtual int32 GetRequiredDataSize(const UPoseSearchFeatureChannel* Channel) const override;
+	virtual bool InitializeData(TArrayView<uint8> Data, const UPoseSearchFeatureChannel* Channel, TConstArrayView<float> QueryValues, UE::PoseSearch::FSearchContext& SearchContext) const override;
+	virtual void DeinitializeData(TArrayView<uint8> Data, const UPoseSearchFeatureChannel* Channel, UE::PoseSearch::FSearchContext& SearchContext) const override;
+	virtual bool IsFilterValid(TConstArrayView<float> PoseValues, TConstArrayView<float> QueryValues, int32 PoseIdx, const UE::PoseSearch::FPoseMetadata& Metadata, TConstArrayView<uint8> Data, const UPoseSearchFeatureChannel* Channel) const override;
+
+#if ENABLE_DRAW_DEBUG
+	virtual void DebugDraw(const UE::PoseSearch::FDebugDrawParams& DrawParams, TConstArrayView<float> PoseValues, const UPoseSearchFeatureChannel* Channel) const override;
+#endif // ENABLE_DRAW_DEBUG
+
+	UPROPERTY(EditAnywhere, Category = "Experimental|Filter", meta = (ClampToMinMaxLimits))
+	FFloatInterval RadiusOffsets = FFloatInterval(-100.f, 100.f);
+	
+	// wedge filter width in degrees
+	UPROPERTY(EditAnywhere, Category = "Experimental|Filter", meta = (ClampMin = "0"))
+	float WedgeWidthInDegrees = 70.f;
+
+private:
+	struct FData
+	{
+		FVector NormalizedQuery;
+		float InnerRadius;
+		float OuterRadius;
+		float DotWedgeWidth;
+	};
+	static_assert(alignof(FData) <= DataAlignment);
+};
+
+namespace UE::PoseSearch
+{
+// Experimental, this feature might be removed without warning, not for production use
+struct FSelectableAssetIdxFilter : public FPoseSearchFilter
+{
+	const FSelectableAssetIdxFilter& Init(TConstArrayView<int32> InSelectableAssetIdxFilter)
+	{
+		check(Algo::IsSorted(InSelectableAssetIdxFilter));
+		SelectableAssetIdxFilter = InSelectableAssetIdxFilter;
+		return *this;
+	}
+
+	virtual int32 GetRequiredDataSize(const UPoseSearchFeatureChannel* Channel) const override
+	{
+		return SelectableAssetIdxFilter.IsEmpty() ? InvalidFilterDataSize : 0;
+	}
+
+	virtual bool IsFilterValid(TConstArrayView<float> PoseValues, TConstArrayView<float> QueryValues, int32 PoseIdx, const UE::PoseSearch::FPoseMetadata& Metadata, TConstArrayView<uint8> Data, const UPoseSearchFeatureChannel* Channel) const override
+	{
+		return Algo::BinarySearch(SelectableAssetIdxFilter, int32(Metadata.GetAssetIndex())) != INDEX_NONE;
+	}
+
+	TConstArrayView<int32> SelectableAssetIdxFilter;
+};
+
+// Experimental, this feature might be removed without warning, not for production use
+struct FNonSelectableIdxFilter : public FPoseSearchFilter
+{
+	const FNonSelectableIdxFilter& Init(TConstArrayView<int32> InNonSelectableIdx)
+	{
+		check(Algo::IsSorted(InNonSelectableIdx));
+		NonSelectableIdx = InNonSelectableIdx;
+		return *this;
+	}
+
+	virtual int32 GetRequiredDataSize(const UPoseSearchFeatureChannel* Channel) const override
+	{
+		return NonSelectableIdx.IsEmpty() ? InvalidFilterDataSize : 0;
+	}
+	virtual bool IsFilterValid(TConstArrayView<float> PoseValues, TConstArrayView<float> QueryValues, int32 PoseIdx, const UE::PoseSearch::FPoseMetadata& Metadata, TConstArrayView<uint8> Data, const UPoseSearchFeatureChannel* Channel) const override
+	{
+		return Algo::BinarySearch(NonSelectableIdx, PoseIdx) == INDEX_NONE;
+	}
+
+	TConstArrayView<int32> NonSelectableIdx;
+};
+
+// Experimental, this feature might be removed without warning, not for production use
+struct FBlockTransitionFilter : public FPoseSearchFilter
+{
+	virtual int32 GetRequiredDataSize(const UPoseSearchFeatureChannel* Channel) const override
+	{
+		return 0;
+	}
+
+	virtual bool IsFilterValid(TConstArrayView<float> PoseValues, TConstArrayView<float> QueryValues, int32 PoseIdx, const UE::PoseSearch::FPoseMetadata& Metadata, TConstArrayView<uint8> Data, const UPoseSearchFeatureChannel* Channel) const override
+	{
+		return !Metadata.IsBlockTransition();
+	}
+};
+
+// Experimental, this feature might be removed without warning, not for production use
+struct FSearchFilters
+{
+	FSearchFilters(const UPoseSearchSchema* Schema, TConstArrayView<int32> NonSelectableIdx, TConstArrayView<int32> SelectableAssetIdx,
+		bool bAddBlockTransitionFilter, TConstArrayView<float> QueryValues, FSearchContext& InSearchContext);
+	~FSearchFilters();
+
+	// @todo: template this with bAlignedAndPadded to be able to use faster ComparePoses
+	bool AreFiltersValid(const FSearchIndex& SearchIndex, TConstArrayView<float> PoseValues, TConstArrayView<float> QueryValues, TConstArrayView<float> DynamicWeightsSqrt, int32 PoseIdx
+#if UE_POSE_SEARCH_TRACE_ENABLED
+		, float ContinuingPoseCostAddend, float ContinuingInteractionCostAddend, const UPoseSearchDatabase* Database
+#endif // UE_POSE_SEARCH_TRACE_ENABLED
+		) const;
+
+	FSearchContext& GetSearchContext() const { return SearchContext; }
+
+private:
+		
+	FSearchContext& SearchContext;
+	FNonSelectableIdxFilter NonSelectableIdxFilter;
+	FSelectableAssetIdxFilter SelectableAssetIdxFilter;
+	FBlockTransitionFilter BlockTransitionFilter;
+	
+	PRAGMA_DISABLE_DEPRECATION_WARNINGS
+	TArray<const IPoseSearchFilter*, TInlineAllocator<64, TMemStackAllocator<>>> Filters_DEPRECATED;
+	PRAGMA_ENABLE_DEPRECATION_WARNINGS
+
+	struct FFilterContext
+	{
+		FFilterContext(const FPoseSearchFilter* InFilter, const UPoseSearchFeatureChannel* InChannel, TArrayView<uint8> InData)
+			: Filter(InFilter)
+			, Channel(InChannel)
+			, Data(InData)
+		{
+		}
+
+		const FPoseSearchFilter* Filter;
+		const UPoseSearchFeatureChannel* Channel;
+		TArrayView<uint8> Data;
+	};
+	TArray<FFilterContext, TInlineAllocator<64, TMemStackAllocator<>>> FilterContexts;
+	TArray<uint8, TMemStackAllocator<FPoseSearchFilter::DataAlignment>> FiltersData;
+};
+
+} // namespace UE::PoseSearch
+
+#undef UE_API
diff --git a/Engine/Plugins/Animation/PoseSearch/Source/Runtime/Public/PoseSearch/PoseSearchHistoryCollectorAnimNodeLibrary.h b/Engine/Plugins/Animation/PoseSearch/Source/Runtime/Public/PoseSearch/PoseSearchHistoryCollectorAnimNodeLibrary.h
--- a/Engine/Plugins/Animation/PoseSearch/Source/Runtime/Public/PoseSearch/PoseSearchHistoryCollectorAnimNodeLibrary.h
+++ b/Engine/Plugins/Animation/PoseSearch/Source/Runtime/Public/PoseSearch/PoseSearchHistoryCollectorAnimNodeLibrary.h
@@ -47,16 +47,6 @@ public:
 
 	UFUNCTION(BlueprintPure, Category = "Animation|Pose Search|Experimental", meta = (BlueprintThreadSafe))
 	static UE_API FPoseHistoryReference GetPoseHistoryReference(const FPoseSearchHistoryCollectorAnimNodeReference& Node);
-	
-	PRAGMA_DISABLE_DEPRECATION_WARNINGS
-	
-	UE_DEPRECATED(5.6, "Please use GetPoseHistoryNodeTrajectory that takes in FTransformTrajectory instead")
-	static UE_API void GetPoseHistoryNodeTrajectory(const FPoseSearchHistoryCollectorAnimNodeReference& PoseSearchHistoryCollectorNode, FPoseSearchQueryTrajectory& Trajectory);
-	
-	UE_DEPRECATED(5.6, "Please use GetPoseHistoryNodeTrajectory that takes in FTransformTrajectory instead")
-	static UE_API void SetPoseHistoryNodeTrajectory(const FPoseSearchHistoryCollectorAnimNodeReference& PoseSearchHistoryCollectorNode, const FPoseSearchQueryTrajectory& Trajectory);
-	
-	PRAGMA_ENABLE_DEPRECATION_WARNINGS
 };
 
 #undef UE_API
diff --git a/Engine/Plugins/Experimental/MovieScenePoseSearchTracks/Source/MovieScenePoseSearchTracks/Private/Sections/MovieSceneStitchAnimSection.cpp b/Engine/Plugins/Experimental/MovieScenePoseSearchTracks/Source/MovieScenePoseSearchTracks/Private/Sections/MovieSceneStitchAnimSection.cpp
--- a/Engine/Plugins/Experimental/MovieScenePoseSearchTracks/Source/MovieScenePoseSearchTracks/Private/Sections/MovieSceneStitchAnimSection.cpp
+++ b/Engine/Plugins/Experimental/MovieScenePoseSearchTracks/Source/MovieScenePoseSearchTracks/Private/Sections/MovieSceneStitchAnimSection.cpp
@@ -85,7 +85,7 @@ void FMovieSceneStitchAnimEvaluationTask::Execute(UE::UAF::FEvaluationVM& VM) co
 	{
 		if (const FAnimNode_PoseSearchHistoryCollector_Base* PoseHistoryNode = UPoseSearchLibrary::FindPoseHistoryNode(StitchData.OverridePoseHistoryTag, AnimInstance))
 		{
-			PoseHistory = &PoseHistoryNode->GetPoseHistory();
+			PoseHistory = PoseHistoryNode->GetPoseHistoryPtr();
 		}
 	}
 
