Fix for inconsistent past velocities in motion matching due to history points being offset by delta time, causing noise in past velocity when delta time is not constant

diff --git a/Engine/Plugins/Animation/PoseSearch/Source/Runtime/Private/PoseSearchTrajectoryLibrary.cpp b/Engine/Plugins/Animation/PoseSearch/Source/Runtime/Private/PoseSearchTrajectoryLibrary.cpp
--- a/Engine/Plugins/Animation/PoseSearch/Source/Runtime/Private/PoseSearchTrajectoryLibrary.cpp
+++ b/Engine/Plugins/Animation/PoseSearch/Source/Runtime/Private/PoseSearchTrajectoryLibrary.cpp
@@ -317,13 +317,15 @@ void UPoseSearchTrajectoryLibrary::UpdateHistory_TransformHistory(FTransformTraj
 		const FVector CurrentTranslationFromMover = CurrentVelocity * DeltaTime;
 		const FVector TranslationSinceLastFrame = CurrentPosition - Trajectory.Samples[NumHistorySamples].Position;
 		const FVector InferredGroundTranslation = TranslationSinceLastFrame - CurrentTranslationFromMover;
+		float TimeShift = Trajectory.Samples[NumHistorySamples].TimeInSeconds;	// This works out as the delta time of the last frame  
+		
 
 		// Shift history Samples when it's time to record a new one.
 		if (SecondsPerHistorySample <= 0.f || FMath::Abs(Trajectory.Samples[NumHistorySamples - 1].TimeInSeconds) >= SecondsPerHistorySample)
 		{
 			for (int32 Index = 0; Index < NumHistorySamples - 1; ++Index)
 			{
-				Trajectory.Samples[Index].TimeInSeconds = Trajectory.Samples[Index + 1].TimeInSeconds - DeltaTime;
+				Trajectory.Samples[Index].TimeInSeconds = Trajectory.Samples[Index + 1].TimeInSeconds - TimeShift;
 				Trajectory.Samples[Index].Position = Trajectory.Samples[Index + 1].Position + InferredGroundTranslation;
 				Trajectory.Samples[Index].Facing = Trajectory.Samples[Index + 1].Facing;
 			}
@@ -340,7 +342,7 @@ void UPoseSearchTrajectoryLibrary::UpdateHistory_TransformHistory(FTransformTraj
 
 			for (int32 Index = 0; Index < NumHistorySamples; ++Index)
 			{
-				Trajectory.Samples[Index].TimeInSeconds -= DeltaTime;
+				Trajectory.Samples[Index].TimeInSeconds -= TimeShift;
 				Trajectory.Samples[Index].Position += InferredGroundTranslation;
 			}
 		}
@@ -360,12 +362,14 @@ void UPoseSearchTrajectoryLibrary::UpdateHistory_WorldSpace(FTransformTrajectory
 		// So num history samples needs to be less than the total number
 		check(NumHistorySamples < Trajectory.Samples.Num());
 
+		float TimeShift = Trajectory.Samples[NumHistorySamples].TimeInSeconds;	// This works out as the delta time of the last frame  
+
 		// Shift history Samples when it's time to record a new one.
 		if (SecondsPerHistorySample <= 0.f || FMath::Abs(Trajectory.Samples[NumHistorySamples - 1].TimeInSeconds) >= SecondsPerHistorySample)
 		{
 			for (int32 Index = 0; Index < NumHistorySamples - 1; ++Index)
 			{
-				Trajectory.Samples[Index].TimeInSeconds = Trajectory.Samples[Index + 1].TimeInSeconds - DeltaTime;
+				Trajectory.Samples[Index].TimeInSeconds = Trajectory.Samples[Index + 1].TimeInSeconds - TimeShift;
 				Trajectory.Samples[Index].Position = Trajectory.Samples[Index + 1].Position;
 				Trajectory.Samples[Index].Facing = Trajectory.Samples[Index + 1].Facing;
 			}
@@ -381,7 +385,7 @@ void UPoseSearchTrajectoryLibrary::UpdateHistory_WorldSpace(FTransformTrajectory
 			// Didn't record a new history position, update timers and shift by ground translation
 			for (int32 Index = 0; Index < NumHistorySamples; ++Index)
 			{
-				Trajectory.Samples[Index].TimeInSeconds -= DeltaTime;
+				Trajectory.Samples[Index].TimeInSeconds -= TimeShift;
 			}
 		}
 	}	
