PoseSearch - Always generating the trajectory in mesh space rather than using the Character.View.Rotation.Yaw when the character is set to orienting the rotation to movement

diff --git a/Engine/Plugins/Animation/PoseSearch/Source/Runtime/Private/PoseSearchTrajectoryLibrary.cpp b/Engine/Plugins/Animation/PoseSearch/Source/Runtime/Private/PoseSearchTrajectoryLibrary.cpp
--- a/Engine/Plugins/Animation/PoseSearch/Source/Runtime/Private/PoseSearchTrajectoryLibrary.cpp
+++ b/Engine/Plugins/Animation/PoseSearch/Source/Runtime/Private/PoseSearchTrajectoryLibrary.cpp
@@ -14,6 +14,12 @@
 
 #include UE_INLINE_GENERATED_CPP_BY_NAME(PoseSearchTrajectoryLibrary)
 
+#if ENABLE_ANIM_DEBUG
+// @todo: remove GVarPoseSearchTrajectoryFacingFromMesh once the behavior is enabled by defaulted
+static bool GVarPoseSearchTrajectoryFacingFromMesh = false;
+static FAutoConsoleVariableRef CVarPoseSearchTrajectoryFacingFromMesh(TEXT("a.PoseSearchTrajectory.FacingFromMesh"), GVarPoseSearchTrajectoryFacingFromMesh, TEXT("Always generating the trajectory in mesh space rather than using the Character.View.Rotation.Yaw when the character is set to orienting the rotation to movement"));
+#endif // ENABLE_ANIM_DEBUG
+
 bool FPoseSearchTrajectoryData::UpdateData(
 	float DeltaTime,
 	const FAnimInstanceProxy& AnimInstanceProxy,
@@ -92,8 +98,12 @@ bool FPoseSearchTrajectoryData::UpdateData(
 
 	TrajectoryDataDerived.Position = MeshComp->GetComponentLocation();
 	TrajectoryDataDerived.MeshCompRelativeRotation = MeshComp->GetRelativeRotation().Quaternion();
-		
-	if (TrajectoryDataDerived.bOrientRotationToMovement)
+	
+	if (TrajectoryDataDerived.bOrientRotationToMovement
+#if ENABLE_ANIM_DEBUG
+		|| GVarPoseSearchTrajectoryFacingFromMesh
+#endif // ENABLE_ANIM_DEBUG
+		)
 	{
 		TrajectoryDataDerived.Facing = MeshComp->GetComponentTransform().GetRotation();
 	}
