PoseSearch - adding mesh transform debug lines when displaying pose history from the pose search debugger to help detecting the provided trajectory to the pose history is stale

diff --git a/Engine/Plugins/Animation/PoseSearch/Source/Editor/Private/PoseSearchDebuggerView.cpp b/Engine/Plugins/Animation/PoseSearch/Source/Editor/Private/PoseSearchDebuggerView.cpp
--- a/Engine/Plugins/Animation/PoseSearch/Source/Editor/Private/PoseSearchDebuggerView.cpp
+++ b/Engine/Plugins/Animation/PoseSearch/Source/Editor/Private/PoseSearchDebuggerView.cpp
@@ -757,6 +757,7 @@ TSharedRef<SHorizontalBox> SDebuggerView::GenerateReturnButtonView()
 				[
 					SNew(STextBlock)
 					.Text(LOCTEXT("PoseSearchDebuggerDrawHistory", "Draw History"))
+					.ToolTipText(LOCTEXT("PoseSearchDebuggerDrawHistory_ToolTip", "Draws recorded bones paths in red. Additionally draws the previous frame directions of the root bone in White, the trajectory in Black, and the mesh at the time of the search in Turquoise"))
 				]
 			]
 		]
diff --git a/Engine/Plugins/Animation/PoseSearch/Source/Runtime/Private/AnimNode_PoseSearchHistoryCollector.cpp b/Engine/Plugins/Animation/PoseSearch/Source/Runtime/Private/AnimNode_PoseSearchHistoryCollector.cpp
--- a/Engine/Plugins/Animation/PoseSearch/Source/Runtime/Private/AnimNode_PoseSearchHistoryCollector.cpp
+++ b/Engine/Plugins/Animation/PoseSearch/Source/Runtime/Private/AnimNode_PoseSearchHistoryCollector.cpp
@@ -7,6 +7,7 @@
 #include "Animation/AnimStats.h"
 #include "Components/SkeletalMeshComponent.h"
 #include "Engine/SkinnedAsset.h"
+#include "PoseSearch/PoseSearchContext.h"
 #include "PoseSearch/PoseSearchHistoryAttribute.h"
 #include "PoseSearch/PoseSearchInteractionValidator.h"
 
@@ -30,8 +31,15 @@ void FAnimNode_PoseSearchHistoryCollector_Base::OnInitializeAnimInstance(const F
 
 void FAnimNode_PoseSearchHistoryCollector_Base::Initialize_AnyThread(const FAnimationInitializeContext& Context)
 {
+	using namespace UE::PoseSearch;
+
 	Super::Initialize_AnyThread(Context);
 
+	check(Context.AnimInstanceProxy);
+	const UObject* AnimContext = Context.AnimInstanceProxy->GetAnimInstanceObject();
+	check(AnimContext);
+	CheckInteractionThreadSafety(AnimContext);
+
 	// reinitializing PoseHistoryPtr if invalid here. it could happen when recompiling animation 
 	// blue prints when OnInitializeAnimInstance is not called
 	if (!PoseHistoryPtr.IsValid())
@@ -47,7 +55,11 @@ void FAnimNode_PoseSearchHistoryCollector_Base::Initialize_AnyThread(const FAnim
 	{
 		PoseHistoryPtr->EvaluateComponentSpace_AnyThread(0.f, ComponentSpacePoseProvider, bStoreScales,
 			RootBoneRecoveryTime, RootBoneTranslationRecoveryRatio, RootBoneRotationRecoveryRatio, true, true,
-			GetRequiredBones(Context.AnimInstanceProxy), FBlendedCurve(), MakeConstArrayView(CollectedCurves));
+			GetRequiredBones(Context.AnimInstanceProxy), FBlendedCurve(), MakeConstArrayView(CollectedCurves)
+#if WITH_EDITORONLY_DATA
+			, GetContextTransform(AnimContext, false)
+#endif //WITH_EDITORONLY_DATA
+		);
 	}
 }
 
@@ -149,7 +161,9 @@ void FAnimNode_PoseSearchHistoryCollector::Evaluate_AnyThread(FPoseContext& Outp
 	Source.Evaluate(Output);
 
 	check(Output.AnimInstanceProxy);
-	CheckInteractionThreadSafety(Output.AnimInstanceProxy->GetAnimInstanceObject());
+	const UObject* AnimContext = Output.AnimInstanceProxy->GetAnimInstanceObject();
+	check(AnimContext);
+	CheckInteractionThreadSafety(AnimContext);
 
 	const bool bNeedsReset = bResetOnBecomingRelevant && UpdateCounter.HasEverBeenUpdated() && !UpdateCounter.WasSynchronizedCounter(Output.AnimInstanceProxy->GetUpdateCounter());
 
@@ -168,7 +182,11 @@ void FAnimNode_PoseSearchHistoryCollector::Evaluate_AnyThread(FPoseContext& Outp
 		UE::PoseSearch::FComponentSpacePoseProvider ComponentSpacePoseProvider(ComponentSpacePose);
 		PoseHistory->EvaluateComponentSpace_AnyThread(Output.AnimInstanceProxy->GetDeltaSeconds(), ComponentSpacePoseProvider, bStoreScales,
 			RootBoneRecoveryTime, RootBoneTranslationRecoveryRatio, RootBoneRotationRecoveryRatio, bNeedsReset, bCacheBones,
-			RequiredBones, Output.Curve, MakeConstArrayView(CollectedCurves));
+			RequiredBones, Output.Curve, MakeConstArrayView(CollectedCurves)
+#if WITH_EDITORONLY_DATA
+			, GetContextTransform(AnimContext, false)
+#endif //WITH_EDITORONLY_DATA
+		);
 
 		// Fill in the pose history into a custom attribute for access downstream
 		if (FPoseHistoryAnimationAttribute* PoseHistoryAttribute = Output.CustomAttributes.FindOrAdd<FPoseHistoryAnimationAttribute>(UE::PoseSearch::PoseHistoryAttributeId))
@@ -225,7 +243,9 @@ void FAnimNode_PoseSearchComponentSpaceHistoryCollector::EvaluateComponentSpace_
 	Source.EvaluateComponentSpace(Output);
 
 	check(Output.AnimInstanceProxy);
-	CheckInteractionThreadSafety(Output.AnimInstanceProxy->GetAnimInstanceObject());
+	const UObject* AnimContext = Output.AnimInstanceProxy->GetAnimInstanceObject();
+	check(AnimContext);
+	CheckInteractionThreadSafety(AnimContext);
 
 	const bool bNeedsReset = bResetOnBecomingRelevant && UpdateCounter.HasEverBeenUpdated() && !UpdateCounter.WasSynchronizedCounter(Output.AnimInstanceProxy->GetUpdateCounter());
 
@@ -241,7 +261,11 @@ void FAnimNode_PoseSearchComponentSpaceHistoryCollector::EvaluateComponentSpace_
 		UE::PoseSearch::FComponentSpacePoseProvider ComponentSpacePoseProvider(Output.Pose);
 		PoseHistory->EvaluateComponentSpace_AnyThread(Output.AnimInstanceProxy->GetDeltaSeconds(), ComponentSpacePoseProvider, bStoreScales,
 			RootBoneRecoveryTime, RootBoneTranslationRecoveryRatio, RootBoneRotationRecoveryRatio, bNeedsReset, bCacheBones,
-			RequiredBones, Output.Curve, MakeConstArrayView(CollectedCurves));
+			RequiredBones, Output.Curve, MakeConstArrayView(CollectedCurves)
+#if WITH_EDITORONLY_DATA
+			, GetContextTransform(AnimContext, false)
+#endif //WITH_EDITORONLY_DATA
+		);
 
 		// Fill in the pose history into a custom attribute for access downstream
 		if (FPoseHistoryAnimationAttribute* PoseHistoryAttribute = Output.CustomAttributes.FindOrAdd<FPoseHistoryAnimationAttribute>(UE::PoseSearch::PoseHistoryAttributeId))
diff --git a/Engine/Plugins/Animation/PoseSearch/Source/Runtime/Private/Chooser/PoseSearchChooserColumn.cpp b/Engine/Plugins/Animation/PoseSearch/Source/Runtime/Private/Chooser/PoseSearchChooserColumn.cpp
--- a/Engine/Plugins/Animation/PoseSearch/Source/Runtime/Private/Chooser/PoseSearchChooserColumn.cpp
+++ b/Engine/Plugins/Animation/PoseSearch/Source/Runtime/Private/Chooser/PoseSearchChooserColumn.cpp
@@ -642,7 +642,7 @@ void FPoseSearchColumn::SetOutputs(FChooserEvaluationContext& Context, int RowIn
 					{
 						// @todo: integrate DeltaTime into SearchContext, and implement it for UAF as well
 						float DeltaTime = FiniteDelta;
-						if (const UAnimInstance* AnimInstance = Cast<UAnimInstance>(Context.GetFirstObjectParam()))
+						if (const UAnimInstance* AnimInstance = Cast<UAnimInstance>(GetAnimContext(&Context)))
 						{
 							DeltaTime = AnimInstance->GetDeltaSeconds();
 						}
@@ -675,7 +675,7 @@ void FPoseSearchColumn::SetOutputs(FChooserEvaluationContext& Context, int RowIn
 						{
 							// @todo: integrate DeltaTime into SearchContext, and implement it for UAF as well
 							float DeltaTime = FiniteDelta;
-							if (const UAnimInstance* AnimInstance = Cast<UAnimInstance>(Context.GetFirstObjectParam()))
+							if (const UAnimInstance* AnimInstance = Cast<UAnimInstance>(GetAnimContext(&Context)))
 							{
 								DeltaTime = AnimInstance->GetDeltaSeconds();
 							}
diff --git a/Engine/Plugins/Animation/PoseSearch/Source/Runtime/Private/PoseSearchContext.cpp b/Engine/Plugins/Animation/PoseSearch/Source/Runtime/Private/PoseSearchContext.cpp
--- a/Engine/Plugins/Animation/PoseSearch/Source/Runtime/Private/PoseSearchContext.cpp
+++ b/Engine/Plugins/Animation/PoseSearch/Source/Runtime/Private/PoseSearchContext.cpp
@@ -17,6 +17,12 @@
 namespace UE::PoseSearch
 {
 	
+const UObject* GetAnimContext(const FChooserEvaluationContext* Context)
+{
+	check(Context);
+	return Context->GetFirstObjectParam();
+}
+
 const FTransform& GetContextTransform(const UObject* AnimContext, bool bCheckIsInGameThread)
 {
 	check(!bCheckIsInGameThread || IsInGameThread());
@@ -70,14 +76,6 @@ const USkeleton* GetContextSkeleton(const UObject* AnimContext, bool bCheckIsInG
 	return Skeleton;
 }
 	
-const USkeleton* GetContextSkeleton(FChooserEvaluationContext& Context, bool bCheckIsInGameThread)
-{
-	check(!bCheckIsInGameThread || IsInGameThread());
-	check(!Context.ObjectParams.IsEmpty());
-
-	return GetContextSkeleton(Context.ObjectParams[0].Object, bCheckIsInGameThread);
-}
-
 const USkeletalMeshComponent* GetContextSkeletalMeshComponent(const UObject* AnimContext, bool bCheckIsInGameThread)
 {
 	check(!bCheckIsInGameThread || IsInGameThread());
@@ -165,19 +163,18 @@ public:
 	}
 };
 
-static FAnimInstanceProxy* GetAnimInstanceProxy(TConstArrayView<FChooserEvaluationContext*> AnimContexts)
+static FAnimInstanceProxy* GetAnimInstanceProxy(TConstArrayView<FChooserEvaluationContext*> Contexts)
 {
-	if (!AnimContexts.IsEmpty())
+	if (!Contexts.IsEmpty())
 	{
-		return UAnimInstanceProxyProvider::GetAnimInstanceProxy(Cast<UAnimInstance>(AnimContexts[0]->GetFirstObjectParam()));
+		return UAnimInstanceProxyProvider::GetAnimInstanceProxy(Cast<UAnimInstance>(GetAnimContext(Contexts[0])));
 	}
 	return nullptr;
 }
 
-static const USkinnedMeshComponent* GetMesh(TConstArrayView<FChooserEvaluationContext*> AnimContexts, int32 RoleIndex = 0)
+static const USkinnedMeshComponent* GetMesh(TConstArrayView<FChooserEvaluationContext*> Contexts, int32 RoleIndex = 0)
 {
-	check(AnimContexts.IsValidIndex(RoleIndex));
-	return Cast<USkinnedMeshComponent>(AnimContexts[RoleIndex]->GetFirstObjectParam());
+	return Cast<USkinnedMeshComponent>(GetAnimContext(Contexts[RoleIndex]));
 }
 
 static const UWorld* GetWorld(TConstArrayView<FChooserEvaluationContext*> AnimContexts)
diff --git a/Engine/Plugins/Animation/PoseSearch/Source/Runtime/Private/PoseSearchFeatureChannel_Curve.cpp b/Engine/Plugins/Animation/PoseSearch/Source/Runtime/Private/PoseSearchFeatureChannel_Curve.cpp
--- a/Engine/Plugins/Animation/PoseSearch/Source/Runtime/Private/PoseSearchFeatureChannel_Curve.cpp
+++ b/Engine/Plugins/Animation/PoseSearch/Source/Runtime/Private/PoseSearchFeatureChannel_Curve.cpp
@@ -37,7 +37,7 @@ void UPoseSearchFeatureChannel_Curve::BuildQuery(UE::PoseSearch::FSearchContext&
 	float CurveValue = 0.0f;
 	if (bUseBlueprintQueryOverride)
 	{
-		if (const UAnimInstance* AnimInstance = Cast<UAnimInstance>(SearchContext.GetContext(SampleRole)->GetFirstObjectParam()))
+		if (const UAnimInstance* AnimInstance = Cast<UAnimInstance>(GetAnimContext(SearchContext.GetContext(SampleRole))))
 		{
 			CurveValue = BP_GetCurveValue(AnimInstance);
 			FFeatureVectorHelper::EncodeFloat(SearchContext.EditFeatureVector(), ChannelDataOffset, CurveValue);
diff --git a/Engine/Plugins/Animation/PoseSearch/Source/Runtime/Private/PoseSearchFeatureChannel_Heading.cpp b/Engine/Plugins/Animation/PoseSearch/Source/Runtime/Private/PoseSearchFeatureChannel_Heading.cpp
--- a/Engine/Plugins/Animation/PoseSearch/Source/Runtime/Private/PoseSearchFeatureChannel_Heading.cpp
+++ b/Engine/Plugins/Animation/PoseSearch/Source/Runtime/Private/PoseSearchFeatureChannel_Heading.cpp
@@ -108,7 +108,7 @@ void UPoseSearchFeatureChannel_Heading::BuildQuery(UE::PoseSearch::FSearchContex
 
 	if (bUseBlueprintQueryOverride)
 	{
-		if (const UAnimInstance* AnimInstance = Cast<UAnimInstance>(SearchContext.GetContext(SampleRole)->GetFirstObjectParam()))
+		if (const UAnimInstance* AnimInstance = Cast<UAnimInstance>(GetAnimContext(SearchContext.GetContext(SampleRole))))
 		{
 			const FQuat BoneRotationWorld = BP_GetWorldRotation(AnimInstance);
 			const FQuat BoneRotation = SearchContext.GetSampleRotation(SampleTimeOffset, OriginTimeOffset, SchemaBoneIdx, SchemaOriginBoneIdx, SampleRole, OriginRole, EPermutationTimeType::UseSampleTime, &BoneRotationWorld);
diff --git a/Engine/Plugins/Animation/PoseSearch/Source/Runtime/Private/PoseSearchFeatureChannel_Position.cpp b/Engine/Plugins/Animation/PoseSearch/Source/Runtime/Private/PoseSearchFeatureChannel_Position.cpp
--- a/Engine/Plugins/Animation/PoseSearch/Source/Runtime/Private/PoseSearchFeatureChannel_Position.cpp
+++ b/Engine/Plugins/Animation/PoseSearch/Source/Runtime/Private/PoseSearchFeatureChannel_Position.cpp
@@ -97,7 +97,7 @@ void UPoseSearchFeatureChannel_Position::BuildQuery(UE::PoseSearch::FSearchConte
 
 	if (bUseBlueprintQueryOverride)
 	{
-		if (const UAnimInstance* AnimInstance = Cast<UAnimInstance>(SearchContext.GetContext(SampleRole)->GetFirstObjectParam()))
+		if (const UAnimInstance* AnimInstance = Cast<UAnimInstance>(GetAnimContext(SearchContext.GetContext(SampleRole))))
 		{
 			const FVector BonePositionWorld = BP_GetWorldPosition(AnimInstance);
 			const FVector BonePosition = SearchContext.GetSamplePosition(SampleTimeOffset, OriginTimeOffset, SchemaBoneIdx, SchemaOriginBoneIdx, SampleRole, OriginRole, PermutationTimeType, &BonePositionWorld);
diff --git a/Engine/Plugins/Animation/PoseSearch/Source/Runtime/Private/PoseSearchFeatureChannel_TimeToEvent.cpp b/Engine/Plugins/Animation/PoseSearch/Source/Runtime/Private/PoseSearchFeatureChannel_TimeToEvent.cpp
--- a/Engine/Plugins/Animation/PoseSearch/Source/Runtime/Private/PoseSearchFeatureChannel_TimeToEvent.cpp
+++ b/Engine/Plugins/Animation/PoseSearch/Source/Runtime/Private/PoseSearchFeatureChannel_TimeToEvent.cpp
@@ -35,9 +35,10 @@ void UPoseSearchFeatureChannel_TimeToEvent::BuildQuery(UE::PoseSearch::FSearchCo
 	if (SearchContext.GetContexts().IsEmpty())
 	{
 		UE_LOG(LogPoseSearch, Error, TEXT("UPoseSearchFeatureChannel_Position::BuildQuery - no provided anim contexts!"));
+		return;
 	}
 	
-	if (const UAnimInstance* AnimInstance = Cast<UAnimInstance>(SearchContext.GetContexts()[0]->GetFirstObjectParam()))
+	if (const UAnimInstance* AnimInstance = Cast<UAnimInstance>(GetAnimContext(SearchContext.GetContexts()[0])))
 	{
 		const float TimeToEvent = BP_GetTimeToEvent(AnimInstance);
 		FFeatureVectorHelper::EncodeFloat(SearchContext.EditFeatureVector(), ChannelDataOffset, TimeToEvent);
diff --git a/Engine/Plugins/Animation/PoseSearch/Source/Runtime/Private/PoseSearchFeatureChannel_Velocity.cpp b/Engine/Plugins/Animation/PoseSearch/Source/Runtime/Private/PoseSearchFeatureChannel_Velocity.cpp
--- a/Engine/Plugins/Animation/PoseSearch/Source/Runtime/Private/PoseSearchFeatureChannel_Velocity.cpp
+++ b/Engine/Plugins/Animation/PoseSearch/Source/Runtime/Private/PoseSearchFeatureChannel_Velocity.cpp
@@ -50,7 +50,7 @@ void UPoseSearchFeatureChannel_Velocity::BuildQuery(UE::PoseSearch::FSearchConte
 
 	if (bUseBlueprintQueryOverride)
 	{
-		if (const UAnimInstance* AnimInstance = Cast<UAnimInstance>(SearchContext.GetContext(SampleRole)->GetFirstObjectParam()))
+		if (const UAnimInstance* AnimInstance = Cast<UAnimInstance>(GetAnimContext(SearchContext.GetContext(SampleRole))))
 		{
 			const FVector LinearVelocityWorld = BP_GetWorldVelocity(AnimInstance);
 			FVector LinearVelocity = SearchContext.GetSampleVelocity(SampleTimeOffset, OriginTimeOffset, SchemaBoneIdx, SchemaOriginBoneIdx, SampleRole, OriginRole, bUseCharacterSpaceVelocities, EPermutationTimeType::UseSampleTime, &LinearVelocityWorld);
diff --git a/Engine/Plugins/Animation/PoseSearch/Source/Runtime/Private/PoseSearchHistory.cpp b/Engine/Plugins/Animation/PoseSearch/Source/Runtime/Private/PoseSearchHistory.cpp
--- a/Engine/Plugins/Animation/PoseSearch/Source/Runtime/Private/PoseSearchHistory.cpp
+++ b/Engine/Plugins/Animation/PoseSearch/Source/Runtime/Private/PoseSearchHistory.cpp
@@ -338,8 +338,16 @@ const USkeleton* FAIPComponentSpacePoseProvider::GetSkeletonAsset() const
 
 //////////////////////////////////////////////////////////////////////////
 // FPoseHistoryEntry
-void FPoseHistoryEntry::Update(float Time, IComponentSpacePoseProvider& ComponentSpacePoseProvider, const FBoneToTransformMap& BoneToTransformMap, bool bStoreScales, const FBlendedCurve& Curves, const TConstArrayView<FName>& CollectedCurves)
+void FPoseHistoryEntry::Update(float Time, IComponentSpacePoseProvider& ComponentSpacePoseProvider, const FBoneToTransformMap& BoneToTransformMap, bool bStoreScales, const FBlendedCurve& Curves, const TConstArrayView<FName>& CollectedCurves
+#if WITH_EDITORONLY_DATA
+	, const FTransform& InDebugMeshTransform
+#endif //WITH_EDITORONLY_DATA
+)
 {
+#if WITH_EDITORONLY_DATA
+	DebugMeshTransform = InDebugMeshTransform;
+#endif //WITH_EDITORONLY_DATA
+
 	AccumulatedSeconds = Time;
 
 	const USkeleton* Skeleton = ComponentSpacePoseProvider.GetSkeletonAsset();
@@ -431,6 +439,11 @@ FArchive& operator<<(FArchive& Ar, FPoseHistoryEntry& Entry)
 	Ar << Entry.ComponentSpaceScales;
 	Ar << Entry.CurveValues;
 	Ar << Entry.AccumulatedSeconds;
+
+#if WITH_EDITORONLY_DATA
+	Ar << Entry.DebugMeshTransform;
+#endif // WITH_EDITORONLY_DATA
+
 	return Ar;
 }
 
@@ -604,6 +617,13 @@ void FArchivedPoseHistory::DebugDraw(const UWorld* World, FColor Color) const
 
 			const bool bIsCurrentTimeEntry = FMath::IsNearlyZero(Entry.AccumulatedSeconds);
 
+#if WITH_EDITORONLY_DATA
+			if (bIsCurrentTimeEntry)
+			{
+				DrawDebugLine(World, Entry.DebugMeshTransform.GetTranslation(), Entry.DebugMeshTransform.GetTranslation() + Entry.DebugMeshTransform.GetUnitAxis(EAxis::Type::X) * 30.f, FColor::Turquoise, false, -1.f, SDPG_Foreground);
+			}
+#endif // WITH_EDITORONLY_DATA
+
 			for (int32 i = 0; i < Entry.Num(); ++i)
 			{
 				const FTransform RootTransform = Trajectory.GetSampleAtTime(Entry.AccumulatedSeconds).GetTransform();
@@ -951,7 +971,11 @@ void FPoseHistory::SetTrajectory(const FTransformTrajectory& InTrajectory, float
 void FPoseHistory::EvaluateComponentSpace_AnyThread(float DeltaTime, IComponentSpacePoseProvider& ComponentSpacePoseProvider, bool bStoreScales,
 	float RootBoneRecoveryTime, float RootBoneTranslationRecoveryRatio, float RootBoneRotationRecoveryRatio,
 	bool bNeedsReset, bool bCacheBones, const TArray<FBoneIndexType>& RequiredBones,
-	const FBlendedCurve& Curves, const TConstArrayView<FName>& InCollectedCurves)
+	const FBlendedCurve& Curves, const TConstArrayView<FName>& InCollectedCurves
+#if WITH_EDITORONLY_DATA
+	, const FTransform& InDebugMeshTransform
+#endif //WITH_EDITORONLY_DATA
+	)
 {
 	UE_MT_SCOPED_WRITE_ACCESS(PoseDataThreadSafeCounter);
 
@@ -1026,7 +1050,11 @@ void FPoseHistory::EvaluateComponentSpace_AnyThread(float DeltaTime, IComponentS
 
 	// Regardless of the retention policy, we always update the most recent Entry
 	FPoseHistoryEntry& MostRecentEntry = Entries.Last();
-	MostRecentEntry.Update(0.f, ComponentSpacePoseProvider, BoneToTransformMap, bStoreScales, Curves, CollectedCurves);
+	MostRecentEntry.Update(0.f, ComponentSpacePoseProvider, BoneToTransformMap, bStoreScales, Curves, CollectedCurves
+#if WITH_EDITORONLY_DATA
+		, InDebugMeshTransform
+#endif //WITH_EDITORONLY_DATA
+	);
 
 	if (RootBoneRecoveryTime > 0.f && !Trajectory.Samples.IsEmpty())
 	{
@@ -1208,16 +1236,28 @@ void FMemStackPoseHistory::AddFutureRootBone(float Time, const FTransform& Futur
 	FutureEntry.AccumulatedSeconds = Time;
 }
 
-void FMemStackPoseHistory::AddFuturePose(float Time, IComponentSpacePoseProvider& ComponentSpacePoseProvider, const FBlendedCurve& Curves)
+void FMemStackPoseHistory::AddFuturePose(float Time, IComponentSpacePoseProvider& ComponentSpacePoseProvider, const FBlendedCurve& Curves
+#if WITH_EDITORONLY_DATA
+	, const FTransform& InDebugMeshTransform
+#endif //WITH_EDITORONLY_DATA
+	)
 {
 	// we don't allow to add "past" or "present" poses to FutureEntries
 	check(Time > 0.f);
 	check(PoseHistory);	
 	const int32 LowerBoundIdx = Algo::LowerBound(FutureEntries, Time, [](const FPoseHistoryEntry& Entry, float Value) { return Value > Entry.AccumulatedSeconds; });
-	FutureEntries.InsertDefaulted_GetRef(LowerBoundIdx).Update(Time, ComponentSpacePoseProvider, GetBoneToTransformMap(), true, Curves, MakeConstArrayView(GetCollectedCurves()));
+	FutureEntries.InsertDefaulted_GetRef(LowerBoundIdx).Update(Time, ComponentSpacePoseProvider, GetBoneToTransformMap(), true, Curves, MakeConstArrayView(GetCollectedCurves())
+#if WITH_EDITORONLY_DATA
+	, InDebugMeshTransform
+#endif //WITH_EDITORONLY_DATA
+	);
 }
 
-void FMemStackPoseHistory::ExtractAndAddFuturePoses(const UAnimationAsset* AnimationAsset, float AnimationTime, float FiniteDeltaTime, const FVector& BlendParame1ters, float IntervalTime, const USkeleton* OverrideSkeleton, bool bUseRefPoseRootBone)
+void FMemStackPoseHistory::ExtractAndAddFuturePoses(const UAnimationAsset* AnimationAsset, float AnimationTime, float FiniteDeltaTime, const FVector& BlendParame1ters, float IntervalTime, const USkeleton* OverrideSkeleton, bool bUseRefPoseRootBone
+#if WITH_EDITORONLY_DATA
+	, const FTransform& InDebugMeshTransform
+#endif //WITH_EDITORONLY_DATA
+)
 {
 	FMemMark Mark(FMemStack::Get());
 
@@ -1284,7 +1324,14 @@ void FMemStackPoseHistory::ExtractAndAddFuturePoses(const UAnimationAsset* Anima
 
 		ComponentSpacePose.InitPose(Pose);
 		FComponentSpacePoseProvider ComponentSpacePoseProvider(ComponentSpacePose);
-		AddFuturePose(FuturePoseAnimationTime, ComponentSpacePoseProvider);
+		AddFuturePose(FuturePoseAnimationTime, ComponentSpacePoseProvider
+#if WITH_EDITORONLY_DATA
+			, FBlendedCurve()
+			// @todo: calculate the predicted future InDebugMeshTransform properly, by extracting the delta transform between the 
+			//        FuturePoseExtractionTime - IntervalTime and FuturePoseExtractionTime and applying it to InDebugMeshTransform!
+			, InDebugMeshTransform
+#endif // WITH_EDITORONLY_DATA
+		);
 	}
 }
 
diff --git a/Engine/Plugins/Animation/PoseSearch/Source/Runtime/Private/PoseSearchInteractionIsland.cpp b/Engine/Plugins/Animation/PoseSearch/Source/Runtime/Private/PoseSearchInteractionIsland.cpp
--- a/Engine/Plugins/Animation/PoseSearch/Source/Runtime/Private/PoseSearchInteractionIsland.cpp
+++ b/Engine/Plugins/Animation/PoseSearch/Source/Runtime/Private/PoseSearchInteractionIsland.cpp
@@ -328,6 +328,8 @@ uint32 FInteractionSearchContextBase::CalculateHashForEquivalence() const
 	// skipping bDisableCollisions and PoseHistories for hashing
 	for (int32 Index = 0; Index < AnimContexts.Num(); ++Index)
 	{
+		// AnimContexts must be sorted to have deterministic searches across multiple frames
+		check(Index == 0 || AnimContexts[Index - 1] < AnimContexts[Index]);
 		Hash = HashCombineFast(Hash, GetTypeHash(AnimContexts[Index]));
 	}
 
@@ -361,11 +363,13 @@ void FInteractionSearchContextBase::AddReferencedObjects(FReferenceCollector& Co
 
 void FInteractionSearchContextBase::Add(const UObject* AnimContext, const IPoseHistory* PoseHistory, const FRole Role, bool bWantsDisableCollisions)
 {
+	// trying to inefficiently call Add again after having calculated the hash via CalculateHashForEquivalence?
+	check(CachedHashForEquivalence == InvalidCachedHashForEquivalence);
+
 	check(AnimContext && PoseHistory);
 
 	// AnimContexts must be sorted to have deterministic searches across multiple frames
-	// as well as for faster IsEquivalent implementation
-	check(AnimContexts.IsEmpty() || GetAnimContext(AnimContexts.Num() - 1) < AnimContext);
+	check(AnimContexts.IsEmpty() || AnimContexts.Last() < AnimContext);
 
 	AnimContexts.Add(AnimContext);
 	PoseHistories.Add(PoseHistory->AsWeak());
diff --git a/Engine/Plugins/Animation/PoseSearch/Source/Runtime/Private/PoseSearchLibrary.cpp b/Engine/Plugins/Animation/PoseSearch/Source/Runtime/Private/PoseSearchLibrary.cpp
--- a/Engine/Plugins/Animation/PoseSearch/Source/Runtime/Private/PoseSearchLibrary.cpp
+++ b/Engine/Plugins/Animation/PoseSearch/Source/Runtime/Private/PoseSearchLibrary.cpp
@@ -655,9 +655,9 @@ void UPoseSearchLibrary::TraceMotionMatching(
 	float RecordingTime = 0.f;
 	if (!SearchContext.GetContexts().IsEmpty())
 	{
-		if (const UObject* FirstObject = SearchContext.GetContexts()[0]->GetFirstObjectParam())
+		if (const UObject* AnimContext = GetAnimContext(SearchContext.GetContexts()[0]))
 		{
-			RecordingTime = FObjectTrace::GetWorldElapsedTime(FirstObject->GetWorld());
+			RecordingTime = FObjectTrace::GetWorldElapsedTime(AnimContext->GetWorld());
 		}
 	}
 
@@ -672,39 +672,35 @@ void UPoseSearchLibrary::TraceMotionMatching(
 
 	for (int32 AnimInstanceIndex = 0; AnimInstanceIndex < AnimContextsNum; ++AnimInstanceIndex)
 	{
-		if (const FChooserEvaluationContext* AnimContext = SearchContext.GetContexts()[AnimInstanceIndex])
+		const UObject* AnimContext = GetAnimContext(SearchContext.GetContexts()[AnimInstanceIndex]);
+		const UObject* SkeletalMeshComponent = nullptr;
+		if (const UAnimInstance* AnimInstance = Cast<UAnimInstance>(AnimContext))
 		{
-			const UObject* FirstObject = AnimContext->GetFirstObjectParam();
-			const UObject* SkeletalMeshComponent = nullptr;
-			if (const UAnimInstance* AnimInstance = Cast<UAnimInstance>(FirstObject))
-			{
-				SkeletalMeshComponent = AnimInstance->GetOuter();
-			}
-			else if (const UActorComponent* ActorComponent = Cast<UActorComponent>(FirstObject))
-			{
-				const AActor* Actor = ActorComponent->GetOwner();
-				check(Actor);
+			SkeletalMeshComponent = AnimInstance->GetOuter();
+		}
+		else if (const UActorComponent* ActorComponent = Cast<UActorComponent>(AnimContext))
+		{
+			const AActor* Actor = ActorComponent->GetOwner();
+			check(Actor);
 				
-				SkeletalMeshComponent = Actor->GetComponentByClass<USkeletalMeshComponent>();
-			}
+			SkeletalMeshComponent = Actor->GetComponentByClass<USkeletalMeshComponent>();
+		}
 
-			if (!SkeletalMeshComponent || CANNOT_TRACE_OBJECT(SkeletalMeshComponent))
-			{
-				return;
-			}
-			
-			TRACE_OBJECT(SkeletalMeshComponent);
-			TraceState.SkeletalMeshComponentIds[AnimInstanceIndex] = FObjectTrace::GetObjectId(SkeletalMeshComponent);
+		if (!SkeletalMeshComponent || CANNOT_TRACE_OBJECT(SkeletalMeshComponent))
+		{
+			return;
 		}
+			
+		TRACE_OBJECT(SkeletalMeshComponent);
+		TraceState.SkeletalMeshComponentIds[AnimInstanceIndex] = FObjectTrace::GetObjectId(SkeletalMeshComponent);
 	}
 
 	for (int32 AnimInstanceIndex = 0; AnimInstanceIndex < AnimContextsNum; ++AnimInstanceIndex)
 	{
-		const FChooserEvaluationContext* Context = SearchContext.GetContexts()[AnimInstanceIndex];
-		if (const UObject* Object = Context->GetFirstObjectParam())
+		if (const UObject* AnimContext = GetAnimContext(SearchContext.GetContexts()[AnimInstanceIndex]))
 		{
-			TRACE_OBJECT(Object);
-			SearchId = HashCombineFast(SearchId, GetTypeHash(FObjectTrace::GetObjectId(Object)));
+			TRACE_OBJECT(AnimContext);
+			SearchId = HashCombineFast(SearchId, GetTypeHash(FObjectTrace::GetObjectId(AnimContext)));
 		}
 	}
 
@@ -786,7 +782,7 @@ void UPoseSearchLibrary::TraceMotionMatching(
 	float DeltaTime = FiniteDelta;
 	if (!SearchContext.GetContexts().IsEmpty())
 	{
-		if (const UAnimInstance* AnimInstance = Cast<UAnimInstance>(SearchContext.GetContexts()[0]->GetFirstObjectParam()))
+		if (const UAnimInstance* AnimInstance = Cast<UAnimInstance>(GetAnimContext(SearchContext.GetContexts()[0])))
 		{
 			DeltaTime = AnimInstance->GetDeltaSeconds();
 		}
@@ -861,7 +857,7 @@ void UPoseSearchLibrary::TraceMotionMatching(
 	// @todo: avoid publishing duplicated TraceState in ALL the AnimContexts! -currently necessary for multi character-
 	for (const FChooserEvaluationContext* Context : SearchContext.GetContexts())
 	{
-		const UObject* AnimContextObject = Context->GetFirstObjectParam();
+		const UObject* AnimContextObject = GetAnimContext(Context);
 		TRACE_OBJECT(AnimContextObject);
 		TraceState.AnimInstanceId = FObjectTrace::GetObjectId(AnimContextObject);
 		TraceState.NodeId = SearchId;
@@ -1252,7 +1248,7 @@ void UPoseSearchLibrary::MotionMatch(
 		{
 			if (const IPoseHistory* PoseHistory = InternalPoseHistories[RoleIndex])
 			{
-				const USkeleton* Skeleton = GetContextSkeleton(Contexts[RoleIndex], false);
+				const USkeleton* Skeleton = GetContextSkeleton(GetAnimContext(&Contexts[RoleIndex]), false);
 				check(Skeleton);
 				// @todo: add input BlendParameters to support sampling FutureAnimation blendspaces and support for multi character
 				const UAnimationAsset* AnimationAsset = Cast<UAnimationAsset>(Future.Animation);
@@ -1269,7 +1265,12 @@ void UPoseSearchLibrary::MotionMatch(
 				}
 
 				MemStackPoseHistories[RoleIndex].Init(InternalPoseHistories[RoleIndex]);
-				MemStackPoseHistories[RoleIndex].ExtractAndAddFuturePoses(AnimationAsset, FutureAnimationTime, FiniteDelta, FVector::ZeroVector, FutureIntervalTime, Skeleton);
+				MemStackPoseHistories[RoleIndex].ExtractAndAddFuturePoses(AnimationAsset, FutureAnimationTime, FiniteDelta, FVector::ZeroVector, FutureIntervalTime, Skeleton
+#if WITH_EDITORONLY_DATA
+				, false
+				, GetContextTransform(GetAnimContext(&Contexts[RoleIndex]), false)
+#endif //WITH_EDITORONLY_DATA
+				);
 				InternalPoseHistories[RoleIndex] = MemStackPoseHistories[RoleIndex].GetThisOrPoseHistory();
 			}
 		}
diff --git a/Engine/Plugins/Animation/PoseSearch/Source/Runtime/Public/PoseSearch/PoseSearchContext.h b/Engine/Plugins/Animation/PoseSearch/Source/Runtime/Public/PoseSearch/PoseSearchContext.h
--- a/Engine/Plugins/Animation/PoseSearch/Source/Runtime/Public/PoseSearch/PoseSearchContext.h
+++ b/Engine/Plugins/Animation/PoseSearch/Source/Runtime/Public/PoseSearch/PoseSearchContext.h
@@ -23,9 +23,10 @@ namespace UE::PoseSearch
 typedef TSet<const UObject*, DefaultKeyFuncs<const UObject*>, TInlineSetAllocator<8, TMemStackSetAllocator<>>> FStackAssetSet;
 
 // Utils functions
-const FTransform& GetContextTransform(const UObject* AnimContext, bool bCheckIsInGameThread = true);
+const UObject* GetAnimContext(const FChooserEvaluationContext* Context);
+// Experimental, this feature might be removed without warning, not for production use
+UE_API const FTransform& GetContextTransform(const UObject* AnimContext, bool bCheckIsInGameThread = true);
 const USkeleton* GetContextSkeleton(const UObject* AnimContext, bool bCheckIsInGameThread = true);
-const USkeleton* GetContextSkeleton(FChooserEvaluationContext& AnimContext, bool bCheckIsInGameThread = true);
 const AActor* GetContextOwningActor(const UObject* AnimContext, bool bCheckIsInGameThread = true);
 FVector GetContextLocation(const UObject* AnimContext, bool bCheckIsInGameThread = true);
 
diff --git a/Engine/Plugins/Animation/PoseSearch/Source/Runtime/Public/PoseSearch/PoseSearchHistory.h b/Engine/Plugins/Animation/PoseSearch/Source/Runtime/Public/PoseSearch/PoseSearchHistory.h
--- a/Engine/Plugins/Animation/PoseSearch/Source/Runtime/Public/PoseSearch/PoseSearchHistory.h
+++ b/Engine/Plugins/Animation/PoseSearch/Source/Runtime/Public/PoseSearch/PoseSearchHistory.h
@@ -94,7 +94,15 @@ struct FPoseHistoryEntry
 	TArray<float> CurveValues;
 	float AccumulatedSeconds = 0.f;
 
-	void Update(float Time, IComponentSpacePoseProvider& ComponentSpacePoseProvider, const FBoneToTransformMap& BoneToTransformMap, bool bStoreScales, const FBlendedCurve& Curves, const TConstArrayView<FName>& CollectedCurves);
+#if WITH_EDITORONLY_DATA
+	FTransform DebugMeshTransform;
+#endif // WITH_EDITORONLY_DATA
+
+	void Update(float Time, IComponentSpacePoseProvider& ComponentSpacePoseProvider, const FBoneToTransformMap& BoneToTransformMap, bool bStoreScales, const FBlendedCurve& Curves, const TConstArrayView<FName>& CollectedCurves
+#if WITH_EDITORONLY_DATA
+		, const FTransform& InDebugMeshTransform = FTransform::Identity
+#endif // WITH_EDITORONLY_DATA
+	);
 
 	POSESEARCH_API void SetNum(int32 Num, bool bStoreScales);
 	int32 Num() const;
@@ -194,7 +202,11 @@ struct FPoseHistory : public IPoseHistory
 	UE_API void EvaluateComponentSpace_AnyThread(float DeltaTime, IComponentSpacePoseProvider& ComponentSpacePoseProvider, bool bStoreScales,
 		float RootBoneRecoveryTime, float RootBoneTranslationRecoveryRatio, float RootBoneRotationRecoveryRatio,
 		bool bNeedsReset, bool bCacheBones, const TArray<FBoneIndexType>& RequiredBones, 
-		const FBlendedCurve& Curves = FBlendedCurve(), const TConstArrayView<FName>& InCollectedCurves = TConstArrayView<FName>());
+		const FBlendedCurve& Curves = FBlendedCurve(), const TConstArrayView<FName>& InCollectedCurves = TConstArrayView<FName>()
+#if WITH_EDITORONLY_DATA
+		, const FTransform& InDebugMeshTransform = FTransform::Identity
+#endif //WITH_EDITORONLY_DATA
+	);
 
 	// IPoseHistory interface
 	UE_API virtual bool GetTransformAtTime(float Time, FTransform& OutBoneTransform, const USkeleton* BoneIndexSkeleton = nullptr, FBoneIndexType BoneIndexType = RootBoneIndexType, FBoneIndexType ReferenceBoneIndexType = ComponentSpaceIndexType, bool bExtrapolate = false) const override;
@@ -295,12 +307,20 @@ struct FMemStackPoseHistory : public IPoseHistory
 	// End of IPoseHistory interface
 
 	UE_API void AddFutureRootBone(float Time, const FTransform& FutureRootBoneTransform, bool bStoreScales);
-	UE_API void AddFuturePose(float Time, IComponentSpacePoseProvider& ComponentSpacePoseProvider, const FBlendedCurve& Curves = FBlendedCurve());
+	UE_API void AddFuturePose(float Time, IComponentSpacePoseProvider& ComponentSpacePoseProvider, const FBlendedCurve& Curves = FBlendedCurve()
+#if WITH_EDITORONLY_DATA
+		, const FTransform& InDebugMeshTransform = FTransform::Identity
+#endif //WITH_EDITORONLY_DATA
+	);
 	
 	// Experimental, this feature might be removed without warning, not for production use
 	// extracting a pose from AnimationAsset at AnimationTime and  (using BlendParameters in case it's a blend space) and placing it in the FMemStackPoseHistory at IntervalTime seconds
 	// if FiniteDeltaTime is > FiniteDeltaTime a second pose will be extracted at AnimationTime - FiniteDeltaTime and placed at IntervalTime - FiniteDeltaTime (for MM to be able to calculate velocities)
-	UE_API void ExtractAndAddFuturePoses(const UAnimationAsset* AnimationAsset, float AnimationTime, float FiniteDeltaTime, const FVector& BlendParameters, float IntervalTime, const USkeleton* OverrideSkeleton = nullptr, bool bUseRefPoseRootBone = false);
+	UE_API void ExtractAndAddFuturePoses(const UAnimationAsset* AnimationAsset, float AnimationTime, float FiniteDeltaTime, const FVector& BlendParameters, float IntervalTime, const USkeleton* OverrideSkeleton = nullptr, bool bUseRefPoseRootBone = false
+#if WITH_EDITORONLY_DATA
+		, const FTransform& InDebugMeshTransform = FTransform::Identity
+#endif //WITH_EDITORONLY_DATA
+	);
 
 	const IPoseHistory* GetThisOrPoseHistory() const { return FutureEntries.IsEmpty() ? PoseHistory : this; }
 
diff --git a/Engine/Plugins/Experimental/UAF/UAFPoseSearch/Source/UAFPoseSearch/Private/HistoryCollectorTrait.cpp b/Engine/Plugins/Experimental/UAF/UAFPoseSearch/Source/UAFPoseSearch/Private/HistoryCollectorTrait.cpp
--- a/Engine/Plugins/Experimental/UAF/UAFPoseSearch/Source/UAFPoseSearch/Private/HistoryCollectorTrait.cpp
+++ b/Engine/Plugins/Experimental/UAF/UAFPoseSearch/Source/UAFPoseSearch/Private/HistoryCollectorTrait.cpp
@@ -7,8 +7,9 @@
 #include "EvaluationVM/EvaluationVM.h"
 #include "GenerationTools.h"
 #include "Module/AnimNextModuleInstance.h"
-#include "TraitCore/NodeInstance.h"
 #include "PoseHistoryEvaluation.h"
+#include "PoseSearch/PoseSearchContext.h"
+#include "TraitCore/NodeInstance.h"
 
 namespace UE::UAF
 {
@@ -158,7 +159,12 @@ void FAnimNextHistoryCollectorTask::Execute(UE::UAF::FEvaluationVM& VM) const
 			FComponentSpacePoseProvider ComponentSpacePoseProvider(Pose);
 			PoseHistory->EvaluateComponentSpace_AnyThread(DeltaTime, ComponentSpacePoseProvider, bStoreScales,
 				SharedData->RootBoneRecoveryTime, SharedData->RootBoneTranslationRecoveryRatio, SharedData->RootBoneRotationRecoveryRatio,
-				bNeedsReset, bCacheBones, BoneIndicesWithParents);
+				bNeedsReset, bCacheBones, BoneIndicesWithParents
+#if WITH_EDITORONLY_DATA
+				, FBlendedCurve(), TConstArrayView<FName>()
+				, AnimContext.IsValid() ? UE::PoseSearch::GetContextTransform(AnimContext.Pin().Get(), false) : FTransform::Identity
+#endif //WITH_EDITORONLY_DATA
+			);
 		}
 	}
 }
@@ -274,6 +280,14 @@ namespace UE::UAF
 		Task.bStoreScales = bStoreScales;
 		Task.DeltaTime = InstanceData->DeltaTime;
 		Task.HostObject = Context.GetHostObject();
+
+#if WITH_EDITORONLY_DATA
+		if (const FAnimNextModuleInstance* ModuleInstance = Context.GetRootGraphInstance().GetModuleInstance())
+		{
+			Task.AnimContext = ModuleInstance->GetObject();
+		}
+#endif // WITH_EDITORONLY_DATA
+
 		Context.AppendTask(Task);
 	}
 
diff --git a/Engine/Plugins/Experimental/UAF/UAFPoseSearch/Source/UAFPoseSearch/Private/HistoryCollectorTrait.h b/Engine/Plugins/Experimental/UAF/UAFPoseSearch/Source/UAFPoseSearch/Private/HistoryCollectorTrait.h
--- a/Engine/Plugins/Experimental/UAF/UAFPoseSearch/Source/UAFPoseSearch/Private/HistoryCollectorTrait.h
+++ b/Engine/Plugins/Experimental/UAF/UAFPoseSearch/Source/UAFPoseSearch/Private/HistoryCollectorTrait.h
@@ -27,6 +27,9 @@ struct FAnimNextHistoryCollectorTask : public FAnimNextEvaluationTask
 	bool bStoreScales = false;
 	float DeltaTime = 0.f;
 	TWeakObjectPtr<const UObject> HostObject;
+#if WITH_EDITORONLY_DATA
+	TWeakObjectPtr<const UObject> AnimContext;
+#endif // WITH_EDITORONLY_DATA
 };
 
 namespace UE::UAF
