PoseSearch - interaction cost biases, fix for FInteractionIsland::AddTickDependencies

diff --git a/Engine/Plugins/Animation/PoseSearch/Source/Editor/Private/PoseSearchDebuggerDatabaseColumns.h b/Engine/Plugins/Animation/PoseSearch/Source/Editor/Private/PoseSearchDebuggerDatabaseColumns.h
--- a/Engine/Plugins/Animation/PoseSearch/Source/Editor/Private/PoseSearchDebuggerDatabaseColumns.h
+++ b/Engine/Plugins/Animation/PoseSearch/Source/Editor/Private/PoseSearchDebuggerDatabaseColumns.h
@@ -115,12 +115,14 @@ struct FPoseIdx : ITextColumn
 
 	virtual FText GetLabel() const override
 	{
-		return LOCTEXT("ColumnLabelPoseIndex", "Pose Id");
+		static const FText Label = LOCTEXT("ColumnLabelPoseIndex", "Pose Id");
+		return Label;
 	}
 
 	virtual FText GetLabelTooltip() const override
 	{
-		return LOCTEXT("ColumnLabelTooltipPoseIndex", "Index of the Pose in the Database");
+		static const FText LabelTooltip = LOCTEXT("ColumnLabelTooltipPoseIndex", "Index of the Pose in the Database");
+		return LabelTooltip;
 	}
 
 	virtual FSortPredicate GetSortPredicate() const override
@@ -147,12 +149,14 @@ struct FAssetIdx : ITextColumn
 
 	virtual FText GetLabel() const override
 	{
-		return LOCTEXT("ColumnLabelAssetIndex", "Asset Id");
+		static const FText Label = LOCTEXT("ColumnLabelAssetIndex", "Asset Id");
+		return Label;
 	}
 
 	virtual FText GetLabelTooltip() const override
 	{
-		return LOCTEXT("ColumnLabelTooltipAssetIndex", "Index of the Asset in the Database");
+		static const FText LabelTooltip = LOCTEXT("ColumnLabelTooltipAssetIndex", "Index of the Asset in the Database");
+		return LabelTooltip;
 	}
 
 	virtual FSortPredicate GetSortPredicate() const override
@@ -186,12 +190,14 @@ struct FDatabaseName : IColumn
 
 	virtual FText GetLabel() const override
 	{
-		return LOCTEXT("ColumnLabelDatabaseName", "Database");
+		static const FText Label = LOCTEXT("ColumnLabelDatabaseName", "Database");
+		return Label;
 	}
 
 	virtual FText GetLabelTooltip() const override
 	{
-		return LOCTEXT("ColumnLabelTooltipDatabaseName", "Database Name");
+		static const FText LabelTooltip = LOCTEXT("ColumnLabelTooltipDatabaseName", "Database Name");
+		return LabelTooltip;
 	}
 
 	virtual FSortPredicate GetSortPredicate() const override
@@ -271,12 +277,14 @@ struct FAssetName : IColumn
 
 	virtual FText GetLabel() const override
 	{
-		return LOCTEXT("ColumnLabelAssetName", "Asset");
+		static const FText Label = LOCTEXT("ColumnLabelAssetName", "Asset");
+		return Label;
 	}
 	
 	virtual FText GetLabelTooltip() const override
 	{
-		return LOCTEXT("ColumnLabelTooltipAssetName", "Animation Asset Name");
+		static const FText LabelTooltip = LOCTEXT("ColumnLabelTooltipAssetName", "Animation Asset Name");
+		return LabelTooltip;
 	}
 
 	virtual FSortPredicate GetSortPredicate() const override
@@ -372,12 +380,14 @@ struct FFrame : ITextColumn
 
 	virtual FText GetLabel() const override
 	{
-		return LOCTEXT("ColumnLabelFrame", "Frame");
+		static const FText Label = LOCTEXT("ColumnLabelFrame", "Frame");
+		return Label;
 	}
 	
 	virtual FText GetLabelTooltip() const override
 	{
-		return LOCTEXT("ColumnLabelTooltipFrame", "Frame number from the start of the Animation Asset");
+		static const FText LabelTooltip = LOCTEXT("ColumnLabelTooltipFrame", "Frame number from the start of the Animation Asset");
+		return LabelTooltip;
 	}
 
 	virtual FSortPredicate GetSortPredicate() const override
@@ -404,12 +414,14 @@ struct FTime : ITextColumn
 
 	virtual FText GetLabel() const override
 	{
-		return LOCTEXT("ColumnLabelTime", "Time");
+		static const FText Label = LOCTEXT("ColumnLabelTime", "Time");
+		return Label;
 	}
 
 	virtual FText GetLabelTooltip() const override
 	{
-		return LOCTEXT("ColumnLabelTooltipTime", "Time in seconds from the start of the Animation Asset");
+		static const FText LabelTooltip = LOCTEXT("ColumnLabelTooltipTime", "Time in seconds from the start of the Animation Asset");
+		return LabelTooltip;
 	}
 
 	virtual FSortPredicate GetSortPredicate() const override
@@ -436,12 +448,14 @@ struct FPercentage : ITextColumn
 
 	virtual FText GetLabel() const override
 	{
-		return LOCTEXT("ColumnLabelPercentage", "Percent");
+		static const FText Label = LOCTEXT("ColumnLabelPercentage", "Percent");
+		return Label;
 	}
 
 	virtual FText GetLabelTooltip() const override
 	{
-		return LOCTEXT("ColumnLabelTooltipPercentage", "Time in percentage from the start of the Animation Asset");
+		static const FText LabelTooltip = LOCTEXT("ColumnLabelTooltipPercentage", "Time in percentage from the start of the Animation Asset");
+		return LabelTooltip;
 	}
 
 	virtual FSortPredicate GetSortPredicate() const override
@@ -468,12 +482,14 @@ struct FCost : ITextColumn
 
 	virtual FText GetLabel() const override
 	{
-		return LOCTEXT("ColumnLabelCost", "Cost");
+		static const FText Label = LOCTEXT("ColumnLabelCost", "Cost");
+		return Label;
 	}
 	
 	virtual FText GetLabelTooltip() const override
 	{
-		return LOCTEXT("ColumnLabelTooltipCost", "Total Cost of the associated Pose");
+		static const FText LabelTooltip = LOCTEXT("ColumnLabelTooltipCost", "Total Cost of the associated Pose");
+		return LabelTooltip;
 	}
 
 	virtual FSortPredicate GetSortPredicate() const override
@@ -505,12 +521,14 @@ struct FPCACost : ITextColumn
 
 	virtual FText GetLabel() const override
 	{
-		return LOCTEXT("ColumnLabelPCACost", "PCA Cost");
+		static const FText Label = LOCTEXT("ColumnLabelPCACost", "PCA Cost");
+		return Label;
 	}
 	
 	virtual FText GetLabelTooltip() const override
 	{
-		return LOCTEXT("ColumnLabelTooltipPCACost", "Total PCA Cost of the associated Pose");
+		static const FText LabelTooltip = LOCTEXT("ColumnLabelTooltipPCACost", "Total PCA Cost of the associated Pose");
+		return LabelTooltip;
 	}
 
 	virtual FSortPredicate GetSortPredicate() const override
@@ -541,6 +559,7 @@ struct FChannelBreakdownCostColumn : ITextColumn
 	FChannelBreakdownCostColumn(int32 SortIndex, int32 InBreakdownCostIndex, const FText& InLabel)
 		: ITextColumn(SortIndex)
 		, Label(InLabel)
+		, LabelTooltip(FText::Format(LOCTEXT("ColumnLabelTooltipChannelBreakdownCost", "Breakdown Cost for the Channel '{0}'"), InLabel))
 		, BreakdownCostIndex(InBreakdownCostIndex)
 	{
 	}
@@ -552,7 +571,7 @@ struct FChannelBreakdownCostColumn : ITextColumn
 
 	virtual FText GetLabelTooltip() const override
 	{
-		return FText::Format(LOCTEXT("ColumnLabelTooltipChannelBreakdownCost", "Breakdown Cost for the Channel '{0}'"), Label);
+		return LabelTooltip;
 	}
 
 	virtual FSortPredicate GetSortPredicate() const override
@@ -594,6 +613,7 @@ struct FChannelBreakdownCostColumn : ITextColumn
 	}
 
 	FText Label;
+	FText LabelTooltip;
 	int32 BreakdownCostIndex = INDEX_NONE;
 };
 
@@ -603,12 +623,14 @@ struct FCostModifier : ITextColumn
 
 	virtual FText GetLabel() const override
 	{
-		return LOCTEXT("ColumnLabelCostModifier", "Bias");
+		static const FText Label = LOCTEXT("ColumnLabelCostModifier", "Bias");
+		return Label;
 	}
 
 	virtual FText GetLabelTooltip() const override
 	{
-		return LOCTEXT("ColumnLabelTooltipCostModifier", "Total Cost for all the Bias contributions");
+		static const FText LabelTooltip = LOCTEXT("ColumnLabelTooltipCostModifier", "Total Cost for all the Bias contributions");
+		return LabelTooltip;
 	}
 
 	virtual FSortPredicate GetSortPredicate() const override
@@ -635,12 +657,14 @@ struct FContinuingPoseCost : ITextColumn
 
 	virtual FText GetLabel() const override
 	{
-		return LOCTEXT("ColumnLabelContinuingPoseCost", "ContinuingPoseCost");
+		static const FText Label = LOCTEXT("ColumnLabelContinuingPoseCost", "ContinuingPoseCost");
+		return Label;
 	}
 
 	virtual FText GetLabelTooltip() const override
 	{
-		return LOCTEXT("ColumnLabelTooltipContinuingPoseCost", "Continuing Pose Cost Bias contribution");
+		static const FText LabelTooltip = LOCTEXT("ColumnLabelTooltipContinuingPoseCost", "Continuing Pose Cost Bias contribution");
+		return LabelTooltip;
 	}
 
 	virtual FSortPredicate GetSortPredicate() const override
@@ -667,12 +691,14 @@ struct FContinuingInteractionCost : ITextColumn
 
 	virtual FText GetLabel() const override
 	{
-		return LOCTEXT("ColumnLabelContinuingInteractionCost", "ContinuingInteractionCost");
+		static const FText Label = LOCTEXT("ColumnLabelContinuingInteractionCost", "ContinuingInteractionCost");
+		return Label;
 	}
 
 	virtual FText GetLabelTooltip() const override
 	{
-		return LOCTEXT("ColumnLabelTooltipContinuingInteractionCost", "Continuing Interaction Cost Bias contribution");
+		static const FText LabelTooltip = LOCTEXT("ColumnLabelTooltipContinuingInteractionCost", "Continuing Interaction Cost Bias contribution (interaction with the same anim contexts with the same assigned roles)");
+		return LabelTooltip;
 	}
 
 	virtual FSortPredicate GetSortPredicate() const override
@@ -693,18 +719,54 @@ struct FContinuingInteractionCost : ITextColumn
 	}
 };
 
+struct FContinuingContextInteractionCost : ITextColumn
+{
+	using ITextColumn::ITextColumn;
+
+	virtual FText GetLabel() const override
+	{
+		static const FText Label = LOCTEXT("ColumnLabelContinuingContextInteractionCost", "ContinuingContextInteractionCost");
+		return Label;
+	}
+
+	virtual FText GetLabelTooltip() const override
+	{
+		static const FText LabelTooltip = LOCTEXT("ColumnLabelTooltipContinuingContextInteractionCost", "Continuing Context Interaction Cost Bias contribution (interaction with the same anim contexts, assigned roles may differ)");
+		return LabelTooltip;
+	}
+
+	virtual FSortPredicate GetSortPredicate() const override
+	{
+		return [](const FRowDataRef& Row0, const FRowDataRef& Row1) -> bool
+			{
+				return Row0->PoseCost.GetContinuingContextInteractionCostAddend() < Row1->PoseCost.GetContinuingContextInteractionCostAddend();
+			};
+	}
+
+	virtual FText GetRowText(const FRowDataRef& Row) const override
+	{
+		if (Row->IsValid())
+		{
+			return FText::AsNumber(Row->PoseCost.GetContinuingContextInteractionCostAddend());
+		}
+		return FText::FromString(TEXT("-"));
+	}
+};
+
 struct FNotifyCost : ITextColumn
 {
 	using ITextColumn::ITextColumn;
 
 	virtual FText GetLabel() const override
 	{
-		return LOCTEXT("ColumnLabelNotifyCost", "NotifyCost");
+		static const FText Label = LOCTEXT("ColumnLabelNotifyCost", "NotifyCost");
+		return Label;
 	}
 
 	virtual FText GetLabelTooltip() const override
 	{
-		return LOCTEXT("ColumnLabelTooltipNotifyCost", "Notify Cost Bias contribution");
+		static const FText LabelTooltip = LOCTEXT("ColumnLabelTooltipNotifyCost", "Notify Cost Bias contribution");
+		return LabelTooltip;
 	}
 
 	virtual FSortPredicate GetSortPredicate() const override
@@ -732,12 +794,14 @@ struct FMirrored : ITextColumn
 
 	virtual FText GetLabel() const override
 	{
-		return LOCTEXT("ColumnLabelMirrored", "Mirror");
+		static const FText Label = LOCTEXT("ColumnLabelMirrored", "Mirror");
+		return Label;
 	}
 
 	virtual FText GetLabelTooltip() const override
 	{
-		return LOCTEXT("ColumnLabelTooltipMirrored", "Mirror state of the associated Pose");
+		static const FText LabelTooltip = LOCTEXT("ColumnLabelTooltipMirrored", "Mirror state of the associated Pose");
+		return LabelTooltip;
 	}
 
 	virtual FSortPredicate GetSortPredicate() const override
@@ -750,11 +814,21 @@ struct FMirrored : ITextColumn
 
 	virtual FText GetRowText(const FRowDataRef& Row) const override
 	{
-		if (Row->IsValid())
+		static const FText RowTextMirrored = FText::Format(LOCTEXT("ColumnMirrored_Mirrored", "{0}"), { true } );
+		static const FText RowTextUnmirrored = FText::Format(LOCTEXT("ColumnMirrored_Mirrored", "{0}"), { false } );
+		static const FText RowTextUnknown = LOCTEXT("ColumnMirrored_Unknown", "-");
+
+		if (!Row->IsValid())
 		{
-			return FText::Format(LOCTEXT("Mirrored", "{0}"), { Row->bMirrored } );
+			return RowTextUnknown;
 		}
-		return FText::FromString(TEXT("-"));
+
+		if (Row->bMirrored)
+		{
+			return RowTextMirrored;
+		}
+			
+		return RowTextUnmirrored;
 	}
 };
 
@@ -764,12 +838,14 @@ struct FLooping : ITextColumn
 
 	virtual FText GetLabel() const override
 	{
-		return LOCTEXT("ColumnLabelLooping", "Loop");
+		static const FText Label = LOCTEXT("ColumnLabelLooping", "Loop");
+		return Label;
 	}
 
 	virtual FText GetLabelTooltip() const override
 	{
-		return LOCTEXT("ColumnLabelTooltipLooping", "Loop state of the associated Pose");
+		static const FText LabelTooltip = LOCTEXT("ColumnLabelTooltipLooping", "Loop state of the associated Pose");
+		return LabelTooltip;
 	}
 
 	virtual FSortPredicate GetSortPredicate() const override
@@ -782,11 +858,21 @@ struct FLooping : ITextColumn
 
 	virtual FText GetRowText(const FRowDataRef& Row) const override
 	{
-		if (Row->IsValid())
+		static const FText RowTextLooping = FText::Format(LOCTEXT("ColumnLooping_Looping", "{0}"), { true } );
+		static const FText RowTextNotLooping = FText::Format(LOCTEXT("ColumnLooping_NotLooping", "{0}"), { false } );
+		static const FText RowTextUnknown = LOCTEXT("ColumnLooping_Unknown", "-");
+
+		if (!Row->IsValid())
 		{
-			return FText::Format(LOCTEXT("Looping", "{0}"), { Row->bLooping });
+			return RowTextUnknown;
 		}
-		return FText::FromString(TEXT("-"));
+
+		if (Row->bLooping)
+		{
+			return RowTextLooping;
+		}
+			
+		return RowTextNotLooping;
 	}
 };
 
@@ -796,12 +882,14 @@ struct FBlendParameters : ITextColumn
 
 	virtual FText GetLabel() const override
 	{
-		return LOCTEXT("ColumnLabelBlendParams", "Blend Params");
+		static const FText Label = LOCTEXT("ColumnLabelBlendParams", "Blend Params");
+		return Label;
 	}
 
 	virtual FText GetLabelTooltip() const override
 	{
-		return LOCTEXT("ColumnLabelBlendTooltipParams", "Blend Params used to sample the associated BlendSpace asset");
+		static const FText LabelTooltip = LOCTEXT("ColumnLabelBlendTooltipParams", "Blend Params used to sample the associated BlendSpace asset");
+		return LabelTooltip;
 	}
 
 	virtual FSortPredicate GetSortPredicate() const override
@@ -830,12 +918,14 @@ struct FPoseCandidateFlags : ITextColumn
 
 	virtual FText GetLabel() const override
 	{
-		return LOCTEXT("ColumnLabelPoseCandidateFlags", "Flags");
+		static const FText Label = LOCTEXT("ColumnLabelPoseCandidateFlags", "Flags");
+		return Label;
 	}
 
 	virtual FText GetLabelTooltip() const override
 	{
-		return LOCTEXT("ColumnLabelTooltipPoseCandidateFlags", "Flags indicating why a Pose has been discarded");
+		static const FText LabelTooltip = LOCTEXT("ColumnLabelTooltipPoseCandidateFlags", "Flags indicating why a Pose has been discarded");
+		return LabelTooltip;
 	}
 
 	virtual FSortPredicate GetSortPredicate() const override
diff --git a/Engine/Plugins/Animation/PoseSearch/Source/Editor/Private/PoseSearchDebuggerDatabaseView.cpp b/Engine/Plugins/Animation/PoseSearch/Source/Editor/Private/PoseSearchDebuggerDatabaseView.cpp
--- a/Engine/Plugins/Animation/PoseSearch/Source/Editor/Private/PoseSearchDebuggerDatabaseView.cpp
+++ b/Engine/Plugins/Animation/PoseSearch/Source/Editor/Private/PoseSearchDebuggerDatabaseView.cpp
@@ -211,10 +211,11 @@ static void AddUnfilteredDatabaseRow(const UPoseSearchDatabase* Database, TConst
 		}
 		else
 		{
-			// @todo: gather ContinuingPoseCostAddend, ContinuingInteractionCostAddend properly to reconstruct the real FPoseSearchCost
+			// @todo: gather ContinuingPoseCostAddend, ContinuingInteractionCostAddend, and ContinuingContextInteractionCostAddend properly to reconstruct the real FPoseSearchCost
 			const float ContinuingPoseCostAddend = 0.f;
 			const float ContinuingInteractionCostAddend = 0.f;
-			Row->PoseCost = FPoseSearchCost(ArraySum(Row->CostVector, 0, Database->Schema->SchemaCardinality), SearchIndex.PoseMetadata[DbPoseIdx].GetCostAddend(), ContinuingPoseCostAddend, ContinuingInteractionCostAddend);
+			const float ContinuingContextInteractionCostAddend = 0.f;
+			Row->PoseCost = FPoseSearchCost(ArraySum(Row->CostVector, 0, Database->Schema->SchemaCardinality), SearchIndex.PoseMetadata[DbPoseIdx].GetCostAddend(), ContinuingPoseCostAddend, ContinuingInteractionCostAddend, ContinuingContextInteractionCostAddend);
 		}
 
 		if (!SharedData->PCAQueryVector.IsEmpty())
@@ -443,6 +444,7 @@ void SDebuggerDatabaseView::Update(const FTraceMotionMatchingStateMessage& State
 			AddColumn(MakeShared<FContinuingPoseCost>(ColumnIdx++));
 			AddColumn(MakeShared<FNotifyCost>(ColumnIdx++));
 			AddColumn(MakeShared<FContinuingInteractionCost>(ColumnIdx++));
+			AddColumn(MakeShared<FContinuingContextInteractionCost>(ColumnIdx++));
 		}
 		else
 		{
diff --git a/Engine/Plugins/Animation/PoseSearch/Source/Runtime/Private/PoseSearchContext.cpp b/Engine/Plugins/Animation/PoseSearch/Source/Runtime/Private/PoseSearchContext.cpp
--- a/Engine/Plugins/Animation/PoseSearch/Source/Runtime/Private/PoseSearchContext.cpp
+++ b/Engine/Plugins/Animation/PoseSearch/Source/Runtime/Private/PoseSearchContext.cpp
@@ -1336,7 +1336,7 @@ void FSearchContext::Track(const UPoseSearchDatabase* Database, int32 PoseIdx, E
 	}
 }
 
-void FSearchContext::Track(const UPoseSearchDatabase* Database, int32 PoseIdx, EPoseCandidateFlags PoseCandidateFlags, TConstArrayView<float> DynamicWeightsSqrt, float ContinuingPoseCostAddend, float ContinuingInteractionCostAddend)
+void FSearchContext::Track(const UPoseSearchDatabase* Database, int32 PoseIdx, EPoseCandidateFlags PoseCandidateFlags, TConstArrayView<float> DynamicWeightsSqrt, float ContinuingPoseCostAddend, float ContinuingInteractionCostAddend, float ContinuingContextInteractionCostAddend)
 {
 	check(PoseIdx >= 0 && Database && Database->Schema);
 
@@ -1364,7 +1364,7 @@ void FSearchContext::Track(const UPoseSearchDatabase* Database, int32 PoseIdx, E
 			TConstArrayView<float> QueryValues = GetOrBuildQuery(Database->Schema);
 			DissimilarityCost = CompareFeatureVectors<false>(SearchIndex.GetPoseValues(PoseIdx), QueryValues, DynamicWeightsSqrt);
 		}
-		const FPoseSearchCost PoseCost(DissimilarityCost, SearchIndex.PoseMetadata[PoseIdx].GetCostAddend(), ContinuingPoseCostAddend, ContinuingInteractionCostAddend);
+		const FPoseSearchCost PoseCost(DissimilarityCost, SearchIndex.PoseMetadata[PoseIdx].GetCostAddend(), ContinuingPoseCostAddend, ContinuingInteractionCostAddend, ContinuingContextInteractionCostAddend);
 		Track(Database, PoseIdx, PoseCandidateFlags, PoseCost);
 	}
 }
diff --git a/Engine/Plugins/Animation/PoseSearch/Source/Runtime/Private/PoseSearchDatabase.cpp b/Engine/Plugins/Animation/PoseSearch/Source/Runtime/Private/PoseSearchDatabase.cpp
--- a/Engine/Plugins/Animation/PoseSearch/Source/Runtime/Private/PoseSearchDatabase.cpp
+++ b/Engine/Plugins/Animation/PoseSearch/Source/Runtime/Private/PoseSearchDatabase.cpp
@@ -149,7 +149,7 @@ FPoseSearchDatabaseAnimationAsset UpgradeToDatabaseAnimationAsset(FInstancedStru
 
 template<bool bReconstructPoseValues, bool bAlignedAndPadded>
 static inline void EvaluatePoseKernel(FSearchResults& SearchResults, const FSearchIndex& SearchIndex, TConstArrayView<float> QueryValues, TArrayView<float> ReconstructedPoseValuesBuffer,
-	int32 PoseIdx, int32 EventPoseIdx, const FSearchFilters& SearchFilters, float ContinuingPoseCostAddend, float ContinuingInteractionCostAddend,
+	int32 PoseIdx, int32 EventPoseIdx, const FSearchFilters& SearchFilters, float ContinuingPoseCostAddend, float ContinuingInteractionCostAddend, float ContinuingContextInteractionCostAddend,
 	const UPoseSearchDatabase* Database, TConstArrayView<float> DynamicWeightsSqrt
 #if UE_POSE_SEARCH_TRACE_ENABLED
 	, bool bTrackResult = true
@@ -165,12 +165,12 @@ static inline void EvaluatePoseKernel(FSearchResults& SearchResults, const FSear
 
 	if (SearchFilters.AreFiltersValid(SearchIndex, PoseValues, QueryValues, DynamicWeightsSqrt, PoseIdx
 #if UE_POSE_SEARCH_TRACE_ENABLED
-		, ContinuingPoseCostAddend, ContinuingInteractionCostAddend, Database
+		, ContinuingPoseCostAddend, ContinuingInteractionCostAddend, ContinuingContextInteractionCostAddend, Database
 #endif // UE_POSE_SEARCH_TRACE_ENABLED
 	))
 	{
 		FSearchResult Result;
-		Result.PoseCost = FPoseSearchCost(CompareFeatureVectors<bAlignedAndPadded>(PoseValues, QueryValues, DynamicWeightsSqrt), SearchIndex.PoseMetadata[PoseIdx].GetCostAddend(), ContinuingPoseCostAddend, ContinuingInteractionCostAddend);
+		Result.PoseCost = FPoseSearchCost(CompareFeatureVectors<bAlignedAndPadded>(PoseValues, QueryValues, DynamicWeightsSqrt), SearchIndex.PoseMetadata[PoseIdx].GetCostAddend(), ContinuingPoseCostAddend, ContinuingInteractionCostAddend, ContinuingContextInteractionCostAddend);
 		Result.PoseIdx = PoseIdx;
 		Result.EventPoseIdx = EventPoseIdx;
 		Result.Database = Database;
@@ -1381,7 +1381,7 @@ void UPoseSearchDatabase::Search(UE::PoseSearch::FSearchContext& SearchContext,
 
 void UPoseSearchDatabase::PopulateNonSelectableIdx(FNonSelectableIdx& NonSelectableIdx, UE::PoseSearch::FSearchContext& SearchContext
 #if UE_POSE_SEARCH_TRACE_ENABLED
-	, float ContinuingPoseCostAddend, float ContinuingInteractionCostAddend, TConstArrayView<float> DynamicWeightsSqrt
+	, float ContinuingPoseCostAddend, float ContinuingInteractionCostAddend, float ContinuingContextInteractionCostAddend, TConstArrayView<float> DynamicWeightsSqrt
 #endif //UE_POSE_SEARCH_TRACE_ENABLED
 ) const
 {
@@ -1411,7 +1411,7 @@ void UPoseSearchDatabase::PopulateNonSelectableIdx(FNonSelectableIdx& NonSelecta
 						NonSelectableIdx.Add(PoseIdx);
 
 #if UE_POSE_SEARCH_TRACE_ENABLED
-						SearchContext.Track(this, PoseIdx, EPoseCandidateFlags::DiscardedBy_AssetReselection, DynamicWeightsSqrt, ContinuingPoseCostAddend, ContinuingInteractionCostAddend);
+						SearchContext.Track(this, PoseIdx, EPoseCandidateFlags::DiscardedBy_AssetReselection, DynamicWeightsSqrt, ContinuingPoseCostAddend, ContinuingInteractionCostAddend, ContinuingContextInteractionCostAddend);
 #endif // UE_POSE_SEARCH_TRACE_ENABLED
 					}
 				}
@@ -1419,6 +1419,7 @@ void UPoseSearchDatabase::PopulateNonSelectableIdx(FNonSelectableIdx& NonSelecta
 			else if (!FMath::IsNearlyEqual(SearchContext.GetPoseJumpThresholdTime().Min, SearchContext.GetPoseJumpThresholdTime().Max))
 			{
 				const int32 ContinuingPoseIdx = SearchContext.GetContinuingPoseSearchResult().PoseIdx;
+				check(ContinuingPoseIdx != INDEX_NONE);
 				const int32 UnboundMinPoseIdx = ContinuingPoseIdx + FMath::FloorToInt(SearchContext.GetPoseJumpThresholdTime().Min * Schema->SampleRate);
 				const int32 UnboundMaxPoseIdx = ContinuingPoseIdx + FMath::CeilToInt(SearchContext.GetPoseJumpThresholdTime().Max * Schema->SampleRate);
 				const int32 CurrentIndexAssetFirstPoseIdx = CurrentIndexAsset->GetFirstPoseIdx();
@@ -1436,7 +1437,7 @@ void UPoseSearchDatabase::PopulateNonSelectableIdx(FNonSelectableIdx& NonSelecta
 						NonSelectableIdx.AddUnique(PoseIdx);
 
 #if UE_POSE_SEARCH_TRACE_ENABLED
-						SearchContext.Track(this, PoseIdx, EPoseCandidateFlags::DiscardedBy_PoseJumpThresholdTime, DynamicWeightsSqrt, ContinuingPoseCostAddend, ContinuingInteractionCostAddend);
+						SearchContext.Track(this, PoseIdx, EPoseCandidateFlags::DiscardedBy_PoseJumpThresholdTime, DynamicWeightsSqrt, ContinuingPoseCostAddend, ContinuingInteractionCostAddend, ContinuingContextInteractionCostAddend);
 #endif // UE_POSE_SEARCH_TRACE_ENABLED
 					}
 				}
@@ -1449,7 +1450,7 @@ void UPoseSearchDatabase::PopulateNonSelectableIdx(FNonSelectableIdx& NonSelecta
 					{
 						NonSelectableIdx.AddUnique(PoseIdx);
 #if UE_POSE_SEARCH_TRACE_ENABLED
-						SearchContext.Track(this, PoseIdx, EPoseCandidateFlags::DiscardedBy_PoseJumpThresholdTime, DynamicWeightsSqrt, ContinuingPoseCostAddend, ContinuingInteractionCostAddend);
+						SearchContext.Track(this, PoseIdx, EPoseCandidateFlags::DiscardedBy_PoseJumpThresholdTime, DynamicWeightsSqrt, ContinuingPoseCostAddend, ContinuingInteractionCostAddend, ContinuingContextInteractionCostAddend);
 #endif // UE_POSE_SEARCH_TRACE_ENABLED
 					}
 				}
@@ -1472,7 +1473,7 @@ void UPoseSearchDatabase::PopulateNonSelectableIdx(FNonSelectableIdx& NonSelecta
 					NonSelectableIdx.AddUnique(HistoricalPoseIndex.PoseIndex);
 
 #if UE_POSE_SEARCH_TRACE_ENABLED
-					SearchContext.Track(this, HistoricalPoseIndex.PoseIndex, EPoseCandidateFlags::DiscardedBy_PoseReselectHistory, DynamicWeightsSqrt, ContinuingPoseCostAddend, ContinuingInteractionCostAddend);
+					SearchContext.Track(this, HistoricalPoseIndex.PoseIndex, EPoseCandidateFlags::DiscardedBy_PoseReselectHistory, DynamicWeightsSqrt, ContinuingPoseCostAddend, ContinuingInteractionCostAddend, ContinuingContextInteractionCostAddend);
 #endif // UE_POSE_SEARCH_TRACE_ENABLED
 				}
 			}
@@ -1550,30 +1551,55 @@ UE::PoseSearch::FSearchResult UPoseSearchDatabase::SearchContinuingPose(UE::Pose
 	return SearchResults.GetBestResult();
 }
 
-void UPoseSearchDatabase::SearchContinuingPose(UE::PoseSearch::FSearchContext& SearchContext, UE::PoseSearch::FSearchResults& SearchResults) const
+void UPoseSearchDatabase::ExtractCostAddends(const UE::PoseSearch::FSearchContext& SearchContext, bool bIsContinuingPoseSearch, float& ContinuingPoseCostAddend, float& ContinuingInteractionCostAddend, float& ContinuingContextInteractionCostAddend) const
 {
-	QUICK_SCOPE_CYCLE_COUNTER(STAT_PoseSearch_ContinuingPose);
-
 	using namespace UE::PoseSearch;
 
-	check(SearchContext.GetContinuingPoseSearchResult().Database.Get() == this);
+	ContinuingPoseCostAddend = 0.f;
+	ContinuingInteractionCostAddend = 0.f;
+	ContinuingContextInteractionCostAddend = 0.f;
 
-#if WITH_EDITOR
-	if (EAsyncBuildIndexResult::Success != FAsyncPoseSearchDatabasesManagement::RequestAsyncBuildIndex(this, ERequestAsyncBuildFlag::ContinueRequest))
+	bool bSearchForContinuingPoseCostAddend = false;
+	bool bSearchForContinuingInteractionCostAddend = false;
+	bool bSearchForNotifyStateContinuingContextInteractionCostBiasFound = false;
+
+	if (bIsContinuingPoseSearch)
 	{
-		SearchContext.SetAsyncBuildIndexInProgress();
+		bSearchForContinuingPoseCostAddend = true;
+		ContinuingPoseCostAddend = ContinuingPoseCostBias;
+	}
+	
+	if (SearchContext.IsContinuingInteraction())
+	{
+		ContinuingInteractionCostAddend = ContinuingInteractionCostBias;
+		bSearchForContinuingInteractionCostAddend = true;
+	}
+
+	if (SearchContext.IsContinuingContextInteraction())
+	{
+		bSearchForNotifyStateContinuingContextInteractionCostBiasFound = true;
+		ContinuingContextInteractionCostAddend = ContinuingContextInteractionCostBias;
+	}
+
+	if (!bSearchForContinuingPoseCostAddend && !bSearchForContinuingInteractionCostAddend && !bSearchForNotifyStateContinuingContextInteractionCostBiasFound)
+	{
+		// no need to search!
 		return;
 	}
-#endif // WITH_EDITOR
 
-	// extracting notifies from the database animation asset at time SampleTime to search for UAnimNotifyState_PoseSearchOverrideContinuingPoseCostBias eventually overriding the database ContinuingPoseCostBias
-	const FSearchIndex& SearchIndex = GetSearchIndex();
 	const int32 ContinuingPoseIdx = SearchContext.GetContinuingPoseSearchResult().PoseIdx;
+	if (ContinuingPoseIdx == INDEX_NONE)
+	{
+		// no need to search!
+		return;
+	}
+
+	// extracting notifies from the database animation asset at time SampleTime
+	const FSearchIndex& SearchIndex = GetSearchIndex();
 	const FSearchIndexAsset& SearchIndexAsset = SearchIndex.GetAssetForPose(ContinuingPoseIdx);
 	const FPoseSearchDatabaseAnimationAssetBase* DatabaseAnimationAssetBase = GetDatabaseAnimationAsset(SearchIndexAsset);
 	check(DatabaseAnimationAssetBase);
 
-	float ContinuingPoseCostAddend = ContinuingPoseCostBias;
 	const float SampleTime = GetRealAssetTime(ContinuingPoseIdx);
 	FAnimNotifyContext PreAllocatedNotifyContext;
 	for (int32 RoleIndex = 0; RoleIndex < DatabaseAnimationAssetBase->GetNumRoles(); ++RoleIndex)
@@ -1582,32 +1608,92 @@ void UPoseSearchDatabase::SearchContinuingPose(UE::PoseSearch::FSearchContext& S
 		{
 			// sampler used only to extract the notify states. RootTransformOrigin can be set as Identity, since will not be relevant
 			const FAnimationAssetSampler SequenceBaseSampler(AnimationAsset, FTransform::Identity, SearchIndexAsset.GetBlendParameters(), FAnimationAssetSampler::DefaultRootTransformSamplingRate, false, false);
-
-			bool bDone = false;
-			SequenceBaseSampler.ExtractAnimNotifyStates(SampleTime, PreAllocatedNotifyContext, [&ContinuingPoseCostAddend, &bDone](const UAnimNotifyState* AnimNotifyState)
+			SequenceBaseSampler.ExtractAnimNotifyStates(SampleTime, PreAllocatedNotifyContext, 
+				[&bSearchForContinuingPoseCostAddend, &bSearchForContinuingInteractionCostAddend, &bSearchForNotifyStateContinuingContextInteractionCostBiasFound,
+				&ContinuingPoseCostAddend, &ContinuingInteractionCostAddend, &ContinuingContextInteractionCostAddend](const UAnimNotifyState* AnimNotifyState)
 				{
-					if (const UAnimNotifyState_PoseSearchOverrideContinuingPoseCostBias* NotifyStateContinuingPoseCostBias = Cast<const UAnimNotifyState_PoseSearchOverrideContinuingPoseCostBias>(AnimNotifyState))
+					if (bSearchForContinuingPoseCostAddend)
+					{
+						if (const UAnimNotifyState_PoseSearchOverrideContinuingPoseCostBias* NotifyStateContinuingPoseCostBias = Cast<const UAnimNotifyState_PoseSearchOverrideContinuingPoseCostBias>(AnimNotifyState))
+						{
+							ContinuingPoseCostAddend = NotifyStateContinuingPoseCostBias->CostAddend;
+							bSearchForContinuingPoseCostAddend = false;
+
+							if (!bSearchForContinuingInteractionCostAddend && !bSearchForNotifyStateContinuingContextInteractionCostBiasFound)
+							{
+								// we're done searching!
+								return false;
+							}
+						}
+					}
+
+					if (bSearchForContinuingInteractionCostAddend)
 					{
-						ContinuingPoseCostAddend = NotifyStateContinuingPoseCostBias->CostAddend;
-						bDone = true;
-						return false;
+						if (const UAnimNotifyState_PoseSearchOverrideContinuingInteractionPoseCostBias* NotifyStateContinuingInteractionPoseCostBias = Cast<const UAnimNotifyState_PoseSearchOverrideContinuingInteractionPoseCostBias>(AnimNotifyState))
+						{
+							ContinuingInteractionCostAddend = NotifyStateContinuingInteractionPoseCostBias->CostAddend;
+							bSearchForContinuingInteractionCostAddend = false;
+
+							if (!bSearchForContinuingPoseCostAddend && !bSearchForNotifyStateContinuingContextInteractionCostBiasFound)
+							{
+								// we're done searching!
+								return false;
+							}
+						}
+					}
+
+					if (bSearchForNotifyStateContinuingContextInteractionCostBiasFound)
+					{
+						if (const UAnimNotifyState_PoseSearchOverrideContinuingContextInteractionCostBias* NotifyStateContinuingContextInteractionCostBias = Cast<const UAnimNotifyState_PoseSearchOverrideContinuingContextInteractionCostBias>(AnimNotifyState))
+						{
+							ContinuingContextInteractionCostAddend = NotifyStateContinuingContextInteractionCostBias->CostAddend;
+							bSearchForNotifyStateContinuingContextInteractionCostBiasFound = false;
+
+							if (!bSearchForContinuingPoseCostAddend && !bSearchForContinuingInteractionCostAddend)
+							{
+								// we're done searching!
+								return false;
+							}
+						}
 					}
 					return true;
 				});
+		}
 
-			if (bDone)
-			{
-				break;
-			}
+		if (!bSearchForContinuingPoseCostAddend && !bSearchForContinuingInteractionCostAddend && !bSearchForNotifyStateContinuingContextInteractionCostBiasFound)
+		{
+			// we're done searching!
+			return;
 		}
 	}
+}
+
+void UPoseSearchDatabase::SearchContinuingPose(UE::PoseSearch::FSearchContext& SearchContext, UE::PoseSearch::FSearchResults& SearchResults) const
+{
+	QUICK_SCOPE_CYCLE_COUNTER(STAT_PoseSearch_ContinuingPose);
 
-	const float ContinuingInteractionCostAddend = SearchContext.IsContinuingInteraction() ? ContinuingInteractionCostBias : 0.f;
+	using namespace UE::PoseSearch;
 
-	// since any PoseCost calculated here is at least SearchIndex.MinCostAddend + ContinuingPoseCostAddend + ContinuingInteractionCostAddend,
+	check(SearchContext.GetContinuingPoseSearchResult().Database.Get() == this);
+
+#if WITH_EDITOR
+	if (EAsyncBuildIndexResult::Success != FAsyncPoseSearchDatabasesManagement::RequestAsyncBuildIndex(this, ERequestAsyncBuildFlag::ContinueRequest))
+	{
+		SearchContext.SetAsyncBuildIndexInProgress();
+		return;
+	}
+#endif // WITH_EDITOR
+
+	float ContinuingPoseCostAddend = 0.f, ContinuingInteractionCostAddend = 0.f, ContinuingContextInteractionCostAddend = 0.f;
+	ExtractCostAddends(SearchContext, true, ContinuingPoseCostAddend, ContinuingInteractionCostAddend, ContinuingContextInteractionCostAddend);
+
+	// since any PoseCost calculated here is at least SearchIndex.MinCostAddend + ContinuingPoseCostAddend + ContinuingInteractionCostAddend + ContinuingContextInteractionCostAddend,
 	// there's no point in performing the search if the worst search result is already better than that
-	if (!GetSkipSearchIfPossible() || SearchResults.ShouldPerformSearch(SearchIndex.MinCostAddend + ContinuingPoseCostAddend + ContinuingInteractionCostAddend))
+	const FSearchIndex& SearchIndex = GetSearchIndex();
+	if (!GetSkipSearchIfPossible() || SearchResults.ShouldPerformSearch(SearchIndex.MinCostAddend + ContinuingPoseCostAddend + ContinuingInteractionCostAddend + ContinuingContextInteractionCostAddend))
 	{
+		const int32 ContinuingPoseIdx = SearchContext.GetContinuingPoseSearchResult().PoseIdx;
+		check(ContinuingPoseIdx != INDEX_NONE);
 		const int32 NumDimensions = Schema->SchemaCardinality;
 		// FMemory_Alloca is forced 16 bytes aligned
 		TArrayView<float> ReconstructedPoseValuesBuffer((float*)FMemory_Alloca(NumDimensions * sizeof(float)), NumDimensions);
@@ -1624,12 +1710,12 @@ void UPoseSearchDatabase::SearchContinuingPose(UE::PoseSearch::FSearchContext& S
 		// is the data padded at 16 bytes (and 16 bytes aligned by construction)?
 		if (NumDimensions % 4 == 0)
 		{
-			Result.PoseCost = FPoseSearchCost(CompareFeatureVectors<true>(PoseValues, SearchContext.GetOrBuildQuery(Schema), DynamicWeightsSqrt), SearchIndex.PoseMetadata[ContinuingPoseIdx].GetCostAddend(), ContinuingPoseCostAddend, ContinuingInteractionCostAddend);
+			Result.PoseCost = FPoseSearchCost(CompareFeatureVectors<true>(PoseValues, SearchContext.GetOrBuildQuery(Schema), DynamicWeightsSqrt), SearchIndex.PoseMetadata[ContinuingPoseIdx].GetCostAddend(), ContinuingPoseCostAddend, ContinuingInteractionCostAddend, ContinuingContextInteractionCostAddend);
 		}
 		// data is not 16 bytes padded
 		else
 		{
-			Result.PoseCost = FPoseSearchCost(CompareFeatureVectors<false>(PoseValues, SearchContext.GetOrBuildQuery(Schema), DynamicWeightsSqrt), SearchIndex.PoseMetadata[ContinuingPoseIdx].GetCostAddend(), ContinuingPoseCostAddend, ContinuingInteractionCostAddend);
+			Result.PoseCost = FPoseSearchCost(CompareFeatureVectors<false>(PoseValues, SearchContext.GetOrBuildQuery(Schema), DynamicWeightsSqrt), SearchIndex.PoseMetadata[ContinuingPoseIdx].GetCostAddend(), ContinuingPoseCostAddend, ContinuingInteractionCostAddend, ContinuingContextInteractionCostAddend);
 		}
 
 		Result.SetAssetTime(SearchContext.GetContinuingPoseSearchResult().GetAssetTime());
@@ -1677,15 +1763,14 @@ void UPoseSearchDatabase::SearchPCAKDTree(UE::PoseSearch::FSearchContext& Search
 	}
 #endif // WITH_EDITOR && ENABLE_ANIM_DEBUG && UE_POSE_SEARCH_TRACE_ENABLED
 
+	float ContinuingPoseCostAddend = 0.f, ContinuingInteractionCostAddend = 0.f, ContinuingContextInteractionCostAddend = 0.f;
+	ExtractCostAddends(SearchContext, false, ContinuingPoseCostAddend, ContinuingInteractionCostAddend, ContinuingContextInteractionCostAddend);
 
-	const int32 NumDimensions = Schema->SchemaCardinality;
-	const FSearchIndex& SearchIndex = GetSearchIndex();
-	const float ContinuingPoseCostAddend = 0.f;
-	const float ContinuingInteractionCostAddend = SearchContext.IsContinuingInteraction() ? ContinuingInteractionCostBias : 0.f;
-
-	// since any PoseCost calculated here is at least SearchIndex.MinCostAddend,
+	// since any PoseCost calculated here is at least SearchIndex.MinCostAddend + ContinuingPoseCostAddend + ContinuingInteractionCostAddend + ContinuingContextInteractionCostAddend,
 	// there's no point in performing the search if the worst search result is already better than that
-	if (!GetSkipSearchIfPossible() || SearchResults.ShouldPerformSearch(SearchIndex.MinCostAddend))
+	const FSearchIndex& SearchIndex = GetSearchIndex();
+	const int32 NumDimensions = Schema->SchemaCardinality;
+	if (!GetSkipSearchIfPossible() || SearchResults.ShouldPerformSearch(SearchIndex.MinCostAddend + ContinuingPoseCostAddend + ContinuingInteractionCostAddend + ContinuingContextInteractionCostAddend))
 	{
 		const uint32 ClampedNumberOfPrincipalComponents = SearchIndex.GetNumberOfPrincipalComponents();
 		const uint32 ClampedKDTreeQueryNumNeighbors = FMath::Clamp<uint32>(KDTreeQueryNumNeighbors, 1, SearchIndex.GetNumPoses());
@@ -1713,7 +1798,7 @@ void UPoseSearchDatabase::SearchPCAKDTree(UE::PoseSearch::FSearchContext& Search
 		FNonSelectableIdx NonSelectableIdx;
 		PopulateNonSelectableIdx(NonSelectableIdx, SearchContext
 #if UE_POSE_SEARCH_TRACE_ENABLED
-			, ContinuingPoseCostAddend, ContinuingInteractionCostAddend, DynamicWeightsSqrt
+			, ContinuingPoseCostAddend, ContinuingInteractionCostAddend, ContinuingContextInteractionCostAddend, DynamicWeightsSqrt
 #endif // UE_POSE_SEARCH_TRACE_ENABLED
 		);
 
@@ -1864,7 +1949,7 @@ void UPoseSearchDatabase::SearchPCAKDTree(UE::PoseSearch::FSearchContext& Search
 						const TConstArrayView<int32> PoseIndexes = SearchIndex.PCAValuesVectorToPoseIndexes[Results[ResultIndex].Index];
 						for (int32 Index = 0; Index < PoseIndexes.Num() && NumEvaluatePoseKernelCalls < MaxNumEvaluatePoseKernelCalls; ++Index, ++NumEvaluatePoseKernelCalls)
 						{
-							EvaluatePoseKernel<true, true>(SearchResults, SearchIndex, QueryValues, ReconstructedPoseValuesBuffer, PoseIndexes[Index], INDEX_NONE, SearchFilters, ContinuingPoseCostAddend, ContinuingInteractionCostAddend, this, DynamicWeightsSqrt
+							EvaluatePoseKernel<true, true>(SearchResults, SearchIndex, QueryValues, ReconstructedPoseValuesBuffer, PoseIndexes[Index], INDEX_NONE, SearchFilters, ContinuingPoseCostAddend, ContinuingInteractionCostAddend, ContinuingContextInteractionCostAddend, this, DynamicWeightsSqrt
 #if UE_POSE_SEARCH_TRACE_ENABLED
 								, true
 #endif // UE_POSE_SEARCH_TRACE_ENABLED
@@ -1883,7 +1968,7 @@ void UPoseSearchDatabase::SearchPCAKDTree(UE::PoseSearch::FSearchContext& Search
 						const TConstArrayView<int32> PoseIndexes = SearchIndex.PCAValuesVectorToPoseIndexes[Results[ResultIndex].Index];
 						for (int32 Index = 0; Index < PoseIndexes.Num() && NumEvaluatePoseKernelCalls < MaxNumEvaluatePoseKernelCalls; ++Index, ++NumEvaluatePoseKernelCalls)
 						{
-							EvaluatePoseKernel<true, false>(SearchResults, SearchIndex, QueryValues, ReconstructedPoseValuesBuffer, PoseIndexes[Index], INDEX_NONE, SearchFilters, ContinuingPoseCostAddend, ContinuingInteractionCostAddend, this, DynamicWeightsSqrt
+							EvaluatePoseKernel<true, false>(SearchResults, SearchIndex, QueryValues, ReconstructedPoseValuesBuffer, PoseIndexes[Index], INDEX_NONE, SearchFilters, ContinuingPoseCostAddend, ContinuingInteractionCostAddend, ContinuingContextInteractionCostAddend, this, DynamicWeightsSqrt
 #if UE_POSE_SEARCH_TRACE_ENABLED
 								, true
 #endif // UE_POSE_SEARCH_TRACE_ENABLED
@@ -1906,7 +1991,7 @@ void UPoseSearchDatabase::SearchPCAKDTree(UE::PoseSearch::FSearchContext& Search
 						const TConstArrayView<int32> PoseIndexes = SearchIndex.PCAValuesVectorToPoseIndexes[Results[ResultIndex].Index];
 						for (int32 Index = 0; Index < PoseIndexes.Num() && NumEvaluatePoseKernelCalls < MaxNumEvaluatePoseKernelCalls; ++Index, ++NumEvaluatePoseKernelCalls)
 						{
-							EvaluatePoseKernel<false, true>(SearchResults, SearchIndex, QueryValues, TArrayView<float>(), PoseIndexes[Index], INDEX_NONE, SearchFilters, ContinuingPoseCostAddend, ContinuingInteractionCostAddend, this, DynamicWeightsSqrt
+							EvaluatePoseKernel<false, true>(SearchResults, SearchIndex, QueryValues, TArrayView<float>(), PoseIndexes[Index], INDEX_NONE, SearchFilters, ContinuingPoseCostAddend, ContinuingInteractionCostAddend, ContinuingContextInteractionCostAddend, this, DynamicWeightsSqrt
 #if UE_POSE_SEARCH_TRACE_ENABLED
 								, true
 #endif // UE_POSE_SEARCH_TRACE_ENABLED
@@ -1925,7 +2010,7 @@ void UPoseSearchDatabase::SearchPCAKDTree(UE::PoseSearch::FSearchContext& Search
 						const TConstArrayView<int32> PoseIndexes = SearchIndex.PCAValuesVectorToPoseIndexes[Results[ResultIndex].Index];
 						for (int32 Index = 0; Index < PoseIndexes.Num() && NumEvaluatePoseKernelCalls < MaxNumEvaluatePoseKernelCalls; ++Index, ++NumEvaluatePoseKernelCalls)
 						{
-							EvaluatePoseKernel<false, false>(SearchResults, SearchIndex, QueryValues, TArrayView<float>(), PoseIndexes[Index], INDEX_NONE, SearchFilters, ContinuingPoseCostAddend, ContinuingInteractionCostAddend, this, DynamicWeightsSqrt
+							EvaluatePoseKernel<false, false>(SearchResults, SearchIndex, QueryValues, TArrayView<float>(), PoseIndexes[Index], INDEX_NONE, SearchFilters, ContinuingPoseCostAddend, ContinuingInteractionCostAddend, ContinuingContextInteractionCostAddend, this, DynamicWeightsSqrt
 #if UE_POSE_SEARCH_TRACE_ENABLED
 								, true
 #endif // UE_POSE_SEARCH_TRACE_ENABLED
@@ -1950,7 +2035,7 @@ void UPoseSearchDatabase::SearchPCAKDTree(UE::PoseSearch::FSearchContext& Search
 			{
 				for (int32 ResultIndex = 0; ResultIndex < NumResults; ++ResultIndex)
 				{
-					EvaluatePoseKernel<true, true>(SearchResults, SearchIndex, QueryValues, ReconstructedPoseValuesBuffer, Results[ResultIndex].Index, INDEX_NONE, SearchFilters, ContinuingPoseCostAddend, ContinuingInteractionCostAddend, this, DynamicWeightsSqrt
+					EvaluatePoseKernel<true, true>(SearchResults, SearchIndex, QueryValues, ReconstructedPoseValuesBuffer, Results[ResultIndex].Index, INDEX_NONE, SearchFilters, ContinuingPoseCostAddend, ContinuingInteractionCostAddend, ContinuingContextInteractionCostAddend, this, DynamicWeightsSqrt
 	#if UE_POSE_SEARCH_TRACE_ENABLED
 						, true
 	#endif // UE_POSE_SEARCH_TRACE_ENABLED
@@ -1964,7 +2049,7 @@ void UPoseSearchDatabase::SearchPCAKDTree(UE::PoseSearch::FSearchContext& Search
 			{
 				for (int32 ResultIndex = 0; ResultIndex < NumResults; ++ResultIndex)
 				{
-					EvaluatePoseKernel<true, false>(SearchResults, SearchIndex, QueryValues, ReconstructedPoseValuesBuffer, Results[ResultIndex].Index, INDEX_NONE, SearchFilters, ContinuingPoseCostAddend, ContinuingInteractionCostAddend, this, DynamicWeightsSqrt
+					EvaluatePoseKernel<true, false>(SearchResults, SearchIndex, QueryValues, ReconstructedPoseValuesBuffer, Results[ResultIndex].Index, INDEX_NONE, SearchFilters, ContinuingPoseCostAddend, ContinuingInteractionCostAddend, ContinuingContextInteractionCostAddend, this, DynamicWeightsSqrt
 #if UE_POSE_SEARCH_TRACE_ENABLED
 						, true
 #endif // UE_POSE_SEARCH_TRACE_ENABLED
@@ -1982,7 +2067,7 @@ void UPoseSearchDatabase::SearchPCAKDTree(UE::PoseSearch::FSearchContext& Search
 			{
 				for (int32 ResultIndex = 0; ResultIndex < NumResults; ++ResultIndex)
 				{
-					EvaluatePoseKernel<false, true>(SearchResults, SearchIndex, QueryValues, TArrayView<float>(), Results[ResultIndex].Index, INDEX_NONE, SearchFilters, ContinuingPoseCostAddend, ContinuingInteractionCostAddend, this, DynamicWeightsSqrt
+					EvaluatePoseKernel<false, true>(SearchResults, SearchIndex, QueryValues, TArrayView<float>(), Results[ResultIndex].Index, INDEX_NONE, SearchFilters, ContinuingPoseCostAddend, ContinuingInteractionCostAddend, ContinuingContextInteractionCostAddend, this, DynamicWeightsSqrt
 #if UE_POSE_SEARCH_TRACE_ENABLED
 						, true
 #endif // UE_POSE_SEARCH_TRACE_ENABLED
@@ -1997,7 +2082,7 @@ void UPoseSearchDatabase::SearchPCAKDTree(UE::PoseSearch::FSearchContext& Search
 			{
 				for (int32 ResultIndex = 0; ResultIndex < NumResults; ++ResultIndex)
 				{
-					EvaluatePoseKernel<false, false>(SearchResults, SearchIndex, QueryValues, TArrayView<float>(), Results[ResultIndex].Index, INDEX_NONE, SearchFilters, ContinuingPoseCostAddend, ContinuingInteractionCostAddend, this, DynamicWeightsSqrt
+					EvaluatePoseKernel<false, false>(SearchResults, SearchIndex, QueryValues, TArrayView<float>(), Results[ResultIndex].Index, INDEX_NONE, SearchFilters, ContinuingPoseCostAddend, ContinuingInteractionCostAddend, ContinuingContextInteractionCostAddend, this, DynamicWeightsSqrt
 #if UE_POSE_SEARCH_TRACE_ENABLED
 						, true
 #endif // UE_POSE_SEARCH_TRACE_ENABLED
@@ -2018,7 +2103,7 @@ void UPoseSearchDatabase::SearchPCAKDTree(UE::PoseSearch::FSearchContext& Search
 
 		// calling just for reporting non selectable poses
 		FNonSelectableIdx NonSelectableIdx;
-		PopulateNonSelectableIdx(NonSelectableIdx, SearchContext, ContinuingPoseCostAddend, ContinuingInteractionCostAddend, DynamicWeightsSqrt);
+		PopulateNonSelectableIdx(NonSelectableIdx, SearchContext, ContinuingPoseCostAddend, ContinuingInteractionCostAddend, ContinuingContextInteractionCostAddend, DynamicWeightsSqrt);
 #endif // UE_POSE_SEARCH_TRACE_ENABLED
 	}
 
@@ -2046,13 +2131,13 @@ void UPoseSearchDatabase::SearchVPTree(UE::PoseSearch::FSearchContext& SearchCon
 	}
 #endif // WITH_EDITOR && ENABLE_ANIM_DEBUG && UE_POSE_SEARCH_TRACE_ENABLED
 
-	const FSearchIndex& SearchIndex = GetSearchIndex();
-	const float ContinuingPoseCostAddend = 0.f;
-	const float ContinuingInteractionCostAddend = SearchContext.IsContinuingInteraction() ? ContinuingInteractionCostBias : 0.f;
+	float ContinuingPoseCostAddend = 0.f, ContinuingInteractionCostAddend = 0.f, ContinuingContextInteractionCostAddend = 0.f;
+	ExtractCostAddends(SearchContext, false, ContinuingPoseCostAddend, ContinuingInteractionCostAddend, ContinuingContextInteractionCostAddend);
 
-	// since any PoseCost calculated here is at least SearchIndex.MinCostAddend,
+	// since any PoseCost calculated here is at least SearchIndex.MinCostAddend + ContinuingPoseCostAddend + ContinuingInteractionCostAddend + ContinuingContextInteractionCostAddend,
 	// there's no point in performing the search if the worst search result is already better than that
-	if (!GetSkipSearchIfPossible() || SearchResults.ShouldPerformSearch(SearchIndex.MinCostAddend))
+	const FSearchIndex& SearchIndex = GetSearchIndex();
+	if (!GetSkipSearchIfPossible() || SearchResults.ShouldPerformSearch(SearchIndex.MinCostAddend + ContinuingPoseCostAddend + ContinuingInteractionCostAddend + ContinuingContextInteractionCostAddend))
 	{
 		// @todo: perhaps store the DynamicWeightsSqrt into the GetOrBuildQuery to share between the continuing pose and the full search
 		const int32 NumDimensions = Schema->SchemaCardinality;
@@ -2078,7 +2163,7 @@ void UPoseSearchDatabase::SearchVPTree(UE::PoseSearch::FSearchContext& SearchCon
 		FNonSelectableIdx NonSelectableIdx;
 		PopulateNonSelectableIdx(NonSelectableIdx, SearchContext
 #if UE_POSE_SEARCH_TRACE_ENABLED
-			, ContinuingPoseCostAddend, ContinuingInteractionCostAddend, DynamicWeightsSqrt
+			, ContinuingPoseCostAddend, ContinuingInteractionCostAddend, ContinuingContextInteractionCostAddend, DynamicWeightsSqrt
 #endif // UE_POSE_SEARCH_TRACE_ENABLED
 		);
 		
@@ -2104,7 +2189,7 @@ void UPoseSearchDatabase::SearchVPTree(UE::PoseSearch::FSearchContext& SearchCon
 				const TConstArrayView<int32> PoseIndexes = SearchIndex.ValuesVectorToPoseIndexes[IndexDistance.Index];
 				for (int32 Index = 0; Index < PoseIndexes.Num() && NumEvaluatePoseKernelCalls < MaxNumEvaluatePoseKernelCalls; ++Index, ++NumEvaluatePoseKernelCalls)
 				{
-					EvaluatePoseKernel<false, false>(SearchResults, SearchIndex, QueryValues, TArrayView<float>(), PoseIndexes[Index], INDEX_NONE, SearchFilters, ContinuingPoseCostAddend, ContinuingInteractionCostAddend, this, DynamicWeightsSqrt
+					EvaluatePoseKernel<false, false>(SearchResults, SearchIndex, QueryValues, TArrayView<float>(), PoseIndexes[Index], INDEX_NONE, SearchFilters, ContinuingPoseCostAddend, ContinuingInteractionCostAddend, ContinuingContextInteractionCostAddend, this, DynamicWeightsSqrt
 #if UE_POSE_SEARCH_TRACE_ENABLED
 						, true
 #endif // UE_POSE_SEARCH_TRACE_ENABLED
@@ -2122,7 +2207,7 @@ void UPoseSearchDatabase::SearchVPTree(UE::PoseSearch::FSearchContext& SearchCon
 				const FIndexDistance& IndexDistance = UnsortedResults[ResultIndex];
 
 				// @todo: IndexDistance.Distance is the Sqrt(DissimilarityCost), so there's no need to calculate it again in SearchIndex.ComparePoses
-				EvaluatePoseKernel<false, false>(SearchResults, SearchIndex, QueryValues, TArrayView<float>(), IndexDistance.Index, INDEX_NONE, SearchFilters, ContinuingPoseCostAddend, ContinuingInteractionCostAddend, this, DynamicWeightsSqrt
+				EvaluatePoseKernel<false, false>(SearchResults, SearchIndex, QueryValues, TArrayView<float>(), IndexDistance.Index, INDEX_NONE, SearchFilters, ContinuingPoseCostAddend, ContinuingInteractionCostAddend, ContinuingContextInteractionCostAddend, this, DynamicWeightsSqrt
 #if UE_POSE_SEARCH_TRACE_ENABLED
 					, true
 #endif // UE_POSE_SEARCH_TRACE_ENABLED
@@ -2143,7 +2228,7 @@ void UPoseSearchDatabase::SearchVPTree(UE::PoseSearch::FSearchContext& SearchCon
 
 		// calling just for reporting non selectable poses
 		FNonSelectableIdx NonSelectableIdx;
-		PopulateNonSelectableIdx(NonSelectableIdx, SearchContext, ContinuingPoseCostAddend, ContinuingInteractionCostAddend, DynamicWeightsSqrt);
+		PopulateNonSelectableIdx(NonSelectableIdx, SearchContext, ContinuingPoseCostAddend, ContinuingInteractionCostAddend, ContinuingContextInteractionCostAddend, DynamicWeightsSqrt);
 #endif // UE_POSE_SEARCH_TRACE_ENABLED
 	}
 
@@ -2161,14 +2246,15 @@ void UPoseSearchDatabase::SearchBruteForce(UE::PoseSearch::FSearchContext& Searc
 	
 	using namespace UE::PoseSearch;
 	
-	const FSearchIndex& SearchIndex = GetSearchIndex();
-	const float ContinuingPoseCostAddend = 0.f;
-	const float ContinuingInteractionCostAddend = SearchContext.IsContinuingInteraction() ? ContinuingInteractionCostBias : 0.f;
+	float ContinuingPoseCostAddend = 0.f, ContinuingInteractionCostAddend = 0.f, ContinuingContextInteractionCostAddend = 0.f;
+	ExtractCostAddends(SearchContext, false, ContinuingPoseCostAddend, ContinuingInteractionCostAddend, ContinuingContextInteractionCostAddend);
+
 	const bool bTrackResults = PoseSearchMode == EPoseSearchMode::BruteForce;
 
-	// since any PoseCost calculated here is at least SearchIndex.MinCostAddend,
+	// since any PoseCost calculated here is at least SearchIndex.MinCostAddend + ContinuingPoseCostAddend + ContinuingInteractionCostAddend + ContinuingContextInteractionCostAddend,
 	// there's no point in performing the search if the worst search result is already better than that
-	if (!GetSkipSearchIfPossible() || SearchResults.ShouldPerformSearch(SearchIndex.MinCostAddend))
+	const FSearchIndex& SearchIndex = GetSearchIndex();
+	if (!GetSkipSearchIfPossible() || SearchResults.ShouldPerformSearch(SearchIndex.MinCostAddend + ContinuingPoseCostAddend + ContinuingInteractionCostAddend + ContinuingContextInteractionCostAddend))
 	{
 		// @todo: perhaps store the DynamicWeightsSqrt into the GetOrBuildQuery to share between the continuing pose and the full search
 		const int32 NumDimensions = Schema->SchemaCardinality;
@@ -2192,7 +2278,7 @@ void UPoseSearchDatabase::SearchBruteForce(UE::PoseSearch::FSearchContext& Searc
 		FNonSelectableIdx NonSelectableIdx;
 		PopulateNonSelectableIdx(NonSelectableIdx, SearchContext
 #if UE_POSE_SEARCH_TRACE_ENABLED
-			, ContinuingPoseCostAddend, ContinuingInteractionCostAddend, DynamicWeightsSqrt
+			, ContinuingPoseCostAddend, ContinuingInteractionCostAddend, ContinuingContextInteractionCostAddend, DynamicWeightsSqrt
 #endif // UE_POSE_SEARCH_TRACE_ENABLED
 		);
 
@@ -2208,7 +2294,7 @@ void UPoseSearchDatabase::SearchBruteForce(UE::PoseSearch::FSearchContext& Searc
 				check(IsAligned(ReconstructedPoseValuesBuffer.GetData(), alignof(VectorRegister4Float)));
 				for (int32 PoseIdx = 0; PoseIdx < SearchIndex.GetNumPoses(); ++PoseIdx)
 				{
-					EvaluatePoseKernel<true, false>(SearchResults, SearchIndex, QueryValues, ReconstructedPoseValuesBuffer, PoseIdx, INDEX_NONE, SearchFilters, ContinuingPoseCostAddend, ContinuingInteractionCostAddend, this, DynamicWeightsSqrt
+					EvaluatePoseKernel<true, false>(SearchResults, SearchIndex, QueryValues, ReconstructedPoseValuesBuffer, PoseIdx, INDEX_NONE, SearchFilters, ContinuingPoseCostAddend, ContinuingInteractionCostAddend, ContinuingContextInteractionCostAddend, this, DynamicWeightsSqrt
 #if UE_POSE_SEARCH_TRACE_ENABLED
 						, bTrackResults
 #endif // UE_POSE_SEARCH_TRACE_ENABLED
@@ -2223,7 +2309,7 @@ void UPoseSearchDatabase::SearchBruteForce(UE::PoseSearch::FSearchContext& Searc
 			{
 				for (int32 PoseIdx = 0; PoseIdx < SearchIndex.GetNumPoses(); ++PoseIdx)
 				{
-					EvaluatePoseKernel<false, true>(SearchResults, SearchIndex, QueryValues, TArrayView<float>(), PoseIdx, INDEX_NONE, SearchFilters, ContinuingPoseCostAddend, ContinuingInteractionCostAddend, this, DynamicWeightsSqrt
+					EvaluatePoseKernel<false, true>(SearchResults, SearchIndex, QueryValues, TArrayView<float>(), PoseIdx, INDEX_NONE, SearchFilters, ContinuingPoseCostAddend, ContinuingInteractionCostAddend, ContinuingContextInteractionCostAddend, this, DynamicWeightsSqrt
 #if UE_POSE_SEARCH_TRACE_ENABLED
 						, bTrackResults
 #endif // UE_POSE_SEARCH_TRACE_ENABLED
@@ -2238,7 +2324,7 @@ void UPoseSearchDatabase::SearchBruteForce(UE::PoseSearch::FSearchContext& Searc
 			{
 				for (int32 PoseIdx = 0; PoseIdx < SearchIndex.GetNumPoses(); ++PoseIdx)
 				{
-					EvaluatePoseKernel<false, false>(SearchResults, SearchIndex, QueryValues, TArrayView<float>(), PoseIdx, INDEX_NONE, SearchFilters, ContinuingPoseCostAddend, ContinuingInteractionCostAddend, this, DynamicWeightsSqrt
+					EvaluatePoseKernel<false, false>(SearchResults, SearchIndex, QueryValues, TArrayView<float>(), PoseIdx, INDEX_NONE, SearchFilters, ContinuingPoseCostAddend, ContinuingInteractionCostAddend, ContinuingContextInteractionCostAddend, this, DynamicWeightsSqrt
 #if UE_POSE_SEARCH_TRACE_ENABLED
 						, bTrackResults
 #endif // UE_POSE_SEARCH_TRACE_ENABLED
@@ -2267,7 +2353,7 @@ void UPoseSearchDatabase::SearchBruteForce(UE::PoseSearch::FSearchContext& Searc
 					const int32 LastPoseIdx = FirstPoseIdx + SearchIndexAsset.GetNumPoses();
 					for (int32 PoseIdx = FirstPoseIdx; PoseIdx < LastPoseIdx; ++PoseIdx)
 					{
-						EvaluatePoseKernel<true, false>(SearchResults, SearchIndex, QueryValues, ReconstructedPoseValuesBuffer, PoseIdx, INDEX_NONE, SearchFilters, ContinuingPoseCostAddend, ContinuingInteractionCostAddend, this, DynamicWeightsSqrt
+						EvaluatePoseKernel<true, false>(SearchResults, SearchIndex, QueryValues, ReconstructedPoseValuesBuffer, PoseIdx, INDEX_NONE, SearchFilters, ContinuingPoseCostAddend, ContinuingInteractionCostAddend, ContinuingContextInteractionCostAddend, this, DynamicWeightsSqrt
 #if UE_POSE_SEARCH_TRACE_ENABLED
 							, bTrackResults
 #endif // UE_POSE_SEARCH_TRACE_ENABLED
@@ -2288,7 +2374,7 @@ void UPoseSearchDatabase::SearchBruteForce(UE::PoseSearch::FSearchContext& Searc
 					const int32 LastPoseIdx = FirstPoseIdx + SearchIndexAsset.GetNumPoses();
 					for (int32 PoseIdx = FirstPoseIdx; PoseIdx < LastPoseIdx; ++PoseIdx)
 					{
-						EvaluatePoseKernel<false, true>(SearchResults, SearchIndex, QueryValues, TArrayView<float>(), PoseIdx, INDEX_NONE, SearchFilters, ContinuingPoseCostAddend, ContinuingInteractionCostAddend, this, DynamicWeightsSqrt
+						EvaluatePoseKernel<false, true>(SearchResults, SearchIndex, QueryValues, TArrayView<float>(), PoseIdx, INDEX_NONE, SearchFilters, ContinuingPoseCostAddend, ContinuingInteractionCostAddend, ContinuingContextInteractionCostAddend, this, DynamicWeightsSqrt
 #if UE_POSE_SEARCH_TRACE_ENABLED
 							, bTrackResults
 #endif // UE_POSE_SEARCH_TRACE_ENABLED
@@ -2309,7 +2395,7 @@ void UPoseSearchDatabase::SearchBruteForce(UE::PoseSearch::FSearchContext& Searc
 					const int32 LastPoseIdx = FirstPoseIdx + SearchIndexAsset.GetNumPoses();
 					for (int32 PoseIdx = FirstPoseIdx; PoseIdx < LastPoseIdx; ++PoseIdx)
 					{
-						EvaluatePoseKernel<false, false>(SearchResults, SearchIndex, QueryValues, TArrayView<float>(), PoseIdx, INDEX_NONE, SearchFilters, ContinuingPoseCostAddend, ContinuingInteractionCostAddend, this, DynamicWeightsSqrt
+						EvaluatePoseKernel<false, false>(SearchResults, SearchIndex, QueryValues, TArrayView<float>(), PoseIdx, INDEX_NONE, SearchFilters, ContinuingPoseCostAddend, ContinuingInteractionCostAddend, ContinuingContextInteractionCostAddend, this, DynamicWeightsSqrt
 #if UE_POSE_SEARCH_TRACE_ENABLED
 							, bTrackResults
 #endif // UE_POSE_SEARCH_TRACE_ENABLED
@@ -2332,7 +2418,7 @@ void UPoseSearchDatabase::SearchBruteForce(UE::PoseSearch::FSearchContext& Searc
 
 		// calling just for reporting non selectable poses
 		FNonSelectableIdx NonSelectableIdx;
-		PopulateNonSelectableIdx(NonSelectableIdx, SearchContext, ContinuingPoseCostAddend, ContinuingInteractionCostAddend, DynamicWeightsSqrt);
+		PopulateNonSelectableIdx(NonSelectableIdx, SearchContext, ContinuingPoseCostAddend, ContinuingInteractionCostAddend, ContinuingContextInteractionCostAddend, DynamicWeightsSqrt);
 #endif // UE_POSE_SEARCH_TRACE_ENABLED
 	}
 
@@ -2364,8 +2450,8 @@ void UPoseSearchDatabase::SearchEvent(UE::PoseSearch::FSearchContext& SearchCont
 		const TConstArrayView<int32> PosesWithEvent = SearchIndex.EventData.GetPosesWithEvent(EventToSearch.EventTag);
 		if (!PosesWithEvent.IsEmpty())
 		{
-			const float ContinuingPoseCostAddend = 0.f;
-			const float ContinuingInteractionCostAddend = SearchContext.IsContinuingInteraction() ? ContinuingInteractionCostBias : 0.f;
+			float ContinuingPoseCostAddend = 0.f, ContinuingInteractionCostAddend = 0.f, ContinuingContextInteractionCostAddend = 0.f;
+			ExtractCostAddends(SearchContext, false, ContinuingPoseCostAddend, ContinuingInteractionCostAddend, ContinuingContextInteractionCostAddend);
 
 			TConstArrayView<float> QueryValues = SearchContext.GetOrBuildQuery(Schema);
 
@@ -2390,7 +2476,7 @@ void UPoseSearchDatabase::SearchEvent(UE::PoseSearch::FSearchContext& SearchCont
 			{
 				PopulateNonSelectableIdx(NonSelectableIdx, SearchContext
 #if UE_POSE_SEARCH_TRACE_ENABLED
-					, ContinuingPoseCostAddend, ContinuingInteractionCostAddend, DynamicWeightsSqrt
+					, ContinuingPoseCostAddend, ContinuingInteractionCostAddend, ContinuingContextInteractionCostAddend, DynamicWeightsSqrt
 #endif // UE_POSE_SEARCH_TRACE_ENABLED
 				);
 			}
@@ -2401,7 +2487,7 @@ void UPoseSearchDatabase::SearchEvent(UE::PoseSearch::FSearchContext& SearchCont
 			{
 				for (int32 EventPoseIdx : PosesWithEvent)
 				{
-					EvaluatePoseKernel<false, false>(SearchResults, SearchIndex, QueryValues, TArrayView<float>(), EventPoseIdx, EventPoseIdx, SearchFilters, ContinuingPoseCostAddend, ContinuingInteractionCostAddend, this, DynamicWeightsSqrt
+					EvaluatePoseKernel<false, false>(SearchResults, SearchIndex, QueryValues, TArrayView<float>(), EventPoseIdx, EventPoseIdx, SearchFilters, ContinuingPoseCostAddend, ContinuingInteractionCostAddend, ContinuingContextInteractionCostAddend, this, DynamicWeightsSqrt
 #if UE_POSE_SEARCH_TRACE_ENABLED
 						, true
 #endif // UE_POSE_SEARCH_TRACE_ENABLED
@@ -2421,9 +2507,9 @@ void UPoseSearchDatabase::SearchEvent(UE::PoseSearch::FSearchContext& SearchCont
 					const FSearchIndexAsset& SearchIndexAsset = SearchIndex.GetAssetForPose(EventPoseIdx);
 					const float EventTime = SearchIndexAsset.GetTimeFromPoseIndex(EventPoseIdx, Schema->SampleRate);
 					SearchIndexAsset.IteratePoseIndexesOverTime(EventTime - TimeToEventEnd, EventTime - TimeToEventStart, Schema->SampleRate,
-						[this, &SearchResults, &SearchIndex, &QueryValues, EventPoseIdx, &SearchFilters, ContinuingPoseCostAddend, ContinuingInteractionCostAddend, &SearchContext, &DynamicWeightsSqrt](int32 PoseIdx)
+						[this, &SearchResults, &SearchIndex, &QueryValues, EventPoseIdx, &SearchFilters, ContinuingPoseCostAddend, ContinuingInteractionCostAddend, ContinuingContextInteractionCostAddend, &SearchContext, &DynamicWeightsSqrt](int32 PoseIdx)
 						{
-							EvaluatePoseKernel<false, false>(SearchResults, SearchIndex, QueryValues, TArrayView<float>(), PoseIdx, EventPoseIdx, SearchFilters, ContinuingPoseCostAddend, ContinuingInteractionCostAddend, this, DynamicWeightsSqrt
+							EvaluatePoseKernel<false, false>(SearchResults, SearchIndex, QueryValues, TArrayView<float>(), PoseIdx, EventPoseIdx, SearchFilters, ContinuingPoseCostAddend, ContinuingInteractionCostAddend, ContinuingContextInteractionCostAddend, this, DynamicWeightsSqrt
 #if UE_POSE_SEARCH_TRACE_ENABLED
 								, true
 #endif // UE_POSE_SEARCH_TRACE_ENABLED
diff --git a/Engine/Plugins/Animation/PoseSearch/Source/Runtime/Private/PoseSearchFilter.cpp b/Engine/Plugins/Animation/PoseSearch/Source/Runtime/Private/PoseSearchFilter.cpp
--- a/Engine/Plugins/Animation/PoseSearch/Source/Runtime/Private/PoseSearchFilter.cpp
+++ b/Engine/Plugins/Animation/PoseSearch/Source/Runtime/Private/PoseSearchFilter.cpp
@@ -373,7 +373,7 @@ FSearchFilters::~FSearchFilters()
 
 bool FSearchFilters::AreFiltersValid(const FSearchIndex& SearchIndex, TConstArrayView<float> PoseValues, TConstArrayView<float> QueryValues, TConstArrayView<float> DynamicWeightsSqrt, int32 PoseIdx
 #if UE_POSE_SEARCH_TRACE_ENABLED
-	, float ContinuingPoseCostAddend, float ContinuingInteractionCostAddend, const UPoseSearchDatabase* Database
+	, float ContinuingPoseCostAddend, float ContinuingInteractionCostAddend, float ContinuingContextInteractionCostAddend, const UPoseSearchDatabase* Database
 #endif // UE_POSE_SEARCH_TRACE_ENABLED
 ) const
 {
@@ -388,15 +388,15 @@ bool FSearchFilters::AreFiltersValid(const FSearchIndex& SearchIndex, TConstArra
 			}
 			else if (FilterContext.Filter == &SelectableAssetIdxFilter)
 			{
-				SearchContext.Track(Database, PoseIdx, EPoseCandidateFlags::DiscardedBy_AssetIdxFilter, DynamicWeightsSqrt, ContinuingPoseCostAddend, ContinuingInteractionCostAddend);
+				SearchContext.Track(Database, PoseIdx, EPoseCandidateFlags::DiscardedBy_AssetIdxFilter, DynamicWeightsSqrt, ContinuingPoseCostAddend, ContinuingInteractionCostAddend, ContinuingContextInteractionCostAddend);
 			}
 			else if (FilterContext.Filter == &BlockTransitionFilter)
 			{
-				SearchContext.Track(Database, PoseIdx, EPoseCandidateFlags::DiscardedBy_BlockTransition, DynamicWeightsSqrt, ContinuingPoseCostAddend, ContinuingInteractionCostAddend);
+				SearchContext.Track(Database, PoseIdx, EPoseCandidateFlags::DiscardedBy_BlockTransition, DynamicWeightsSqrt, ContinuingPoseCostAddend, ContinuingInteractionCostAddend, ContinuingContextInteractionCostAddend);
 			}
 			else
 			{
-				SearchContext.Track(Database, PoseIdx, EPoseCandidateFlags::DiscardedBy_PoseFilter, DynamicWeightsSqrt, ContinuingPoseCostAddend, ContinuingInteractionCostAddend);
+				SearchContext.Track(Database, PoseIdx, EPoseCandidateFlags::DiscardedBy_PoseFilter, DynamicWeightsSqrt, ContinuingPoseCostAddend, ContinuingInteractionCostAddend, ContinuingContextInteractionCostAddend);
 			}
 #endif // UE_POSE_SEARCH_TRACE_ENABLED
 			return false;
@@ -409,7 +409,7 @@ bool FSearchFilters::AreFiltersValid(const FSearchIndex& SearchIndex, TConstArra
 		if (!Filter->IsFilterValid(PoseValues, QueryValues, PoseIdx, SearchIndex.PoseMetadata[PoseIdx]))
 		{
 #if UE_POSE_SEARCH_TRACE_ENABLED
-			SearchContext.Track(Database, PoseIdx, EPoseCandidateFlags::DiscardedBy_PoseFilter, DynamicWeightsSqrt, ContinuingPoseCostAddend, ContinuingInteractionCostAddend);
+			SearchContext.Track(Database, PoseIdx, EPoseCandidateFlags::DiscardedBy_PoseFilter, DynamicWeightsSqrt, ContinuingPoseCostAddend, ContinuingInteractionCostAddend, ContinuingContextInteractionCostAddend);
 #endif // UE_POSE_SEARCH_TRACE_ENABLED
 			return false;
 		}
diff --git a/Engine/Plugins/Animation/PoseSearch/Source/Runtime/Private/PoseSearchInteractionIsland.cpp b/Engine/Plugins/Animation/PoseSearch/Source/Runtime/Private/PoseSearchInteractionIsland.cpp
--- a/Engine/Plugins/Animation/PoseSearch/Source/Runtime/Private/PoseSearchInteractionIsland.cpp
+++ b/Engine/Plugins/Animation/PoseSearch/Source/Runtime/Private/PoseSearchInteractionIsland.cpp
@@ -314,43 +314,64 @@ const UE::PoseSearch::IPoseHistory* FInteractionSearchContextBase::GetPoseHistor
 	return nullptr;
 }
 
-void FInteractionSearchContextBase::CacheHashForEquivalence()
+void FInteractionSearchContextBase::CacheHashesForEquivalence()
 {
-	// trying to inefficiently call CacheHash multiple times?
-	check(CachedHashForEquivalence == InvalidCachedHashForEquivalence);
-	CachedHashForEquivalence = CalculateHashForEquivalence();
+	// trying to inefficiently call CalculateHashesForEquivalence multiple times?
+	check(CachedHashForEquivalence == InvalidCachedHashForEquivalence && CachedContextHashForEquivalence == InvalidCachedHashForEquivalence);
+	CalculateHashesForEquivalence(CachedHashForEquivalence, CachedContextHashForEquivalence);
 }
 
-uint32 FInteractionSearchContextBase::CalculateHashForEquivalence() const
+void FInteractionSearchContextBase::CalculateHashesForEquivalence(uint32& OutCachedHashForEquivalence, uint32& OutCachedContextHashForEquivalence) const
 {
-	uint32 Hash = GetTypeHash(Database);
+	check(!AnimContexts.IsEmpty());
 
-	// skipping bDisableCollisions and PoseHistories for hashing
-	for (int32 Index = 0; Index < AnimContexts.Num(); ++Index)
+	OutCachedContextHashForEquivalence = GetTypeHash(AnimContexts[0]);
+
+	for (int32 Index = 1; Index < AnimContexts.Num(); ++Index)
 	{
 		// AnimContexts must be sorted to have deterministic searches across multiple frames
-		check(Index == 0 || AnimContexts[Index - 1] < AnimContexts[Index]);
-		Hash = HashCombineFast(Hash, GetTypeHash(AnimContexts[Index]));
+		check(AnimContexts[Index - 1] < AnimContexts[Index]);
+		OutCachedContextHashForEquivalence = HashCombineFast(OutCachedContextHashForEquivalence, GetTypeHash(AnimContexts[Index]));
 	}
 
+	check(OutCachedContextHashForEquivalence != InvalidCachedHashForEquivalence);
+
+	// skipping bDisableCollisions and PoseHistories for hashing
+	OutCachedHashForEquivalence = HashCombineFast(OutCachedContextHashForEquivalence, GetTypeHash(Database));
+
 	for (int32 Index = 0; Index < Roles.Num(); ++Index)
 	{
-		Hash = HashCombineFast(Hash, GetTypeHash(Roles[Index]));
+		OutCachedHashForEquivalence = HashCombineFast(OutCachedHashForEquivalence, GetTypeHash(Roles[Index]));
 	}
 
 	// only an invalid hash is zero!
-	check(Hash != InvalidCachedHashForEquivalence);
-	return Hash;
+	check(OutCachedHashForEquivalence != InvalidCachedHashForEquivalence);
 }
 
 uint32 FInteractionSearchContextBase::GetHashForEquivalence() const
 {
-	check(CachedHashForEquivalence != InvalidCachedHashForEquivalence);
+#if DO_CHECK
+	check(CachedHashForEquivalence != InvalidCachedHashForEquivalence && CachedContextHashForEquivalence != InvalidCachedHashForEquivalence);
 	// making sure properties didn't change without updating CachedHashForEquivalence
-	check(CachedHashForEquivalence == CalculateHashForEquivalence());
+	uint32 TestCachedHashForEquivalence, TestCachedContextHashForEquivalence;
+	CalculateHashesForEquivalence(TestCachedHashForEquivalence, TestCachedContextHashForEquivalence);
+	check(CachedHashForEquivalence == TestCachedHashForEquivalence && CachedContextHashForEquivalence == TestCachedContextHashForEquivalence);
+#endif //DO_CHECK
 	return CachedHashForEquivalence;
 }
 
+uint32 FInteractionSearchContextBase::GetContextHashForEquivalence() const
+{
+#if DO_CHECK
+	check(CachedHashForEquivalence != InvalidCachedHashForEquivalence && CachedContextHashForEquivalence != InvalidCachedHashForEquivalence);
+	// making sure properties didn't change without updating CachedHashForEquivalence
+	uint32 TestCachedHashForEquivalence, TestCachedContextHashForEquivalence;
+	CalculateHashesForEquivalence(TestCachedHashForEquivalence, TestCachedContextHashForEquivalence);
+	check(CachedHashForEquivalence == TestCachedHashForEquivalence && CachedContextHashForEquivalence == TestCachedContextHashForEquivalence);
+#endif // DO_CHECK
+	return CachedContextHashForEquivalence;
+}
+
 void FInteractionSearchContextBase::AddReferencedObjects(FReferenceCollector& Collector)
 {
 	Collector.AddReferencedObject(Database);
@@ -364,7 +385,7 @@ void FInteractionSearchContextBase::AddReferencedObjects(FReferenceCollector& Co
 void FInteractionSearchContextBase::Add(const UObject* AnimContext, const IPoseHistory* PoseHistory, const FRole Role, bool bWantsDisableCollisions)
 {
 	// trying to inefficiently call Add again after having calculated the hash via CalculateHashForEquivalence?
-	check(CachedHashForEquivalence == InvalidCachedHashForEquivalence);
+	check(CachedHashForEquivalence == InvalidCachedHashForEquivalence && CachedContextHashForEquivalence == InvalidCachedHashForEquivalence);
 
 	check(AnimContext && PoseHistory);
 
@@ -598,6 +619,8 @@ void FValidInteractionSearch::UndoDisableCollisions()
 			}
 		}
 	}
+
+	DisabledCollisions.Reset();
 }
 
 #if ENABLE_VISUAL_LOG
@@ -666,6 +689,12 @@ void FValidInteractionSearches::Update(TConstArrayView<FInteractionIsland*> Isla
 	UE_MT_SCOPED_WRITE_ACCESS(ThreadSafeCounter);
 #endif //ENABLE_ANIM_DEBUG
 
+	// avoiding ContextsToSearches.Reset() to preserve the array allocations in ContextsToSearches
+	for (TPair<uint32, TArray<uint32>>& ContextToSearches : ContextsToSearches)
+	{
+		ContextToSearches.Value.Reset();
+	}
+
 	FMemMark Mark(FMemStack::Get());
 	TMap<uint32, bool, TInlineSetAllocator<32, TMemStackSetAllocator<>>> ContinuingOrNewValidSearchContexts;
 	for (const FInteractionIsland* Island : Islands)
@@ -677,6 +706,9 @@ void FValidInteractionSearches::Update(TConstArrayView<FInteractionIsland*> Isla
 			{
 				const FInteractionSearchContext& SearchContext = Island->GetSearchContexts()[SearchResult.SearchIndex];
 				const uint32 ValidSearchContextHash = SearchContext.GetHashForEquivalence();
+
+				ContextsToSearches.FindOrAdd(SearchContext.GetContextHashForEquivalence()).Add(ValidSearchContextHash);
+
 				if (FValidInteractionSearch* ContinuingValidInteractionSearch = Searches.Find(ValidSearchContextHash))
 				{
 					ContinuingOrNewValidSearchContexts.Add(ValidSearchContextHash) = true;
@@ -747,8 +779,22 @@ const FValidInteractionSearch* FValidInteractionSearches::Get(const FInteraction
 	UE_MT_SCOPED_READ_ACCESS(ThreadSafeCounter);
 #endif // ENABLE_ANIM_DEBUG
 
-	const uint32 SearchContextHash = SearchContext.GetHashForEquivalence();
-	return Searches.Find(SearchContextHash);
+	const uint32 HashForEquivalence = SearchContext.GetHashForEquivalence();
+	return Searches.Find(HashForEquivalence);
+}
+
+bool FValidInteractionSearches::IsThereAnyValidSearchWithTheSameContexts(const FInteractionSearchContext& SearchContext) const
+{
+#if ENABLE_ANIM_DEBUG
+	UE_MT_SCOPED_READ_ACCESS(ThreadSafeCounter);
+#endif // ENABLE_ANIM_DEBUG
+
+	const uint32 ContextHashForEquivalence = SearchContext.GetContextHashForEquivalence();
+	if (const TArray<uint32>* HashesForEquivalence = ContextsToSearches.Find(ContextHashForEquivalence))
+	{
+		return !HashesForEquivalence->IsEmpty();
+	}
+	return false;
 }
 
 void FValidInteractionSearches::AddReferencedObjects(FReferenceCollector& Collector)
@@ -791,6 +837,7 @@ FPoseSearchContinuingProperties FInteractionSearchContext::GetContinuingProperti
 	ContinuingProperties.PlayingAssetBlendParameters = PlayingAssetBlendParameters;
 	ContinuingProperties.InterruptMode = InterruptMode;
 	ContinuingProperties.bIsContinuingInteraction = bIsContinuingInteraction;
+	ContinuingProperties.bIsContinuingContextInteraction = bIsContinuingContextInteraction;
 	return ContinuingProperties;
 }
 
@@ -1138,7 +1185,7 @@ void FInteractionIsland::AddTickDependencies(UActorComponent* TickActorComponent
 		}
 
 		// adding prerequisite between actor and skeletal mesh component so CBP runs before ABP parallel execution jobs
-		AddPrerequisite(TickActor->PrimaryActorTick, InteractionSubsystem, TickActorComponent->PrimaryComponentTick);
+		AddPrerequisite(TickActorComponent->PrimaryComponentTick, InteractionSubsystem, TickActor->PrimaryActorTick);
 
 		// @todo: find a better way to do this!
 		// adding prerequisite to the ClosingTickFunction to components having a tick dependency on TickActorComponent, so they can run AFTER the ClosingTickFunction
@@ -1175,6 +1222,70 @@ void FInteractionIsland::AddTickDependencies(UActorComponent* TickActorComponent
 #endif // ENABLE_ANIM_DEBUG
 }
 
+#if ENABLE_ANIM_DEBUG
+void FInteractionIsland::ValidateTickDependenciesLeaks() const
+{
+	bool bAnyError = false;
+	for (const TObjectPtr<UActorComponent>& TickActorComponentPtr : TickActorComponents)
+	{
+		if (const UActorComponent* TickActorComponent = TickActorComponentPtr.Get())
+		{
+			const AActor* TickActor = TickActorComponent->GetOwner();
+			check(TickActor && TickActor->CanEverTick());
+
+			// checking the TickActor
+			for (const FTickPrerequisite& TickPrerequisite : TickActor->PrimaryActorTick.GetPrerequisites())
+			{
+				if (TickPrerequisite.PrerequisiteTickFunction == &PreTickFunction)
+				{
+					UE_LOG(LogPoseSearch, Error, TEXT("============== FInteractionIsland::ValidateTickDependenciesLeaks failed! PreTickFunction leaked in %s =============="), *TickActor->GetName());
+					bAnyError = true;
+				}
+
+				if (TickPrerequisite.PrerequisiteTickFunction == &PostTickFunction)
+				{
+					UE_LOG(LogPoseSearch, Error, TEXT("============== FInteractionIsland::ValidateTickDependenciesLeaks failed! PostTickFunction leaked in %s =============="), *TickActor->GetName());
+					bAnyError = true;
+				}
+
+				if (TickPrerequisite.PrerequisiteTickFunction == &ClosingTickFunction)
+				{
+					UE_LOG(LogPoseSearch, Error, TEXT("============== FInteractionIsland::ValidateTickDependenciesLeaks failed! ClosingTickFunction leaked in %s =============="), *TickActor->GetName());
+					bAnyError = true;
+				}
+			}
+
+			// checking the TickActorComponent
+			for (const FTickPrerequisite& TickPrerequisite : TickActorComponent->PrimaryComponentTick.GetPrerequisites())
+			{
+				if (TickPrerequisite.PrerequisiteTickFunction == &PreTickFunction)
+				{
+					UE_LOG(LogPoseSearch, Error, TEXT("============== FInteractionIsland::ValidateTickDependenciesLeaks failed! PreTickFunction leaked in %s =============="), *TickActorComponent->GetName());
+					bAnyError = true;
+				}
+
+				if (TickPrerequisite.PrerequisiteTickFunction == &PostTickFunction)
+				{
+					UE_LOG(LogPoseSearch, Error, TEXT("============== FInteractionIsland::ValidateTickDependenciesLeaks failed! PostTickFunction leaked in %s =============="), *TickActorComponent->GetName());
+					bAnyError = true;
+				}
+
+				if (TickPrerequisite.PrerequisiteTickFunction == &ClosingTickFunction)
+				{
+					UE_LOG(LogPoseSearch, Error, TEXT("============== FInteractionIsland::ValidateTickDependenciesLeaks failed! ClosingTickFunction leaked in %s =============="), *TickActorComponent->GetName());
+					bAnyError = true;
+				}
+			}
+		}
+	}
+
+	if (bAnyError)
+	{
+		LogTickDependencies();
+	}
+}
+#endif // ENABLE_ANIM_DEBUG
+
 void FInteractionIsland::RemoveTickDependencies(bool bValidateTickDependencies)
 {
 	// Called by UPoseSearchInteractionSubsystem::Tick when there aren't animation jobs flying.
@@ -1273,7 +1384,7 @@ void FInteractionIsland::RemoveTickDependencies(bool bValidateTickDependencies)
 					{
 						// [undo] adding to TickActorComponent and owner actor TickActor dependencies to the PreTickFunction
 						TickActorComponent->PrimaryComponentTick.RemovePrerequisite(InteractionSubsystem, PreTickFunction);
-						PostTickFunction.RemovePrerequisite(TickActorComponent, TickActorComponent->PrimaryComponentTick);
+						TickActor->PrimaryActorTick.RemovePrerequisite(InteractionSubsystem, PreTickFunction);
 
 						// [undo] add PostTickFunction function dependencies to TickActorComponent and owner actor
 						PostTickFunction.RemovePrerequisite(TickActorComponent, TickActorComponent->PrimaryComponentTick);
@@ -1291,7 +1402,7 @@ void FInteractionIsland::RemoveTickDependencies(bool bValidateTickDependencies)
 					}
 
 					// [undo] adding prerequisite between actor and skeletal mesh component so CBP runs before ABP parallel execution jobs
-					TickActor->PrimaryActorTick.RemovePrerequisite(InteractionSubsystem, TickActorComponent->PrimaryComponentTick);
+					TickActorComponent->PrimaryComponentTick.RemovePrerequisite(InteractionSubsystem, TickActor->PrimaryActorTick);
 				}
 			}
 			bMainActor = false;
@@ -1306,6 +1417,13 @@ void FInteractionIsland::RemoveTickDependencies(bool bValidateTickDependencies)
 			}
 		}
 
+#if ENABLE_ANIM_DEBUG
+		if (GVarPoseSearchInteractionDiagnoseTickDependencies)
+		{
+			ValidateTickDependenciesLeaks();
+		}
+#endif // ENABLE_ANIM_DEBUG
+
 		if (!PreTickFunction.GetPrerequisites().IsEmpty())
 		{
 			UE_LOG(LogPoseSearch, Warning, TEXT("FInteractionIsland::RemoveTickDependencies, unable to cleanly remove all the PreTickFunction prerequisites! (This could have happened if an actor got garbage collected or a BP got recompiled while PIE was paused)"));
diff --git a/Engine/Plugins/Animation/PoseSearch/Source/Runtime/Private/PoseSearchInteractionSubsystem.cpp b/Engine/Plugins/Animation/PoseSearch/Source/Runtime/Private/PoseSearchInteractionSubsystem.cpp
--- a/Engine/Plugins/Animation/PoseSearch/Source/Runtime/Private/PoseSearchInteractionSubsystem.cpp
+++ b/Engine/Plugins/Animation/PoseSearch/Source/Runtime/Private/PoseSearchInteractionSubsystem.cpp
@@ -1185,12 +1185,12 @@ void UPoseSearchInteractionSubsystem::GenerateSearchContexts(float DeltaSeconds,
 							}
 
 							// looking for a preexisting valid interaction resembling SearchContext
-							SearchContext.CacheHashForEquivalence();
+							SearchContext.CacheHashesForEquivalence();
 
 							// looking for a previous valid interaction with the SearchContext properties in ValidInteractionSearches 
-							// (it's using SearchContext.CacheHashForEquivalence to speed up the search)
+							// (it's using SearchContext.CacheHashesForEquivalence to speed up the search)
 							const FValidInteractionSearch* ValidInteractionSearch = ValidInteractionSearches.Get(SearchContext);
-							bool bInIsContinuingInteraction = ValidInteractionSearch != nullptr;
+							const bool bInIsContinuingInteraction = ValidInteractionSearch != nullptr;
 
 							SearchContext.SetContinuingInteraction(bInIsContinuingInteraction);
 
@@ -1228,7 +1228,7 @@ void UPoseSearchInteractionSubsystem::GenerateSearchContexts(float DeltaSeconds,
 							}
 
 							// Populating Continuing properties by using the ValidInteractionSearch->SearchResult
-							if (ValidInteractionSearch)
+							if (bInIsContinuingInteraction)
 							{
 								const FInteractionSearchResult& PreviousSearchResult = ValidInteractionSearch->GetSearchResult();
 								if (const UE::PoseSearch::FSearchIndexAsset* SearchIndexAsset = PreviousSearchResult.GetSearchIndexAsset())
@@ -1248,6 +1248,14 @@ void UPoseSearchInteractionSubsystem::GenerateSearchContexts(float DeltaSeconds,
 											SearchIndexAsset->IsMirrored(), SearchIndexAsset->GetBlendParameters(), InterruptMode);
 									}
 								}
+
+								check(ValidInteractionSearches.IsThereAnyValidSearchWithTheSameContexts(SearchContext));
+								SearchContext.SetContinuingContextInteraction(true);
+							}
+							else
+							{
+								const bool bIsThereAnyValidSearchWithTheSameContexts = ValidInteractionSearches.IsThereAnyValidSearchWithTheSameContexts(SearchContext);
+								SearchContext.SetContinuingContextInteraction(bIsThereAnyValidSearchWithTheSameContexts);
 							}
 
 							SearchContexts.Emplace(SearchContext);
diff --git a/Engine/Plugins/Animation/PoseSearch/Source/Runtime/Private/PoseSearchLibrary.cpp b/Engine/Plugins/Animation/PoseSearch/Source/Runtime/Private/PoseSearchLibrary.cpp
--- a/Engine/Plugins/Animation/PoseSearch/Source/Runtime/Private/PoseSearchLibrary.cpp
+++ b/Engine/Plugins/Animation/PoseSearch/Source/Runtime/Private/PoseSearchLibrary.cpp
@@ -1315,7 +1315,8 @@ void UPoseSearchLibrary::MotionMatch(
 	{
 		CurrentAssetToConsider = SearchContext.GetAssetsToConsiderSet();
 	}
-	SearchContext.SetIsContinuingInteraction(ContinuingProperties.bIsContinuingInteraction);
+	SearchContext.SetContinuingInteraction(ContinuingProperties.bIsContinuingInteraction);
+	SearchContext.SetContinuingContextInteraction(ContinuingProperties.bIsContinuingContextInteraction);
 
 	// collecting all the databases searches in AssetsToSearchPerDatabase
 	// and all the continuing pose searches in ContinuingPoseAssetsToSearchPerDatabase
diff --git a/Engine/Plugins/Animation/PoseSearch/Source/Runtime/Private/PoseSearchResult.cpp b/Engine/Plugins/Animation/PoseSearch/Source/Runtime/Private/PoseSearchResult.cpp
--- a/Engine/Plugins/Animation/PoseSearch/Source/Runtime/Private/PoseSearchResult.cpp
+++ b/Engine/Plugins/Animation/PoseSearch/Source/Runtime/Private/PoseSearchResult.cpp
@@ -208,7 +208,7 @@ float FSearchResult::CalculateTimeToEvent() const
 
 void FSearchResult::InitFrom(const FPoseSearchBlueprintResult& BlueprintResult)
 {
-	PoseCost = FPoseSearchCost(BlueprintResult.SearchCost, 0.f, 0.f, 0.f);
+	PoseCost = FPoseSearchCost(BlueprintResult.SearchCost);
 	PoseIdx = BlueprintResult.SelectedDatabase ? BlueprintResult.SelectedDatabase->GetPoseIndex(BlueprintResult.SelectedAnim.Get(), BlueprintResult.SelectedTime, BlueprintResult.bIsMirrored, BlueprintResult.BlendParameters) : INDEX_NONE;
 	EventPoseIdx = INDEX_NONE;
 	Database = BlueprintResult.SelectedDatabase;
diff --git a/Engine/Plugins/Animation/PoseSearch/Source/Runtime/Public/PoseSearch/PoseSearchAnimNotifies.h b/Engine/Plugins/Animation/PoseSearch/Source/Runtime/Public/PoseSearch/PoseSearchAnimNotifies.h
--- a/Engine/Plugins/Animation/PoseSearch/Source/Runtime/Public/PoseSearch/PoseSearchAnimNotifies.h
+++ b/Engine/Plugins/Animation/PoseSearch/Source/Runtime/Public/PoseSearch/PoseSearchAnimNotifies.h
@@ -66,6 +66,35 @@ public:
 	float CostAddend = -1.0f;
 };
 
+// Pose search cost for the continuing interaction (same anim contexts with the same assigned roles), making the animation segment more or less 
+// likely to be continuing playing based on the notify parameters
+UCLASS(MinimalAPI, Experimental, Blueprintable, meta = (DisplayName = "Pose Search: Override Continuing Interaction Cost Bias"))
+class UAnimNotifyState_PoseSearchOverrideContinuingInteractionPoseCostBias : public UAnimNotifyState_PoseSearchBase
+{
+	GENERATED_BODY()
+
+public:
+
+	// A negative value reduces the cost and makes the segment more likely to continuing playing. A positive value, conversely,
+	// makes the segment less likely to continuing playing
+	UPROPERTY(EditAnywhere, Category = Config, meta = (DisplayName = "Modifier"))
+	float CostAddend = -1.0f;
+};
+
+// Pose search cost for the continuing context interaction (same anim contexts, assigned roles may differ), making the animation segment more or less 
+// likely to be continuing playing based on the notify parameters
+UCLASS(MinimalAPI, Experimental, Blueprintable, meta = (DisplayName = "Pose Search: Override Continuing Context Interaction Cost Bias"))
+class UAnimNotifyState_PoseSearchOverrideContinuingContextInteractionCostBias : public UAnimNotifyState_PoseSearchBase
+{
+	GENERATED_BODY()
+
+public:
+
+	// A negative value reduces the cost and makes selecting a new interaction less likely
+	UPROPERTY(EditAnywhere, Category = Config, meta = (DisplayName = "Modifier"))
+	float CostAddend = -1.0f;
+};
+
 // UPoseSearchFeatureChannel(s) can use this UAnimNotifyState_PoseSearchSamplingEvent to demarcate events identified by SamplingAttributeId
 // during database indexing by specifying their SamplingAttributeId property to match UAnimNotifyState_PoseSearchSamplingAttribute::SamplingAttributeId
 UCLASS(MinimalAPI, Blueprintable, meta = (DisplayName = "Pose Search: Sampling Event"))
diff --git a/Engine/Plugins/Animation/PoseSearch/Source/Runtime/Public/PoseSearch/PoseSearchContext.h b/Engine/Plugins/Animation/PoseSearch/Source/Runtime/Public/PoseSearch/PoseSearchContext.h
--- a/Engine/Plugins/Animation/PoseSearch/Source/Runtime/Public/PoseSearch/PoseSearchContext.h
+++ b/Engine/Plugins/Animation/PoseSearch/Source/Runtime/Public/PoseSearch/PoseSearchContext.h
@@ -330,7 +330,12 @@ struct FSearchContext
 	// Experimental, this feature might be removed without warning, not for production use
 	bool IsContinuingInteraction() const { return bIsContinuingInteraction; }
 	// Experimental, this feature might be removed without warning, not for production use
-	void SetIsContinuingInteraction(bool bInIsContinuingInteraction) { bIsContinuingInteraction = bInIsContinuingInteraction; }
+	void SetContinuingInteraction(bool bInIsContinuingInteraction) { bIsContinuingInteraction = bInIsContinuingInteraction; }
+
+	// Experimental, this feature might be removed without warning, not for production use
+	bool IsContinuingContextInteraction() const { return bIsContinuingContextInteraction; }
+	// Experimental, this feature might be removed without warning, not for production use
+	void SetContinuingContextInteraction(bool bInIsContinuingContextInteraction) { bIsContinuingContextInteraction = bInIsContinuingContextInteraction; }
 
 private:
 	float GetSampleCurveValueInternal(float SampleTime, const FName& CurveName, const FRole& SampleRole);
@@ -384,10 +389,12 @@ private:
 	// mapping channel unique identifier (hash) to FCachedChannel
 	TMap<uint32, FCachedChannel, TInlineSetAllocator<PreallocatedCachedChannelDataNum, TMemStackSetAllocator<>>> CachedChannels;
 
-	// @todo: add FPoseSearchContinuingProperties here and reconstuct the continuing pose FSearchResult for the previous frame from it
 	// Experimental, this feature might be removed without warning, not for production use
 	bool bIsContinuingInteraction = false;
 
+	// Experimental, this feature might be removed without warning, not for production use
+	bool bIsContinuingContextInteraction = false;
+
 #if WITH_EDITOR
 	bool bAsyncBuildIndexInProgress = false;
 #endif // WITH_EDITOR
@@ -447,7 +454,7 @@ public:
 	};
 	
 	UE_API void Track(const UPoseSearchDatabase* Database, int32 PoseIdx = INDEX_NONE, EPoseCandidateFlags PoseCandidateFlags = EPoseCandidateFlags::None, const FPoseSearchCost& Cost = FPoseSearchCost());
-	UE_API void Track(const UPoseSearchDatabase* Database, int32 PoseIdx, EPoseCandidateFlags PoseCandidateFlags, TConstArrayView<float> DynamicWeightsSqrt, float ContinuingPoseCostAddend, float ContinuingInteractionCostAddend);
+	UE_API void Track(const UPoseSearchDatabase* Database, int32 PoseIdx, EPoseCandidateFlags PoseCandidateFlags, TConstArrayView<float> DynamicWeightsSqrt, float ContinuingPoseCostAddend, float ContinuingInteractionCostAddend, float ContinuingContextInteractionCostAddend);
 
 	const TMap<const UPoseSearchDatabase*, FBestPoseCandidates>& GetBestPoseCandidatesMap() const
 	{
diff --git a/Engine/Plugins/Animation/PoseSearch/Source/Runtime/Public/PoseSearch/PoseSearchCost.h b/Engine/Plugins/Animation/PoseSearch/Source/Runtime/Public/PoseSearch/PoseSearchCost.h
--- a/Engine/Plugins/Animation/PoseSearch/Source/Runtime/Public/PoseSearch/PoseSearchCost.h
+++ b/Engine/Plugins/Animation/PoseSearch/Source/Runtime/Public/PoseSearch/PoseSearchCost.h
@@ -17,12 +17,13 @@ public:
 	FPoseSearchCost& operator=(const FPoseSearchCost&) = default;
 	FPoseSearchCost& operator=(FPoseSearchCost&&) = default;
 
-	FPoseSearchCost(float DissimilarityCost, float InNotifyCostAddend, float InContinuingPoseCostAddend, float InContinuingInteractionCostAddend = 0.f)
-		: TotalCost(DissimilarityCost + InNotifyCostAddend + InContinuingPoseCostAddend + InContinuingInteractionCostAddend)
+	FPoseSearchCost(float DissimilarityCost, float InNotifyCostAddend = 0.f, float InContinuingPoseCostAddend = 0.f, float InContinuingInteractionCostAddend = 0.f, float InContinuingContextInteractionCostAddend = 0.f)
+		: TotalCost(DissimilarityCost + InNotifyCostAddend + InContinuingPoseCostAddend + InContinuingInteractionCostAddend + InContinuingContextInteractionCostAddend)
 #if WITH_EDITORONLY_DATA
 		, NotifyCostAddend(InNotifyCostAddend)
 		, ContinuingPoseCostAddend(InContinuingPoseCostAddend)
 		, ContinuingInteractionCostAddend(InContinuingInteractionCostAddend)
+		, ContinuingContextInteractionCostAddend(InContinuingContextInteractionCostAddend)
 #endif // WITH_EDITORONLY_DATA
 	{
 	}
@@ -33,11 +34,12 @@ public:
 	operator float() const { return TotalCost; }
 
 #if WITH_EDITORONLY_DATA
-	float GetCostAddend() const { return GetNotifyCostAddend() + GetContinuingPoseCostAddend() + GetContinuingInteractionCostAddend(); }
+	float GetCostAddend() const { return NotifyCostAddend + ContinuingPoseCostAddend + ContinuingInteractionCostAddend + ContinuingContextInteractionCostAddend; }
 
 	float GetNotifyCostAddend() const { return NotifyCostAddend; }
 	float GetContinuingPoseCostAddend() const { return ContinuingPoseCostAddend; }
 	float GetContinuingInteractionCostAddend() const { return ContinuingInteractionCostAddend; }
+	float GetContinuingContextInteractionCostAddend() const { return ContinuingContextInteractionCostAddend; }
 #endif // WITH_EDITORONLY_DATA
 
 	friend FArchive& operator<<(FArchive& Ar, FPoseSearchCost& PoseSearchCost)
@@ -48,13 +50,14 @@ public:
 		Ar << PoseSearchCost.NotifyCostAddend;
 		Ar << PoseSearchCost.ContinuingPoseCostAddend;
 		Ar << PoseSearchCost.ContinuingInteractionCostAddend;
+		Ar << PoseSearchCost.ContinuingContextInteractionCostAddend;
 #endif // WITH_EDITORONLY_DATA
 			
 		return Ar;
 	}
 
 protected:
-	// TotalCost is the sum of all the Cost contributions (dissimilarity, notifies, continuing pose, continuing interaction costs)
+	// TotalCost is the sum of all the Cost contributions (dissimilarity, notifies, continuing pose, continuing interaction, and continuing context interaction costs)
 	UPROPERTY()
 	float TotalCost = MAX_flt;
 
@@ -73,6 +76,11 @@ protected:
 	UPROPERTY()
 	float ContinuingInteractionCostAddend = 0.f;
 
+	// Experimental, this feature might be removed without warning, not for production use
+	// Continuing Context Interaction Cost Bias contribution
+	UPROPERTY()
+	float ContinuingContextInteractionCostAddend = 0.f;
+
 #endif // WITH_EDITORONLY_DATA
 };
 
diff --git a/Engine/Plugins/Animation/PoseSearch/Source/Runtime/Public/PoseSearch/PoseSearchDatabase.h b/Engine/Plugins/Animation/PoseSearch/Source/Runtime/Public/PoseSearch/PoseSearchDatabase.h
--- a/Engine/Plugins/Animation/PoseSearch/Source/Runtime/Public/PoseSearch/PoseSearchDatabase.h
+++ b/Engine/Plugins/Animation/PoseSearch/Source/Runtime/Public/PoseSearch/PoseSearchDatabase.h
@@ -523,14 +523,20 @@ public:
 	UPROPERTY(EditAnywhere, Category = "Database")
 	float LoopingCostBias = -0.005f;
 
-	// Cost added to poses from this database (if setup with UMultiAnimAsset assets in "interaction mode") as continuation of a previous interaction.
+	// Cost added to poses from this database (if setup with UMultiAnimAsset assets in "interaction mode") as continuation of a previous interaction (same anim contexts with the same assigned roles)
 	// This allows users to apply a cost bias (positive or negative) to poses following up an interaction.
 	// This is useful to help the system stay in one animation segment longer, or shorter depending on how you set this bias.
 	// Negative values make it more likely to be picked, or stayed in, positive values make it less likely to be picked or stay in.
 	// Note: excluded from DDC hash, since used only at runtime
+	// Experimental, this feature might be removed without warning, not for production use
 	UPROPERTY(EditAnywhere, Category = "Database", meta = (ExcludeFromHash))
 	float ContinuingInteractionCostBias = 0.f;
 
+	// behaving like ContinuingInteractionCostBias, but for continuing interactions with the same anim contexts (e.g. same actors)
+	// Experimental, this feature might be removed without warning, not for production use
+	UPROPERTY(EditAnywhere, Category = "Database", meta = (ExcludeFromHash))
+	float ContinuingContextInteractionCostBias = 0.f;
+	
 #if WITH_EDITORONLY_DATA
 	// These settings allow users to trim the start and end of animations in the database to preserve start/end frames for blending, and prevent the system from selecting the very last frames before it blends out.
 	// valid animation frames will be AnimationAssetTimeStart + ExcludeFromDatabaseParameters.Min, AnimationAssetTimeEnd + ExcludeFromDatabaseParameters.Max
@@ -728,11 +734,12 @@ private:
 	typedef TArray<int32, TInlineAllocator<256, TMemStackAllocator<>>> FNonSelectableIdx;
 	void PopulateNonSelectableIdx(FNonSelectableIdx& NonSelectableIdx, UE::PoseSearch::FSearchContext& SearchContext
 #if UE_POSE_SEARCH_TRACE_ENABLED
-		, float ContinuingPoseCostAddend, float ContinuingInteractionCostAddend, TConstArrayView<float> DynamicWeightsSqrt
+		, float ContinuingPoseCostAddend, float ContinuingInteractionCostAddend, float ContinuingContextInteractionCostAddend, TConstArrayView<float> DynamicWeightsSqrt
 #endif //UE_POSE_SEARCH_TRACE_ENABLED
 	) const;
 
 	void UpdateCachedProperties();
+	void ExtractCostAddends(const UE::PoseSearch::FSearchContext& SearchContext, bool bIsContinuingPoseSearch, float& ContinuingPoseCostAddend, float& ContinuingInteractionCostAddend, float& ContinuingContextInteractionCostAddend) const;
 };
 
 template<typename TDatabaseAnimationAsset>
diff --git a/Engine/Plugins/Animation/PoseSearch/Source/Runtime/Public/PoseSearch/PoseSearchFilter.h b/Engine/Plugins/Animation/PoseSearch/Source/Runtime/Public/PoseSearch/PoseSearchFilter.h
--- a/Engine/Plugins/Animation/PoseSearch/Source/Runtime/Public/PoseSearch/PoseSearchFilter.h
+++ b/Engine/Plugins/Animation/PoseSearch/Source/Runtime/Public/PoseSearch/PoseSearchFilter.h
@@ -206,7 +206,7 @@ struct FSearchFilters
 	// @todo: template this with bAlignedAndPadded to be able to use faster ComparePoses
 	bool AreFiltersValid(const FSearchIndex& SearchIndex, TConstArrayView<float> PoseValues, TConstArrayView<float> QueryValues, TConstArrayView<float> DynamicWeightsSqrt, int32 PoseIdx
 #if UE_POSE_SEARCH_TRACE_ENABLED
-		, float ContinuingPoseCostAddend, float ContinuingInteractionCostAddend, const UPoseSearchDatabase* Database
+		, float ContinuingPoseCostAddend, float ContinuingInteractionCostAddend, float ContinuingContextInteractionCostAddend, const UPoseSearchDatabase* Database
 #endif // UE_POSE_SEARCH_TRACE_ENABLED
 		) const;
 
diff --git a/Engine/Plugins/Animation/PoseSearch/Source/Runtime/Public/PoseSearch/PoseSearchInteractionIsland.h b/Engine/Plugins/Animation/PoseSearch/Source/Runtime/Public/PoseSearch/PoseSearchInteractionIsland.h
--- a/Engine/Plugins/Animation/PoseSearch/Source/Runtime/Public/PoseSearch/PoseSearchInteractionIsland.h
+++ b/Engine/Plugins/Animation/PoseSearch/Source/Runtime/Public/PoseSearch/PoseSearchInteractionIsland.h
@@ -39,7 +39,8 @@ struct FInteractionSearchContextBase
 	bool IsDisableCollisions() const { return bDisableCollisions; }
 
 	uint32 GetHashForEquivalence() const;
-	void CacheHashForEquivalence();
+	uint32 GetContextHashForEquivalence() const;
+	void CacheHashesForEquivalence();
 	
 	void SetDatabase(const UPoseSearchDatabase* InDatabase) { Database = InDatabase; }
 	void AddReferencedObjects(FReferenceCollector& Collector);
@@ -50,7 +51,7 @@ struct FInteractionSearchContextBase
 #endif // DO_CHECK
 
 private:
-	uint32 CalculateHashForEquivalence() const;
+	void CalculateHashesForEquivalence(uint32& OutCachedHashForEquivalence, uint32& OutCachedContextHashForEquivalence) const;
 
 	TArray<TObjectPtr<const UObject>, TInlineAllocator<PreallocatedRolesNum>> AnimContexts;
 
@@ -64,6 +65,7 @@ private:
 	bool bDisableCollisions = false;
 	enum { InvalidCachedHashForEquivalence = 0 };
 	uint32 CachedHashForEquivalence = InvalidCachedHashForEquivalence;
+	uint32 CachedContextHashForEquivalence = InvalidCachedHashForEquivalence;
 };
 
 // Experimental, this feature might be removed without warning, not for production use
@@ -112,8 +114,9 @@ private:
 struct FValidInteractionSearches
 {
 	void Update(TConstArrayView<FInteractionIsland*> Islands);
-	//FValidInteractionSearch* Get(const FInteractionSearchContext& SearchContext);
 	const FValidInteractionSearch* Get(const FInteractionSearchContext& SearchContext) const;
+	// returns true if the tuple of AnimContexts (e.g. pair of AnimInstance(s) ~ Characters) from the SearchContext is used in any FValidInteractionSearch with any assigned roles
+	bool IsThereAnyValidSearchWithTheSameContexts(const FInteractionSearchContext& SearchContext) const;
 	void AddReferencedObjects(FReferenceCollector& Collector);
 
 private:
@@ -121,7 +124,15 @@ private:
 	static void OnInteractionContinuing(FValidInteractionSearch& ValidInteractionSearch);
 	static void OnInteractionEnd(FValidInteractionSearch& ValidInteractionSearch);
 
+	// mapping from FValidInteractionSearch::GetHashForEquivalence to FValidInteractionSearch
 	TMap<uint32, FValidInteractionSearch> Searches;
+
+	// mapping from FValidInteractionSearch::GetContextHashForEquivalence to FValidInteractionSearch::GetHashForEquivalence(s)
+	// the map Key represents what AnimContexts (e.g. AnimInstance(s), that can be confused with Actors...) can be interacting. 
+	// each map Value represents a possible interaction (associated to a FValidInteractionSearch), and only one will be a valid search,
+	// and all the others are searches that still has been performed (and are valid)
+	TMap<uint32, TArray<uint32>> ContextsToSearches;
+
 #if ENABLE_ANIM_DEBUG
 	UE_MT_DECLARE_TS_RW_ACCESS_DETECTOR(ThreadSafeCounter);
 #endif // ENABLE_ANIM_DEBUG
@@ -145,10 +156,12 @@ struct FInteractionSearchContext : public FInteractionSearchContextBase
 	FPoseSearchContinuingProperties GetContinuingProperties() const;
 
 	bool IsContinuingInteraction() const { return bIsContinuingInteraction; }
+	bool IsContinuingContextInteraction() const { return bIsContinuingContextInteraction; }
 #if ENABLE_ANIM_DEBUG
 	float GetDebugBroadPhaseRadius(int32 Index) const;
 #endif // ENABLE_ANIM_DEBUG
 	void SetContinuingInteraction(bool bInIsContinuingInteraction) { bIsContinuingInteraction = bInIsContinuingInteraction; }
+	void SetContinuingContextInteraction(bool bInIsContinuingContextInteraction) { bIsContinuingContextInteraction = bInIsContinuingContextInteraction; }
 	void AddReferencedObjects(FReferenceCollector& Collector);
 	void SetContinuingProperties(float InPlayingAssetAccumulatedTime, const UObject* InPlayingAsset,
 		bool bInIsPlayingAssetMirrored, const FVector& InPlayingAssetBlendParameters, EPoseSearchInterruptMode InInterruptMode);
@@ -160,6 +173,7 @@ private:
 	FVector PlayingAssetBlendParameters = FVector::ZeroVector;
 	EPoseSearchInterruptMode InterruptMode = EPoseSearchInterruptMode::DoNotInterrupt;
 	bool bIsContinuingInteraction = false;
+	bool bIsContinuingContextInteraction = false;
 	TArray<int32, TInlineAllocator<PreallocatedRolesNum, TMemStackAllocator<>>> TickPriorities;
 
 #if ENABLE_ANIM_DEBUG
@@ -209,6 +223,7 @@ struct FInteractionIsland
 private:
 #if ENABLE_ANIM_DEBUG
 	static void LogTickDependencies(const TConstArrayView<UActorComponent*> TickActorComponents, int32 InteractionIslandIndex);
+	void ValidateTickDependenciesLeaks() const;
 #endif // ENABLE_ANIM_DEBUG
 
 	void AddTickDependencies(UActorComponent* TickActorComponent, bool bInIsMainActor);
diff --git a/Engine/Plugins/Animation/PoseSearch/Source/Runtime/Public/PoseSearch/PoseSearchLibrary.h b/Engine/Plugins/Animation/PoseSearch/Source/Runtime/Public/PoseSearch/PoseSearchLibrary.h
--- a/Engine/Plugins/Animation/PoseSearch/Source/Runtime/Public/PoseSearch/PoseSearchLibrary.h
+++ b/Engine/Plugins/Animation/PoseSearch/Source/Runtime/Public/PoseSearch/PoseSearchLibrary.h
@@ -125,6 +125,9 @@ struct FPoseSearchContinuingProperties
 	UPROPERTY(Transient, EditAnywhere, BlueprintReadWrite, Category=State)
 	bool bIsContinuingInteraction = false;
 
+	UPROPERTY(Transient, EditAnywhere, BlueprintReadWrite, Category=State)
+	bool bIsContinuingContextInteraction = false;
+	
 	// database where the PlayingAsset was originally from
 	// (optional property, used to provide an additional database to search for the continuing pose associated to PlayingAsset)
 	UPROPERTY(Transient, EditAnywhere, BlueprintReadWrite, Category=State)
