PoseSearch - defaulting GVarPoseSearchTrajectoryFacingFromMesh to true

diff --git a/Engine/Plugins/Animation/PoseSearch/Source/Runtime/Private/PoseSearchTrajectoryLibrary.cpp b/Engine/Plugins/Animation/PoseSearch/Source/Runtime/Private/PoseSearchTrajectoryLibrary.cpp
--- a/Engine/Plugins/Animation/PoseSearch/Source/Runtime/Private/PoseSearchTrajectoryLibrary.cpp
+++ b/Engine/Plugins/Animation/PoseSearch/Source/Runtime/Private/PoseSearchTrajectoryLibrary.cpp
@@ -14,9 +14,9 @@
 
 #include UE_INLINE_GENERATED_CPP_BY_NAME(PoseSearchTrajectoryLibrary)
 
-#if ENABLE_ANIM_DEBUG
 // @todo: remove GVarPoseSearchTrajectoryFacingFromMesh once the behavior is enabled by defaulted
-static bool GVarPoseSearchTrajectoryFacingFromMesh = false;
+static bool GVarPoseSearchTrajectoryFacingFromMesh = true;
+#if ENABLE_ANIM_DEBUG
 static FAutoConsoleVariableRef CVarPoseSearchTrajectoryFacingFromMesh(TEXT("a.PoseSearchTrajectory.FacingFromMesh"), GVarPoseSearchTrajectoryFacingFromMesh, TEXT("Always generating the trajectory in mesh space rather than using the Character.View.Rotation.Yaw when the character is set to orienting the rotation to movement"));
 #endif // ENABLE_ANIM_DEBUG
 
@@ -99,11 +99,7 @@ bool FPoseSearchTrajectoryData::UpdateData(
 	TrajectoryDataDerived.Position = MeshComp->GetComponentLocation();
 	TrajectoryDataDerived.MeshCompRelativeRotation = MeshComp->GetRelativeRotation().Quaternion();
 	
-	if (TrajectoryDataDerived.bOrientRotationToMovement
-#if ENABLE_ANIM_DEBUG
-		|| GVarPoseSearchTrajectoryFacingFromMesh
-#endif // ENABLE_ANIM_DEBUG
-		)
+	if (GVarPoseSearchTrajectoryFacingFromMesh || TrajectoryDataDerived.bOrientRotationToMovement)
 	{
 		TrajectoryDataDerived.Facing = MeshComp->GetComponentTransform().GetRotation();
 	}
