Handle GetPoseIndex returning -1 in FMotionMatchingState::GetEstimatedFutureRootMotionVelocity

diff --git a/Engine/Plugins/Animation/PoseSearch/Source/Runtime/Private/PoseSearchLibrary.cpp b/Engine/Plugins/Animation/PoseSearch/Source/Runtime/Private/PoseSearchLibrary.cpp
--- a/Engine/Plugins/Animation/PoseSearch/Source/Runtime/Private/PoseSearchLibrary.cpp
+++ b/Engine/Plugins/Animation/PoseSearch/Source/Runtime/Private/PoseSearchLibrary.cpp
@@ -590,7 +590,7 @@ FVector FMotionMatchingState::GetEstimatedFutureRootMotionVelocity() const
 		{
 			const int32 PoseIndex = Database->GetPoseIndex(SearchResult.SelectedAnim.Get(), SearchResult.SelectedTime, SearchResult.bIsMirrored, SearchResult.BlendParameters);
 			const FSearchIndex& SearchIndex = Database->GetSearchIndex();
-			if (!SearchIndex.IsValuesEmpty())
+			if (PoseIndex != INDEX_NONE && !SearchIndex.IsValuesEmpty())
 			{
 				TConstArrayView<float> ResultData = SearchIndex.GetPoseValues(PoseIndex);
 				return TrajectoryChannel->GetEstimatedFutureRootMotionVelocity(ResultData);
