PoseSearch - fix for GVarMotionMatchCompareAgainstBruteForce

diff --git a/Engine/Plugins/Animation/PoseSearch/Source/Editor/Private/PoseSearchDebugger.cpp b/Engine/Plugins/Animation/PoseSearch/Source/Editor/Private/PoseSearchDebugger.cpp
--- a/Engine/Plugins/Animation/PoseSearch/Source/Editor/Private/PoseSearchDebugger.cpp
+++ b/Engine/Plugins/Animation/PoseSearch/Source/Editor/Private/PoseSearchDebugger.cpp
@@ -264,7 +264,7 @@ FReply SCostTimelineView::OnMouseMove(const FGeometry& MyGeometry, const FPointe
 					ToolTipTime = FText::Format(LOCTEXT("CostTimelineViewToolTip_TimeFormat", "Search Time: {0}"), FText::AsNumber(Time, &FormattingOptions));
 					ToolTipCost = FText::Format(LOCTEXT("CostTimelineViewToolTip_CostFormat", "Search Cost: {0}"), FText::AsNumber(BestCost, &FormattingOptions));
 
-					if (!FPoseSearchCost::IsCostValid(BruteForceCost) || FMath::IsNearlyEqual(BestCost, BruteForceCost))
+					if (!FPoseSearchCost::IsCostValid(BruteForceCost))
 					{
 						ToolTipCostBruteForce = FText::GetEmpty();
 					}
diff --git a/Engine/Plugins/Animation/PoseSearch/Source/Runtime/Private/PoseSearchDatabase.cpp b/Engine/Plugins/Animation/PoseSearch/Source/Runtime/Private/PoseSearchDatabase.cpp
--- a/Engine/Plugins/Animation/PoseSearch/Source/Runtime/Private/PoseSearchDatabase.cpp
+++ b/Engine/Plugins/Animation/PoseSearch/Source/Runtime/Private/PoseSearchDatabase.cpp
@@ -155,7 +155,7 @@ static inline void EvaluatePoseKernel(FSearchResults& SearchResults, const FSear
 	, bool bTrackResult = true
 #endif // UE_POSE_SEARCH_TRACE_ENABLED
 #if WITH_EDITOR && ENABLE_ANIM_DEBUG && UE_POSE_SEARCH_TRACE_ENABLED
-	, const FSearchResult& BruteForceSearchResult = FSearchResult(), int32 ResultIndex = INDEX_NONE
+	, int32 ResultIndex = INDEX_NONE
 #endif // WITH_EDITOR && ENABLE_ANIM_DEBUG && UE_POSE_SEARCH_TRACE_ENABLED
 )
 {
@@ -177,7 +177,6 @@ static inline void EvaluatePoseKernel(FSearchResults& SearchResults, const FSear
 		
 #if WITH_EDITOR && ENABLE_ANIM_DEBUG && UE_POSE_SEARCH_TRACE_ENABLED
 		Result.BestPosePos = ResultIndex;
-		Result.BruteForcePoseCost = BruteForceSearchResult.PoseCost;
 #endif // WITH_EDITOR && ENABLE_ANIM_DEBUG && UE_POSE_SEARCH_TRACE_ENABLED
 
 		SearchResults.UpdateWith(Result);
@@ -1376,6 +1375,32 @@ void UPoseSearchDatabase::Search(UE::PoseSearch::FSearchContext& SearchContext,
 		case EPoseSearchMode::EventOnly:
 			break;
 		}
+
+#if WITH_EDITOR && ENABLE_ANIM_DEBUG && UE_POSE_SEARCH_TRACE_ENABLED
+		if (GVarMotionMatchCompareAgainstBruteForce)
+		{
+			FPoseSearchCost BruteForcePoseCost;
+			if (PoseSearchMode == EPoseSearchMode::VPTree || PoseSearchMode == EPoseSearchMode::PCAKDTree)
+			{
+				FSearchResults_Single BruteForceSearchResults;
+				SearchBruteForce(SearchContext, BruteForceSearchResults);
+				BruteForcePoseCost = BruteForceSearchResults.GetBestResult().PoseCost;
+			}
+			else
+			{
+				BruteForcePoseCost = SearchResults.GetBestResult().PoseCost;
+			}
+
+			SearchResults.IterateOverSearchResults([&BruteForcePoseCost](FSearchResult& SearchResult)
+				{
+					if (BruteForcePoseCost < SearchResult.BruteForcePoseCost)
+					{
+						SearchResult.BruteForcePoseCost = BruteForcePoseCost;
+					}
+					return false;
+				});
+		}
+#endif // WITH_EDITOR && ENABLE_ANIM_DEBUG && UE_POSE_SEARCH_TRACE_ENABLED
 	}
 }
 
@@ -1690,7 +1715,7 @@ void UPoseSearchDatabase::SearchContinuingPose(UE::PoseSearch::FSearchContext& S
 	// since any PoseCost calculated here is at least SearchIndex.MinCostAddend + ContinuingPoseCostAddend + ContinuingInteractionCostAddend + ContinuingContextInteractionCostAddend,
 	// there's no point in performing the search if the worst search result is already better than that
 	const FSearchIndex& SearchIndex = GetSearchIndex();
-	if (!GetSkipSearchIfPossible() || SearchResults.ShouldPerformSearch(SearchIndex.MinCostAddend + ContinuingPoseCostAddend + ContinuingInteractionCostAddend + ContinuingContextInteractionCostAddend))
+	if (SearchResults.ShouldPerformSearch(SearchIndex.MinCostAddend + ContinuingPoseCostAddend + ContinuingInteractionCostAddend + ContinuingContextInteractionCostAddend))
 	{
 		const int32 ContinuingPoseIdx = SearchContext.GetContinuingPoseSearchResult().PoseIdx;
 		check(ContinuingPoseIdx != INDEX_NONE);
@@ -1726,15 +1751,23 @@ void UPoseSearchDatabase::SearchContinuingPose(UE::PoseSearch::FSearchContext& S
 		SearchContext.Track(this, ContinuingPoseIdx, EPoseCandidateFlags::Valid_ContinuingPose, Result.PoseCost);
 #endif // UE_POSE_SEARCH_TRACE_ENABLED
 
+		check(Result.DebugValidate());
+		SearchResults.UpdateWith(Result);
+
 #if WITH_EDITOR && ENABLE_ANIM_DEBUG && UE_POSE_SEARCH_TRACE_ENABLED
 		if (GVarMotionMatchCompareAgainstBruteForce)
 		{
-			Result.BruteForcePoseCost = Result.PoseCost;
+			const FPoseSearchCost& BruteForcePoseCost = Result.PoseCost;
+			SearchResults.IterateOverSearchResults([&BruteForcePoseCost](FSearchResult& SearchResult)
+				{
+					if (BruteForcePoseCost < SearchResult.BruteForcePoseCost)
+					{
+						SearchResult.BruteForcePoseCost = BruteForcePoseCost;
+					}
+					return false;
+				});
 		}
 #endif // WITH_EDITOR && ENABLE_ANIM_DEBUG && UE_POSE_SEARCH_TRACE_ENABLED
-
-		check(Result.DebugValidate());
-		SearchResults.UpdateWith(Result);
 	}
 	else
 	{
@@ -1753,16 +1786,6 @@ void UPoseSearchDatabase::SearchPCAKDTree(UE::PoseSearch::FSearchContext& Search
 
 	using namespace UE::PoseSearch;
 
-#if WITH_EDITOR && ENABLE_ANIM_DEBUG && UE_POSE_SEARCH_TRACE_ENABLED
-	FSearchResult BruteForceSearchResult;
-	if (GVarMotionMatchCompareAgainstBruteForce)
-	{
-		FSearchResults_Single BruteForceSearchResults;
-		SearchBruteForce(SearchContext, BruteForceSearchResults);
-		BruteForceSearchResult = BruteForceSearchResults.GetBestResult();
-	}
-#endif // WITH_EDITOR && ENABLE_ANIM_DEBUG && UE_POSE_SEARCH_TRACE_ENABLED
-
 	float ContinuingPoseCostAddend = 0.f, ContinuingInteractionCostAddend = 0.f, ContinuingContextInteractionCostAddend = 0.f;
 	ExtractCostAddends(SearchContext, false, ContinuingPoseCostAddend, ContinuingInteractionCostAddend, ContinuingContextInteractionCostAddend);
 
@@ -1770,7 +1793,7 @@ void UPoseSearchDatabase::SearchPCAKDTree(UE::PoseSearch::FSearchContext& Search
 	// there's no point in performing the search if the worst search result is already better than that
 	const FSearchIndex& SearchIndex = GetSearchIndex();
 	const int32 NumDimensions = Schema->SchemaCardinality;
-	if (!GetSkipSearchIfPossible() || SearchResults.ShouldPerformSearch(SearchIndex.MinCostAddend + ContinuingPoseCostAddend + ContinuingInteractionCostAddend + ContinuingContextInteractionCostAddend))
+	if (SearchResults.ShouldPerformSearch(SearchIndex.MinCostAddend + ContinuingPoseCostAddend + ContinuingInteractionCostAddend + ContinuingContextInteractionCostAddend))
 	{
 		const uint32 ClampedNumberOfPrincipalComponents = SearchIndex.GetNumberOfPrincipalComponents();
 		const uint32 ClampedKDTreeQueryNumNeighbors = FMath::Clamp<uint32>(KDTreeQueryNumNeighbors, 1, SearchIndex.GetNumPoses());
@@ -1954,7 +1977,7 @@ void UPoseSearchDatabase::SearchPCAKDTree(UE::PoseSearch::FSearchContext& Search
 								, true
 #endif // UE_POSE_SEARCH_TRACE_ENABLED
 #if WITH_EDITOR && ENABLE_ANIM_DEBUG && UE_POSE_SEARCH_TRACE_ENABLED
-								, BruteForceSearchResult, SortedResultsIndexes[ResultIndex]
+								, SortedResultsIndexes[ResultIndex]
 #endif // WITH_EDITOR && ENABLE_ANIM_DEBUG && UE_POSE_SEARCH_TRACE_ENABLED
 							);
 						}
@@ -1973,7 +1996,7 @@ void UPoseSearchDatabase::SearchPCAKDTree(UE::PoseSearch::FSearchContext& Search
 								, true
 #endif // UE_POSE_SEARCH_TRACE_ENABLED
 #if WITH_EDITOR && ENABLE_ANIM_DEBUG && UE_POSE_SEARCH_TRACE_ENABLED
-								, BruteForceSearchResult, SortedResultsIndexes[ResultIndex]
+								, SortedResultsIndexes[ResultIndex]
 #endif // WITH_EDITOR && ENABLE_ANIM_DEBUG && UE_POSE_SEARCH_TRACE_ENABLED
 							);
 						}
@@ -1996,7 +2019,7 @@ void UPoseSearchDatabase::SearchPCAKDTree(UE::PoseSearch::FSearchContext& Search
 								, true
 #endif // UE_POSE_SEARCH_TRACE_ENABLED
 #if WITH_EDITOR && ENABLE_ANIM_DEBUG && UE_POSE_SEARCH_TRACE_ENABLED
-								, BruteForceSearchResult, SortedResultsIndexes[ResultIndex]
+								, SortedResultsIndexes[ResultIndex]
 #endif // WITH_EDITOR && ENABLE_ANIM_DEBUG && UE_POSE_SEARCH_TRACE_ENABLED
 							);
 						}
@@ -2015,7 +2038,7 @@ void UPoseSearchDatabase::SearchPCAKDTree(UE::PoseSearch::FSearchContext& Search
 								, true
 #endif // UE_POSE_SEARCH_TRACE_ENABLED
 #if WITH_EDITOR && ENABLE_ANIM_DEBUG && UE_POSE_SEARCH_TRACE_ENABLED
-								, BruteForceSearchResult, SortedResultsIndexes[ResultIndex]
+								, SortedResultsIndexes[ResultIndex]
 #endif // WITH_EDITOR && ENABLE_ANIM_DEBUG && UE_POSE_SEARCH_TRACE_ENABLED
 							);
 						}
@@ -2040,7 +2063,7 @@ void UPoseSearchDatabase::SearchPCAKDTree(UE::PoseSearch::FSearchContext& Search
 						, true
 	#endif // UE_POSE_SEARCH_TRACE_ENABLED
 	#if WITH_EDITOR && ENABLE_ANIM_DEBUG && UE_POSE_SEARCH_TRACE_ENABLED
-						, BruteForceSearchResult, SortedResultsIndexes[ResultIndex]
+						, SortedResultsIndexes[ResultIndex]
 	#endif // WITH_EDITOR && ENABLE_ANIM_DEBUG && UE_POSE_SEARCH_TRACE_ENABLED
 					);
 				}
@@ -2054,7 +2077,7 @@ void UPoseSearchDatabase::SearchPCAKDTree(UE::PoseSearch::FSearchContext& Search
 						, true
 #endif // UE_POSE_SEARCH_TRACE_ENABLED
 #if WITH_EDITOR && ENABLE_ANIM_DEBUG && UE_POSE_SEARCH_TRACE_ENABLED
-						, BruteForceSearchResult, SortedResultsIndexes[ResultIndex]
+						, SortedResultsIndexes[ResultIndex]
 #endif // WITH_EDITOR && ENABLE_ANIM_DEBUG && UE_POSE_SEARCH_TRACE_ENABLED
 					);
 				}
@@ -2072,7 +2095,7 @@ void UPoseSearchDatabase::SearchPCAKDTree(UE::PoseSearch::FSearchContext& Search
 						, true
 #endif // UE_POSE_SEARCH_TRACE_ENABLED
 #if WITH_EDITOR && ENABLE_ANIM_DEBUG && UE_POSE_SEARCH_TRACE_ENABLED
-						, BruteForceSearchResult, SortedResultsIndexes[ResultIndex]
+						, SortedResultsIndexes[ResultIndex]
 #endif // WITH_EDITOR && ENABLE_ANIM_DEBUG && UE_POSE_SEARCH_TRACE_ENABLED
 					);
 				}
@@ -2087,7 +2110,7 @@ void UPoseSearchDatabase::SearchPCAKDTree(UE::PoseSearch::FSearchContext& Search
 						, true
 #endif // UE_POSE_SEARCH_TRACE_ENABLED
 #if WITH_EDITOR && ENABLE_ANIM_DEBUG && UE_POSE_SEARCH_TRACE_ENABLED
-						, BruteForceSearchResult, SortedResultsIndexes[ResultIndex]
+						, SortedResultsIndexes[ResultIndex]
 #endif // WITH_EDITOR && ENABLE_ANIM_DEBUG && UE_POSE_SEARCH_TRACE_ENABLED
 					);
 				}
@@ -2121,23 +2144,13 @@ void UPoseSearchDatabase::SearchVPTree(UE::PoseSearch::FSearchContext& SearchCon
 
 	using namespace UE::PoseSearch;
 
-#if WITH_EDITOR && ENABLE_ANIM_DEBUG && UE_POSE_SEARCH_TRACE_ENABLED
-	FSearchResult BruteForceSearchResult;
-	if (GVarMotionMatchCompareAgainstBruteForce)
-	{
-		FSearchResults_Single BruteForceSearchResults;
-		SearchBruteForce(SearchContext, BruteForceSearchResults);
-		BruteForceSearchResult = BruteForceSearchResults.GetBestResult();
-	}
-#endif // WITH_EDITOR && ENABLE_ANIM_DEBUG && UE_POSE_SEARCH_TRACE_ENABLED
-
 	float ContinuingPoseCostAddend = 0.f, ContinuingInteractionCostAddend = 0.f, ContinuingContextInteractionCostAddend = 0.f;
 	ExtractCostAddends(SearchContext, false, ContinuingPoseCostAddend, ContinuingInteractionCostAddend, ContinuingContextInteractionCostAddend);
 
 	// since any PoseCost calculated here is at least SearchIndex.MinCostAddend + ContinuingPoseCostAddend + ContinuingInteractionCostAddend + ContinuingContextInteractionCostAddend,
 	// there's no point in performing the search if the worst search result is already better than that
 	const FSearchIndex& SearchIndex = GetSearchIndex();
-	if (!GetSkipSearchIfPossible() || SearchResults.ShouldPerformSearch(SearchIndex.MinCostAddend + ContinuingPoseCostAddend + ContinuingInteractionCostAddend + ContinuingContextInteractionCostAddend))
+	if (SearchResults.ShouldPerformSearch(SearchIndex.MinCostAddend + ContinuingPoseCostAddend + ContinuingInteractionCostAddend + ContinuingContextInteractionCostAddend))
 	{
 		// @todo: perhaps store the DynamicWeightsSqrt into the GetOrBuildQuery to share between the continuing pose and the full search
 		const int32 NumDimensions = Schema->SchemaCardinality;
@@ -2194,7 +2207,7 @@ void UPoseSearchDatabase::SearchVPTree(UE::PoseSearch::FSearchContext& SearchCon
 						, true
 #endif // UE_POSE_SEARCH_TRACE_ENABLED
 #if WITH_EDITOR && ENABLE_ANIM_DEBUG && UE_POSE_SEARCH_TRACE_ENABLED
-						, BruteForceSearchResult, ResultIndex
+						, ResultIndex
 #endif // WITH_EDITOR && ENABLE_ANIM_DEBUG && UE_POSE_SEARCH_TRACE_ENABLED
 					);
 				}
@@ -2212,7 +2225,7 @@ void UPoseSearchDatabase::SearchVPTree(UE::PoseSearch::FSearchContext& SearchCon
 					, true
 #endif // UE_POSE_SEARCH_TRACE_ENABLED
 #if WITH_EDITOR && ENABLE_ANIM_DEBUG && UE_POSE_SEARCH_TRACE_ENABLED
-					, BruteForceSearchResult, ResultIndex
+					, ResultIndex
 #endif // WITH_EDITOR && ENABLE_ANIM_DEBUG && UE_POSE_SEARCH_TRACE_ENABLED
 				);
 			}
@@ -2254,7 +2267,7 @@ void UPoseSearchDatabase::SearchBruteForce(UE::PoseSearch::FSearchContext& Searc
 	// since any PoseCost calculated here is at least SearchIndex.MinCostAddend + ContinuingPoseCostAddend + ContinuingInteractionCostAddend + ContinuingContextInteractionCostAddend,
 	// there's no point in performing the search if the worst search result is already better than that
 	const FSearchIndex& SearchIndex = GetSearchIndex();
-	if (!GetSkipSearchIfPossible() || SearchResults.ShouldPerformSearch(SearchIndex.MinCostAddend + ContinuingPoseCostAddend + ContinuingInteractionCostAddend + ContinuingContextInteractionCostAddend))
+	if (SearchResults.ShouldPerformSearch(SearchIndex.MinCostAddend + ContinuingPoseCostAddend + ContinuingInteractionCostAddend + ContinuingContextInteractionCostAddend))
 	{
 		// @todo: perhaps store the DynamicWeightsSqrt into the GetOrBuildQuery to share between the continuing pose and the full search
 		const int32 NumDimensions = Schema->SchemaCardinality;
@@ -2299,7 +2312,7 @@ void UPoseSearchDatabase::SearchBruteForce(UE::PoseSearch::FSearchContext& Searc
 						, bTrackResults
 #endif // UE_POSE_SEARCH_TRACE_ENABLED
 #if WITH_EDITOR && ENABLE_ANIM_DEBUG && UE_POSE_SEARCH_TRACE_ENABLED
-						, FSearchResult(), PoseIdx
+						, PoseIdx
 #endif // WITH_EDITOR && ENABLE_ANIM_DEBUG && UE_POSE_SEARCH_TRACE_ENABLED
 					);
 				}
@@ -2314,7 +2327,7 @@ void UPoseSearchDatabase::SearchBruteForce(UE::PoseSearch::FSearchContext& Searc
 						, bTrackResults
 #endif // UE_POSE_SEARCH_TRACE_ENABLED
 #if WITH_EDITOR && ENABLE_ANIM_DEBUG && UE_POSE_SEARCH_TRACE_ENABLED
-						, FSearchResult(), PoseIdx
+						, PoseIdx
 #endif // WITH_EDITOR && ENABLE_ANIM_DEBUG && UE_POSE_SEARCH_TRACE_ENABLED
 					);
 				}
@@ -2329,7 +2342,7 @@ void UPoseSearchDatabase::SearchBruteForce(UE::PoseSearch::FSearchContext& Searc
 						, bTrackResults
 #endif // UE_POSE_SEARCH_TRACE_ENABLED
 #if WITH_EDITOR && ENABLE_ANIM_DEBUG && UE_POSE_SEARCH_TRACE_ENABLED
-						, FSearchResult(), PoseIdx
+						, PoseIdx
 #endif // WITH_EDITOR && ENABLE_ANIM_DEBUG && UE_POSE_SEARCH_TRACE_ENABLED
 					);
 				}
@@ -2358,7 +2371,7 @@ void UPoseSearchDatabase::SearchBruteForce(UE::PoseSearch::FSearchContext& Searc
 							, bTrackResults
 #endif // UE_POSE_SEARCH_TRACE_ENABLED
 #if WITH_EDITOR && ENABLE_ANIM_DEBUG && UE_POSE_SEARCH_TRACE_ENABLED
-							, FSearchResult(), ++ResultIndex
+							, ++ResultIndex
 #endif // WITH_EDITOR && ENABLE_ANIM_DEBUG && UE_POSE_SEARCH_TRACE_ENABLED
 						);
 					}
@@ -2379,7 +2392,7 @@ void UPoseSearchDatabase::SearchBruteForce(UE::PoseSearch::FSearchContext& Searc
 							, bTrackResults
 #endif // UE_POSE_SEARCH_TRACE_ENABLED
 #if WITH_EDITOR && ENABLE_ANIM_DEBUG && UE_POSE_SEARCH_TRACE_ENABLED
-							, FSearchResult(), ++ResultIndex
+							, ++ResultIndex
 #endif // WITH_EDITOR && ENABLE_ANIM_DEBUG && UE_POSE_SEARCH_TRACE_ENABLED
 						);
 					}
@@ -2400,7 +2413,7 @@ void UPoseSearchDatabase::SearchBruteForce(UE::PoseSearch::FSearchContext& Searc
 							, bTrackResults
 #endif // UE_POSE_SEARCH_TRACE_ENABLED
 #if WITH_EDITOR && ENABLE_ANIM_DEBUG && UE_POSE_SEARCH_TRACE_ENABLED
-							, FSearchResult(), ++ResultIndex
+							, ++ResultIndex
 #endif // WITH_EDITOR && ENABLE_ANIM_DEBUG && UE_POSE_SEARCH_TRACE_ENABLED
 						);
 					}
@@ -2492,7 +2505,7 @@ void UPoseSearchDatabase::SearchEvent(UE::PoseSearch::FSearchContext& SearchCont
 						, true
 #endif // UE_POSE_SEARCH_TRACE_ENABLED
 #if WITH_EDITOR && ENABLE_ANIM_DEBUG && UE_POSE_SEARCH_TRACE_ENABLED
-						, FSearchResult(), EventPoseIdx
+						, EventPoseIdx
 #endif // WITH_EDITOR && ENABLE_ANIM_DEBUG && UE_POSE_SEARCH_TRACE_ENABLED
 					);
 				}
@@ -2514,7 +2527,7 @@ void UPoseSearchDatabase::SearchEvent(UE::PoseSearch::FSearchContext& SearchCont
 								, true
 #endif // UE_POSE_SEARCH_TRACE_ENABLED
 #if WITH_EDITOR && ENABLE_ANIM_DEBUG && UE_POSE_SEARCH_TRACE_ENABLED
-								, FSearchResult(), PoseIdx
+								, PoseIdx
 #endif // WITH_EDITOR && ENABLE_ANIM_DEBUG && UE_POSE_SEARCH_TRACE_ENABLED
 							);
 						});
diff --git a/Engine/Plugins/Animation/PoseSearch/Source/Runtime/Private/PoseSearchLibrary.cpp b/Engine/Plugins/Animation/PoseSearch/Source/Runtime/Private/PoseSearchLibrary.cpp
--- a/Engine/Plugins/Animation/PoseSearch/Source/Runtime/Private/PoseSearchLibrary.cpp
+++ b/Engine/Plugins/Animation/PoseSearch/Source/Runtime/Private/PoseSearchLibrary.cpp
@@ -721,7 +721,7 @@ void UPoseSearchLibrary::TraceMotionMatching(
 
 	// finalizing SearchContext.GetBestPoseCandidatesMap() by calling SearchContext::Track for all the SearchResults flagging them as EPoseCandidateFlags::Valid_CurrentPose
 	// because, those candidate could have been discarded becasue of their cost was too high and didn't fit into SearchContext.GetBestPoseCandidatesMap() limits
-	SearchResults.IterateOverSearchResults([&SearchContext](const FSearchResult& SearchResult)
+	SearchResults.IterateOverSearchResults([&SearchContext](FSearchResult& SearchResult)
 		{
 			if (SearchResult.IsValid())
 			{
@@ -1382,11 +1382,8 @@ void UPoseSearchLibrary::MotionMatch(
 		PRAGMA_ENABLE_DEPRECATION_WARNINGS
 	}
 
-#if (ENABLE_DRAW_DEBUG && ENABLE_ANIM_DEBUG) || UE_POSE_SEARCH_TRACE_ENABLED
-	const FSearchResult SearchResult = SearchResults.GetBestResult();
-#endif // (ENABLE_DRAW_DEBUG && ENABLE_ANIM_DEBUG) || UE_POSE_SEARCH_TRACE_ENABLED
-
 #if ENABLE_DRAW_DEBUG && ENABLE_ANIM_DEBUG
+	const FSearchResult SearchResult = SearchResults.GetBestResult();
 	if (SearchResult.IsValid())
 	{
 		FDebugDrawParams DrawParams;
diff --git a/Engine/Plugins/Animation/PoseSearch/Source/Runtime/Private/PoseSearchResult.cpp b/Engine/Plugins/Animation/PoseSearch/Source/Runtime/Private/PoseSearchResult.cpp
--- a/Engine/Plugins/Animation/PoseSearch/Source/Runtime/Private/PoseSearchResult.cpp
+++ b/Engine/Plugins/Animation/PoseSearch/Source/Runtime/Private/PoseSearchResult.cpp
@@ -255,7 +255,7 @@ FSearchResults_Single::FSearchResults_Single()
 			return SingleSearchResult.PoseCost > SearchMinimumCost;
 		};
 
-	IterateOverSearchResults = [this](const TFunctionRef<bool(const FSearchResult& SearchResult)> IterateOverSearchResultsFunction)
+	IterateOverSearchResults = [this](const TFunctionRef<bool(FSearchResult& SearchResult)> IterateOverSearchResultsFunction)
 		{
 			return IterateOverSearchResultsFunction(SingleSearchResult);
 		};
@@ -330,9 +330,9 @@ FSearchResults_Multi::FSearchResults_Multi(int32 InMaxNumberOfWantedResults)
 			return SearchResults.Last().PoseCost > SearchMinimumCost;
 		};
 
-	IterateOverSearchResults = [this](const TFunctionRef<bool(const FSearchResult& SearchResult)> IterateOverSearchResultsFunction)
+	IterateOverSearchResults = [this](const TFunctionRef<bool(FSearchResult& SearchResult)> IterateOverSearchResultsFunction)
 		{
-			for (const FSearchResult& SearchResult : SearchResults)
+			for (FSearchResult& SearchResult : SearchResults)
 			{
 				if (IterateOverSearchResultsFunction(SearchResult))
 				{
@@ -400,11 +400,11 @@ FSearchResults_AssetBests::FSearchResults_AssetBests()
 			return true;
 		};
 
-	IterateOverSearchResults = [this](const TFunctionRef<bool(const FSearchResult& SearchResult)> IterateOverSearchResultsFunction)
+	IterateOverSearchResults = [this](const TFunctionRef<bool(FSearchResult& SearchResult)> IterateOverSearchResultsFunction)
 		{
-			for (const TPair<const UPoseSearchDatabase*, FSourceAssetIdxToBestSearchResult>& PerDatabaseSourceAssetIdxToBestSearchResultPair : PerDatabaseSourceAssetIdxToBestSearchResult)
+			for (TPair<const UPoseSearchDatabase*, FSourceAssetIdxToBestSearchResult>& PerDatabaseSourceAssetIdxToBestSearchResultPair : PerDatabaseSourceAssetIdxToBestSearchResult)
 			{
-				for (const TPair<int32, FSearchResult>& SourceAssetIdxToBestSearchResultPair : PerDatabaseSourceAssetIdxToBestSearchResultPair.Value)
+				for (TPair<int32, FSearchResult>& SourceAssetIdxToBestSearchResultPair : PerDatabaseSourceAssetIdxToBestSearchResultPair.Value)
 				{
 					if (IterateOverSearchResultsFunction(SourceAssetIdxToBestSearchResultPair.Value))
 					{
diff --git a/Engine/Plugins/Animation/PoseSearch/Source/Runtime/Public/PoseSearch/PoseSearchDatabase.h b/Engine/Plugins/Animation/PoseSearch/Source/Runtime/Public/PoseSearch/PoseSearchDatabase.h
--- a/Engine/Plugins/Animation/PoseSearch/Source/Runtime/Public/PoseSearch/PoseSearchDatabase.h
+++ b/Engine/Plugins/Animation/PoseSearch/Source/Runtime/Public/PoseSearch/PoseSearchDatabase.h
@@ -635,6 +635,7 @@ public:
 	UE_API void SetSearchIndex(const UE::PoseSearch::FSearchIndex& SearchIndex);
 	UE_API const UE::PoseSearch::FSearchIndex& GetSearchIndex() const;
 	
+	UE_DEPRECATED(5.8, "no longer necessary API")
 	UE_API bool GetSkipSearchIfPossible() const;
 
 	// given an AnimationAsset, a normalized time AnimationAssetTime (in [0-1] range for blend spaces), and mirror state bMirrored
diff --git a/Engine/Plugins/Animation/PoseSearch/Source/Runtime/Public/PoseSearch/PoseSearchResult.h b/Engine/Plugins/Animation/PoseSearch/Source/Runtime/Public/PoseSearch/PoseSearchResult.h
--- a/Engine/Plugins/Animation/PoseSearch/Source/Runtime/Public/PoseSearch/PoseSearchResult.h
+++ b/Engine/Plugins/Animation/PoseSearch/Source/Runtime/Public/PoseSearch/PoseSearchResult.h
@@ -126,7 +126,7 @@ struct FSearchResults
 		, FinalizeResults([]() {})
 		, GetBestResult([]() { return FSearchResult(); })
 		, ShouldPerformSearch([](float) { return true; })
-		, IterateOverSearchResults([](const TFunctionRef<bool(const FSearchResult& SearchResult)>) { return true; })
+		, IterateOverSearchResults([](const TFunctionRef<bool(FSearchResult& SearchResult)>) { return true; })
 	{
 	}
 
@@ -147,7 +147,7 @@ struct FSearchResults
 
 	// calls IterateOverSearchResultsFunction for each FSearchResult this FSearchResults is holding
 	// is IterateOverSearchResultsFunction returns true, the iteration will stop and IterateOverSearchResults will return true otherwise will return false (useful to implement contains / find methods)
-	TFunction<bool(const TFunctionRef<bool(const FSearchResult& SearchResult)>)> IterateOverSearchResults;
+	TFunction<bool(const TFunctionRef<bool(FSearchResult& SearchResult)>)> IterateOverSearchResults;
 };
 
 // Experimental, this feature might be removed without warning, not for production use
