PoseSearch
- fix for overly constrained extrapolation by delegating time clamping between 0 and PlayLength to the FAnimationAssetSampler
- fix for UPoseSearchFeatureChannel_Heading and UPoseSearchFeatureChannel_Position FindOrAddToSchema to fix channel debug drawing

diff --git a/Engine/Plugins/Animation/PoseSearch/Source/Runtime/Private/PoseSearchAssetIndexer.cpp b/Engine/Plugins/Animation/PoseSearch/Source/Runtime/Private/PoseSearchAssetIndexer.cpp
--- a/Engine/Plugins/Animation/PoseSearch/Source/Runtime/Private/PoseSearchAssetIndexer.cpp
+++ b/Engine/Plugins/Animation/PoseSearch/Source/Runtime/Private/PoseSearchAssetIndexer.cpp
@@ -25,63 +25,6 @@ static int32 GVarMotionMatchTestExtractPoseDeterminismNumIterations = 0;
 static FAutoConsoleVariableRef CVarMotionMatchTestExtractPoseDeterminismNumIterations(TEXT("a.MotionMatch.TestExtractPoseDeterminismNumIterations"), GVarMotionMatchTestExtractPoseDeterminismNumIterations, TEXT("Test Motion Matching ExtractPose Determinism via this NumIterations retries"));
 #endif // ENABLE_ANIM_DEBUG
 
-//////////////////////////////////////////////////////////////////////////
-// FSamplingParam helpers
-struct FSamplingParam
-{
-	float WrappedParam = 0.0f;
-	int32 NumCycles = 0;
-
-	// If the animation can't loop, WrappedParam contains the clamped value and whatever is left is stored here
-	float Extrapolation = 0.0f;
-};
-
-static FSamplingParam WrapOrClampSamplingParam(bool bCanWrap, float SamplingParamExtent, float SamplingParam)
-{
-	// This is a helper function used by both time and distance sampling. A schema may specify time or distance
-	// offsets that are multiple cycles of a clip away from the current pose being sampled.
-	// And that time or distance offset may before the beginning of the clip (SamplingParam < 0.0f)
-	// or after the end of the clip (SamplingParam > SamplingParamExtent). So this function
-	// helps determine how many cycles need to be applied and what the wrapped value should be, clamping
-	// if necessary.
-
-	FSamplingParam Result;
-
-	Result.WrappedParam = SamplingParam;
-
-	const bool bIsSamplingParamExtentKindaSmall = SamplingParamExtent <= UE_KINDA_SMALL_NUMBER;
-	if (!bIsSamplingParamExtentKindaSmall && bCanWrap)
-	{
-		if (SamplingParam < 0.0f)
-		{
-			while (Result.WrappedParam < 0.0f)
-			{
-				Result.WrappedParam += SamplingParamExtent;
-				++Result.NumCycles;
-			}
-		}
-
-		else
-		{
-			while (Result.WrappedParam > SamplingParamExtent)
-			{
-				Result.WrappedParam -= SamplingParamExtent;
-				++Result.NumCycles;
-			}
-		}
-	}
-
-	const float ParamClamped = FMath::Clamp(Result.WrappedParam, 0.0f, SamplingParamExtent);
-	if (ParamClamped != Result.WrappedParam)
-	{
-		check(bIsSamplingParamExtentKindaSmall || !bCanWrap);
-		Result.Extrapolation = Result.WrappedParam - ParamClamped;
-		Result.WrappedParam = ParamClamped;
-	}
-
-	return Result;
-}
-
 //////////////////////////////////////////////////////////////////////////
 // FAssetSamplingContext
 FAssetSamplingContext::FAssetSamplingContext(const UPoseSearchDatabase& Database)
@@ -333,101 +276,27 @@ void FAssetIndexer::ComputeStats()
 		{
 			const float SampleTime = FMath::Min(CalculateSampleTime(SampleIdx), GetPlayLength());
 
-			bool AnyClamped = false;
-			const FTransform TrajTransformsPast = GetTransform(SampleTime - FiniteDelta, RoleIndex, AnyClamped, RootBoneReference);
-			if (!AnyClamped)
-			{
-				const FTransform TrajTransformsPresent = GetTransform(SampleTime, RoleIndex, AnyClamped, RootBoneReference);
-				if (!AnyClamped)
-				{
-					const FTransform TrajTransformsFuture = GetTransform(SampleTime + FiniteDelta, RoleIndex, AnyClamped, RootBoneReference);
-					if (!AnyClamped)
-					{
-						// if any transform is clamped we just skip the sample entirely
-						const FVector LinearVelocityPresent = (TrajTransformsPresent.GetTranslation() - TrajTransformsPast.GetTranslation()) / FiniteDelta;
-						const FVector LinearVelocityFuture = (TrajTransformsFuture.GetTranslation() - TrajTransformsPresent.GetTranslation()) / FiniteDelta;
-						const FVector LinearAcceleration = (LinearVelocityFuture - LinearVelocityPresent) / FiniteDelta;
+			const FTransform TrajTransformsPast = GetTransform(SampleTime - FiniteDelta, RoleIndex, RootBoneReference);
+			const FTransform TrajTransformsPresent = GetTransform(SampleTime, RoleIndex, RootBoneReference);
+			const FTransform TrajTransformsFuture = GetTransform(SampleTime + FiniteDelta, RoleIndex, RootBoneReference);
+			const FVector LinearVelocityPresent = (TrajTransformsPresent.GetTranslation() - TrajTransformsPast.GetTranslation()) / FiniteDelta;
+			const FVector LinearVelocityFuture = (TrajTransformsFuture.GetTranslation() - TrajTransformsPresent.GetTranslation()) / FiniteDelta;
+			const FVector LinearAcceleration = (LinearVelocityFuture - LinearVelocityPresent) / FiniteDelta;
 
-						const float Speed = LinearVelocityPresent.Length();
-						const float Acceleration = LinearAcceleration.Length();
+			const float Speed = LinearVelocityPresent.Length();
+			const float Acceleration = LinearAcceleration.Length();
 
-						Stats.AccumulatedSpeed += Speed;
-						Stats.MaxSpeed = FMath::Max(Stats.MaxSpeed, Speed);
+			Stats.AccumulatedSpeed += Speed;
+			Stats.MaxSpeed = FMath::Max(Stats.MaxSpeed, Speed);
 
-						Stats.AccumulatedAcceleration += Acceleration;
-						Stats.MaxAcceleration = FMath::Max(Stats.MaxAcceleration, Acceleration);
+			Stats.AccumulatedAcceleration += Acceleration;
+			Stats.MaxAcceleration = FMath::Max(Stats.MaxAcceleration, Acceleration);
 
-						++Stats.NumAccumulatedSamples;
-					}
-				}
-			}
+			++Stats.NumAccumulatedSamples;
 		}
 	}
 }
 
-void FAssetIndexer::GetSampleInfo(float SampleTime, int32 RoleIndex, FTransform& OutRootTransform, float& OutClipTime, bool& bOutClamped) const
-{
-	const float PlayLength = GetPlayLength();
-	const bool bCanWrap = AssetSamplers.IsLoopable();
-
-	float MainRelativeTime = SampleTime;
-	if (SampleTime < 0.0f && bCanWrap)
-	{
-		// In this case we're sampling a loop backwards, so MainRelativeTime must adjust so the number of cycles is
-		// counted correctly.
-		MainRelativeTime += PlayLength;
-	}
-
-	const FSamplingParam SamplingParam = WrapOrClampSamplingParam(bCanWrap, PlayLength, MainRelativeTime);
-
-	if (FMath::Abs(SamplingParam.Extrapolation) > SMALL_NUMBER)
-	{
-		bOutClamped = true;
-		OutClipTime = SamplingParam.WrappedParam + SamplingParam.Extrapolation;
-		OutRootTransform = AssetSamplers.ExtractRootTransform(OutClipTime, RoleIndex);
-	}
-	else
-	{
-		bOutClamped = false;
-		OutClipTime = SamplingParam.WrappedParam;
-		OutRootTransform = FTransform::Identity;
-
-		// Find the remaining motion deltas after wrapping
-		FTransform RootMotionRemainder = AssetSamplers.ExtractRootTransform(OutClipTime, RoleIndex);
-
-  		const bool bNegativeSampleTime = SampleTime < 0.f;
-		if (SamplingParam.NumCycles > 0 || bNegativeSampleTime)
-		{
-			const FTransform RootMotionLast = AssetSamplers.GetTotalRootTransform(RoleIndex);
-
-			// Determine how to accumulate motion for every cycle of the anim. If the sample
-			// had to be clamped, this motion will end up not getting applied below.
-			// Also invert the accumulation direction if the requested sample was wrapped backwards.
-			FTransform RootMotionPerCycle = RootMotionLast;
-
-			if (bNegativeSampleTime)
-			{
-				RootMotionPerCycle = RootMotionPerCycle.Inverse();
-			}
-			
-			// Invert motion deltas if we wrapped backwards
-			if (bNegativeSampleTime)
-			{
-				RootMotionRemainder.SetToRelativeTransform(RootMotionLast);
-			}
-
-			// Note if the sample was clamped, no motion will be applied here because NumCycles will be zero
-			int32 CyclesRemaining = SamplingParam.NumCycles;
-			while (CyclesRemaining--)
-			{
-				OutRootTransform = RootMotionPerCycle * OutRootTransform;
-			}
-		}
-
-		OutRootTransform = RootMotionRemainder * OutRootTransform;
-	}
-}
-
 FTransform FAssetIndexer::MirrorTransform(const FTransform& Transform, int32 RoleIndex) const
 {
 	return SearchIndexAsset.IsMirrored() ? AssetSamplers.MirrorTransform(Transform, RoleIndex) : Transform;
@@ -470,23 +339,14 @@ FAssetIndexer::FCachedEntry& FAssetIndexer::GetEntry(float SampleTime)
 					BoneContainers[RoleIndex].GetCompactPoseNumBones());
 			}
 
-			FTransform SampleRootTransform;
-			bool bSampleClamped;
-			float CurrentTime;
-
-			GetSampleInfo(SampleTime, RoleIndex, SampleRootTransform, CurrentTime, bSampleClamped);
-
-			if (!bLoopable)
-			{
-				CurrentTime = FMath::Clamp(CurrentTime, 0.f, PlayLength);
-			}
+			const FTransform SampleRootTransform = AssetSamplers.ExtractRootTransform(SampleTime, RoleIndex);
 
 			FMemMark Mark(FMemStack::Get());
 			FCompactPose Pose;
 			FBlendedCurve Curve;
 			Pose.SetBoneContainer(&BoneContainers[RoleIndex]);
 			Curve.InitFrom(BoneContainers[RoleIndex]);
-			AssetSamplers.ExtractPose(CurrentTime, Pose, Curve, RoleIndex);
+			AssetSamplers.ExtractPose(SampleTime, Pose, Curve, RoleIndex);
 
 #if ENABLE_ANIM_DEBUG
 			const int32 NumIterations = GVarMotionMatchTestExtractPoseDeterminismNumIterations;
@@ -494,7 +354,7 @@ FAssetIndexer::FCachedEntry& FAssetIndexer::GetEntry(float SampleTime)
 			{
 				FCompactPose TestPose;
 				TestPose.SetBoneContainer(&BoneContainers[RoleIndex]);
-				AssetSamplers.ExtractPose(CurrentTime, TestPose, RoleIndex);
+				AssetSamplers.ExtractPose(SampleTime, TestPose, RoleIndex);
 
 				const TConstArrayView<FTransform> Bones = Pose.GetBones();
 				const TConstArrayView<FTransform> TestBones = TestPose.GetBones();
@@ -524,7 +384,6 @@ FAssetIndexer::FCachedEntry& FAssetIndexer::GetEntry(float SampleTime)
 			Entry->Curves[RoleIndex].CopyFrom(Curve);
 
 			Entry->RootTransform[RoleIndex] = MirrorTransform(SampleRootTransform, RoleIndex);
-			Entry->bClamped |= bSampleClamped;
 		}
 	}
 
@@ -532,7 +391,7 @@ FAssetIndexer::FCachedEntry& FAssetIndexer::GetEntry(float SampleTime)
 }
 
 // returns the transform in component space for the bone indexed by Schema->BoneReferences[SchemaBoneIdx] at SampleTime seconds
-FTransform FAssetIndexer::GetComponentSpaceTransform(float SampleTime, const FRole& Role, bool& bClamped, int8 SchemaBoneIdx)
+FTransform FAssetIndexer::GetComponentSpaceTransform(float SampleTime, const FRole& Role, int8 SchemaBoneIdx)
 {
 	using namespace UE::PoseSearch;
 
@@ -542,21 +401,17 @@ FTransform FAssetIndexer::GetComponentSpaceTransform(float SampleTime, const FRo
 	}
 
 	FCachedEntry& Entry = GetEntry(SampleTime);
-	bClamped = Entry.bClamped;
-
 	const int32 RoleIndex = RoleToIndex[Role];
 	const FBoneReference& BoneReference = Schema.GetBoneReferences(Role)[SchemaBoneIdx];
 	return CalculateComponentSpaceTransform(Entry, BoneReference, RoleIndex);
 }
 
 // returns the transform in animation space for the bone indexed by Schema->BoneReferences[SchemaBoneIdx] at SampleTime seconds
-FTransform FAssetIndexer::GetTransform(float SampleTime, const FRole& Role, bool& bClamped, int8 SchemaBoneIdx)
+FTransform FAssetIndexer::GetTransform(float SampleTime, const FRole& Role, int8 SchemaBoneIdx)
 {
 	using namespace UE::PoseSearch;
 
 	FCachedEntry& Entry = GetEntry(SampleTime);
-	bClamped = Entry.bClamped;
-
 	const int32 RoleIndex = RoleToIndex[Role];
 	if (SchemaBoneIdx == TrajectorySchemaBoneIdx)
 	{
@@ -568,10 +423,9 @@ FTransform FAssetIndexer::GetTransform(float SampleTime, const FRole& Role, bool
 }
 
 // returns the transform in animation space for the BoneReference at SampleTime seconds
-FTransform FAssetIndexer::GetTransform(float SampleTime, int32 RoleIndex, bool& bClamped, const FBoneReference& BoneReference)
+FTransform FAssetIndexer::GetTransform(float SampleTime, int32 RoleIndex, const FBoneReference& BoneReference)
 {
 	FCachedEntry& Entry = GetEntry(SampleTime);
-	bClamped = Entry.bClamped;
 	return CalculateComponentSpaceTransform(Entry, BoneReference, RoleIndex) * Entry.RootTransform[RoleIndex];
 }
 
@@ -604,7 +458,6 @@ bool FAssetIndexer::GetSampleRotation(FQuat& OutSampleRotation, float SampleTime
 		const FPoseSearchTimedNotifies<UAnimNotifyState_PoseSearchSamplingAttribute>::FItem TimedNotifiesItem = TimedNotifies.GetClosestFutureEvent(SampleTime);
 		if (const UAnimNotifyState_PoseSearchSamplingAttribute* SamplingAttribute = TimedNotifiesItem.NotifyState)
 		{
-			bool bUnused;
 			if (SamplingAttribute->Bone.BoneName != NAME_None)
 			{
 				FBoneReference TempBoneReference = SamplingAttribute->Bone;
@@ -613,8 +466,8 @@ bool FAssetIndexer::GetSampleRotation(FQuat& OutSampleRotation, float SampleTime
 				if (TempBoneReference.HasValidSetup())
 				{
 					const float SamplingAttributeTime = TimedNotifiesItem.Time;
-					const FTransform RootBoneTransform = GetTransform(OriginTime, OriginRole, bUnused, RootSchemaBoneIdx);
-					const FTransform SamplingAttributeBoneTransform = GetTransform(SamplingAttributeTime, SampleRoleIndex, bUnused, TempBoneReference);
+					const FTransform RootBoneTransform = GetTransform(OriginTime, OriginRole, RootSchemaBoneIdx);
+					const FTransform SamplingAttributeBoneTransform = GetTransform(SamplingAttributeTime, SampleRoleIndex, TempBoneReference);
 					OutSampleRotation = RootBoneTransform.InverseTransformRotation(SamplingAttributeBoneTransform.GetRotation());
 					return true;
 				}
@@ -623,19 +476,18 @@ bool FAssetIndexer::GetSampleRotation(FQuat& OutSampleRotation, float SampleTime
 				return false;
 			}
 
-			const FTransform RootBoneTransform = GetTransform(OriginTime, OriginRole, bUnused, RootSchemaBoneIdx);
+			const FTransform RootBoneTransform = GetTransform(OriginTime, OriginRole, RootSchemaBoneIdx);
 			OutSampleRotation = RootBoneTransform.InverseTransformRotation(SamplingAttribute->Rotation);
 			return true;
 		}
 
-		UE_LOG(LogPoseSearch, Error, TEXT("FAssetIndexer::GetSamplePositionInternal: required UAnimNotifyState_PoseSearchSamplingAttribute not found in '%s'"), *AssetSamplers.GetAssetName());
+		UE_LOG(LogPoseSearch, Error, TEXT("FAssetIndexer::GetSampleRotation: required UAnimNotifyState_PoseSearchSamplingAttribute not found in '%s'"), *AssetSamplers.GetAssetName());
 		OutSampleRotation = FQuat::Identity;
 		return false;
 	}
 
-	bool bUnused;
-	const FTransform OriginBoneTransform = GetTransform(OriginTime, OriginRole, bUnused, SchemaOriginBoneIdx);
-	const FTransform SampleBoneTransform = GetTransform(SampleTime, SampleRole, bUnused, SchemaSampleBoneIdx);
+	const FTransform OriginBoneTransform = GetTransform(OriginTime, OriginRole, SchemaOriginBoneIdx);
+	const FTransform SampleBoneTransform = GetTransform(SampleTime, SampleRole, SchemaSampleBoneIdx);
 	OutSampleRotation = OriginBoneTransform.InverseTransformRotation(SampleBoneTransform.GetRotation());
 
 	return true;
@@ -669,11 +521,10 @@ bool FAssetIndexer::GetSamplePosition(FVector& OutSamplePosition, float SampleTi
 	const float SampleTime = Time + SampleTimeOffset + PermutationSampleTimeOffset;
 	const float OriginTime = Time + OriginTimeOffset + PermutationOriginTimeOffset;
 	
-	bool bUnused;
-	return GetSamplePositionInternal(OutSamplePosition, SampleTime, OriginTime, bUnused, SchemaSampleBoneIdx, SchemaOriginBoneIdx, SampleRole, OriginRole, SamplingAttributeId);
+	return GetSamplePositionInternal(OutSamplePosition, SampleTime, OriginTime, SchemaSampleBoneIdx, SchemaOriginBoneIdx, SampleRole, OriginRole, SamplingAttributeId);
 }
 
-bool FAssetIndexer::GetSamplePositionInternal(FVector& OutSamplePosition, float SampleTime, float OriginTime, bool& bClamped, int8 SchemaSampleBoneIdx, int8 SchemaOriginBoneIdx, const FRole& SampleRole, const FRole& OriginRole, int32 SamplingAttributeId)
+bool FAssetIndexer::GetSamplePositionInternal(FVector& OutSamplePosition, float SampleTime, float OriginTime, int8 SchemaSampleBoneIdx, int8 SchemaOriginBoneIdx, const FRole& SampleRole, const FRole& OriginRole, int32 SamplingAttributeId)
 {
 	using namespace UE::PoseSearch;
 
@@ -683,7 +534,6 @@ bool FAssetIndexer::GetSamplePositionInternal(FVector& OutSamplePosition, float
 		const FPoseSearchTimedNotifies<UAnimNotifyState_PoseSearchSamplingAttribute>::FItem TimedNotifiesItem = TimedNotifies.GetClosestFutureEvent(SampleTime);
 		if (const UAnimNotifyState_PoseSearchSamplingAttribute* SamplingAttribute = TimedNotifiesItem.NotifyState)
 		{
-			bool bUnused;
 			if (SamplingAttribute->Bone.BoneName != NAME_None)
 			{
 				FBoneReference TempBoneReference = SamplingAttribute->Bone;
@@ -692,17 +542,15 @@ bool FAssetIndexer::GetSamplePositionInternal(FVector& OutSamplePosition, float
 				if (TempBoneReference.HasValidSetup())
 				{
 					const float SamplingAttributeTime = TimedNotifiesItem.Time;
-					const FTransform RootBoneTransform = GetTransform(OriginTime, OriginRole, bUnused, RootSchemaBoneIdx);
-					const FTransform SamplingAttributeBoneTransform = GetTransform(SamplingAttributeTime, SampleRoleIndex, bClamped, TempBoneReference);
+					const FTransform RootBoneTransform = GetTransform(OriginTime, OriginRole, RootSchemaBoneIdx);
+					const FTransform SamplingAttributeBoneTransform = GetTransform(SamplingAttributeTime, SampleRoleIndex, TempBoneReference);
 					if (SchemaOriginBoneIdx == RootSchemaBoneIdx)
 					{
 						OutSamplePosition = RootBoneTransform.InverseTransformPosition(SamplingAttributeBoneTransform.GetTranslation());
 					}
 					else
 					{
-						bool bOriginClamped;
-						const FTransform OriginBoneTransform = GetTransform(OriginTime, OriginRole, bOriginClamped, SchemaOriginBoneIdx);
-						bClamped |= bOriginClamped;
+						const FTransform OriginBoneTransform = GetTransform(OriginTime, OriginRole, SchemaOriginBoneIdx);
 						const FVector DeltaBoneTranslation = SamplingAttributeBoneTransform.GetTranslation() - OriginBoneTransform.GetTranslation();
 						OutSamplePosition = RootBoneTransform.InverseTransformVector(DeltaBoneTranslation);
 					}
@@ -713,7 +561,7 @@ bool FAssetIndexer::GetSamplePositionInternal(FVector& OutSamplePosition, float
 				return false;
 			}
 
-			const FTransform RootBoneTransform = GetTransform(OriginTime, OriginRole, bUnused, RootSchemaBoneIdx);
+			const FTransform RootBoneTransform = GetTransform(OriginTime, OriginRole, RootSchemaBoneIdx);
 			OutSamplePosition = RootBoneTransform.InverseTransformPosition(SamplingAttribute->Position);
 			return true;
 		}
@@ -723,18 +571,15 @@ bool FAssetIndexer::GetSamplePositionInternal(FVector& OutSamplePosition, float
 		return false;
 	}
 
-	bool bUnused;
-	const FTransform RootBoneTransform = GetTransform(OriginTime, OriginRole, bUnused, RootSchemaBoneIdx);
-	const FTransform SampleBoneTransform = GetTransform(SampleTime, SampleRole, bClamped, SchemaSampleBoneIdx);
+	const FTransform RootBoneTransform = GetTransform(OriginTime, OriginRole, RootSchemaBoneIdx);
+	const FTransform SampleBoneTransform = GetTransform(SampleTime, SampleRole, SchemaSampleBoneIdx);
 	if (SchemaOriginBoneIdx == RootSchemaBoneIdx)
 	{
 		OutSamplePosition = RootBoneTransform.InverseTransformPosition(SampleBoneTransform.GetTranslation());
 		return true;
 	}
 
-	bool bOriginClamped;
-	const FTransform OriginBoneTransform = GetTransform(OriginTime, OriginRole, bOriginClamped, SchemaOriginBoneIdx);
-	bClamped |= bOriginClamped;
+	const FTransform OriginBoneTransform = GetTransform(OriginTime, OriginRole, SchemaOriginBoneIdx);
 	const FVector DeltaBoneTranslation = SampleBoneTransform.GetTranslation() - OriginBoneTransform.GetTranslation();
 	OutSamplePosition = RootBoneTransform.InverseTransformVector(DeltaBoneTranslation);
 	return true;
@@ -756,7 +601,6 @@ bool FAssetIndexer::GetSampleVelocity(FVector& OutSampleVelocity, float SampleTi
 		const FPoseSearchTimedNotifies<UAnimNotifyState_PoseSearchSamplingAttribute>::FItem TimedNotifiesItem = TimedNotifies.GetClosestFutureEvent(SampleTime);
 		if (const UAnimNotifyState_PoseSearchSamplingAttribute* SamplingAttribute = TimedNotifiesItem.NotifyState)
 		{
-			bool bUnused, bClampedPast;
 			FVector BonePositionPast, BonePositionPresent;
 			if (SamplingAttribute->Bone.BoneName != NAME_None)
 			{
@@ -767,21 +611,11 @@ bool FAssetIndexer::GetSampleVelocity(FVector& OutSampleVelocity, float SampleTi
 				{
 					const float SamplingAttributeTime = TimedNotifiesItem.Time;
 
-					if (GetSamplePositionInternal(BonePositionPast, SamplingAttributeTime - FiniteDelta, bUseCharacterSpaceVelocities ? OriginTime - FiniteDelta : OriginTime, bClampedPast, SchemaSampleBoneIdx, SchemaOriginBoneIdx, SampleRole, OriginRole, SamplingAttributeId) &&
-						GetSamplePositionInternal(BonePositionPresent, SamplingAttributeTime, OriginTime, bUnused, SchemaSampleBoneIdx, SchemaOriginBoneIdx, SampleRole, OriginRole, SamplingAttributeId))
+					if (GetSamplePositionInternal(BonePositionPast, SamplingAttributeTime - FiniteDelta, bUseCharacterSpaceVelocities ? OriginTime - FiniteDelta : OriginTime, SchemaSampleBoneIdx, SchemaOriginBoneIdx, SampleRole, OriginRole, SamplingAttributeId) &&
+						GetSamplePositionInternal(BonePositionPresent, SamplingAttributeTime, OriginTime, SchemaSampleBoneIdx, SchemaOriginBoneIdx, SampleRole, OriginRole, SamplingAttributeId))
 					{
-						if (!bClampedPast)
-						{
-							OutSampleVelocity = (BonePositionPresent - BonePositionPast) / FiniteDelta;
-							return true;
-						}
-
-						FVector BonePositionFuture;
-						if (GetSamplePositionInternal(BonePositionFuture, SamplingAttributeTime + FiniteDelta, bUseCharacterSpaceVelocities ? OriginTime + FiniteDelta : OriginTime, bUnused, SchemaSampleBoneIdx, SchemaOriginBoneIdx, SampleRole, OriginRole, SamplingAttributeId))
-						{
-							OutSampleVelocity = (BonePositionFuture - BonePositionPresent) / FiniteDelta;
-							return true;
-						}
+						OutSampleVelocity = (BonePositionPresent - BonePositionPast) / FiniteDelta;
+						return true;
 					}
 
 					return false;
@@ -791,8 +625,8 @@ bool FAssetIndexer::GetSampleVelocity(FVector& OutSampleVelocity, float SampleTi
 				return false;
 			}
 
-			const FTransform RootBoneTransform = GetTransform(OriginTime, OriginRole, bUnused, RootSchemaBoneIdx);
-			OutSampleVelocity = RootBoneTransform.InverseTransformPosition(SamplingAttribute->LinearVelocity);
+			const FTransform RootBoneTransform = GetTransform(OriginTime, OriginRole, RootSchemaBoneIdx);
+			OutSampleVelocity = RootBoneTransform.InverseTransformVector(SamplingAttribute->LinearVelocity);
 			return true;
 		}
 
@@ -801,23 +635,12 @@ bool FAssetIndexer::GetSampleVelocity(FVector& OutSampleVelocity, float SampleTi
 		return false;
 	}
 
-	bool bUnused, bClampedPast;
 	FVector BonePositionPast, BonePositionPresent;
-	if (GetSamplePositionInternal(BonePositionPast, SampleTime - FiniteDelta, bUseCharacterSpaceVelocities ? OriginTime - FiniteDelta : OriginTime, bClampedPast, SchemaSampleBoneIdx, SchemaOriginBoneIdx, SampleRole, OriginRole, INDEX_NONE) &&
-		GetSamplePositionInternal(BonePositionPresent, SampleTime, OriginTime, bUnused, SchemaSampleBoneIdx, SchemaOriginBoneIdx, SampleRole, OriginRole, INDEX_NONE))
+	if (GetSamplePositionInternal(BonePositionPast, SampleTime - FiniteDelta, bUseCharacterSpaceVelocities ? OriginTime - FiniteDelta : OriginTime, SchemaSampleBoneIdx, SchemaOriginBoneIdx, SampleRole, OriginRole, INDEX_NONE) &&
+		GetSamplePositionInternal(BonePositionPresent, SampleTime, OriginTime, SchemaSampleBoneIdx, SchemaOriginBoneIdx, SampleRole, OriginRole, INDEX_NONE))
 	{
-		if (!bClampedPast)
-		{
-			OutSampleVelocity = (BonePositionPresent - BonePositionPast) / FiniteDelta;
-			return true;
-		}
-
-		FVector BonePositionFuture;
-		if (GetSamplePositionInternal(BonePositionFuture, SampleTime + FiniteDelta, bUseCharacterSpaceVelocities ? OriginTime + FiniteDelta : OriginTime, bUnused, SchemaSampleBoneIdx, SchemaOriginBoneIdx, SampleRole, OriginRole, INDEX_NONE))
-		{
-			OutSampleVelocity = (BonePositionFuture - BonePositionPresent) / FiniteDelta;
-			return true;
-		}
+		OutSampleVelocity = (BonePositionPresent - BonePositionPast) / FiniteDelta;
+		return true;
 	}
 
 	OutSampleVelocity = FVector::ZeroVector;
@@ -904,11 +727,6 @@ void FAssetIndexer::CompareCachedEntries(const FAssetIndexer& Other) const
 					UE_LOG(LogPoseSearch, Warning, TEXT("CompareCachedEntries - FAssetIndexer::CachedEntries::SampleTime is not deterministic (%f, %f)"), Entry->SampleTime, OtherEntry->SampleTime);
 				}
 
-				if (Entry->bClamped != OtherEntry->bClamped)
-				{
-					UE_LOG(LogPoseSearch, Warning, TEXT("CompareCachedEntries - FAssetIndexer::CachedEntries::bClamped is not deterministic"));
-				}
-
 				if (Entry->RootTransform.Num() != OtherEntry->RootTransform.Num())
 				{
 					UE_LOG(LogPoseSearch, Warning, TEXT("CompareCachedEntries - FAssetIndexer::CachedEntries::RootTransform::Num is not deterministic"));
diff --git a/Engine/Plugins/Animation/PoseSearch/Source/Runtime/Private/PoseSearchDerivedData.cpp b/Engine/Plugins/Animation/PoseSearch/Source/Runtime/Private/PoseSearchDerivedData.cpp
--- a/Engine/Plugins/Animation/PoseSearch/Source/Runtime/Private/PoseSearchDerivedData.cpp
+++ b/Engine/Plugins/Animation/PoseSearch/Source/Runtime/Private/PoseSearchDerivedData.cpp
@@ -2773,6 +2773,7 @@ EAsyncBuildIndexResult FAsyncPoseSearchDatabasesManagement::RequestAsyncBuildInd
 	if (Database->GetPackage()->HasAnyPackageFlags(PKG_Cooked))
 	{
 		// Don't cache for cooked packages
+		check(!Database->GetSearchIndex().IsEmpty());
 		return EAsyncBuildIndexResult::Success;
 	}
 
@@ -2821,6 +2822,7 @@ EAsyncBuildIndexResult FAsyncPoseSearchDatabasesManagement::RequestAsyncBuildInd
 
 	if (Task->GetState() == FPoseSearchDatabaseAsyncCacheTask::EState::Ended)
 	{
+		check(!Database->GetSearchIndex().IsEmpty());
 		return EAsyncBuildIndexResult::Success;
 	}
 
diff --git a/Engine/Plugins/Animation/PoseSearch/Source/Runtime/Private/PoseSearchFeatureChannel_Distance.cpp b/Engine/Plugins/Animation/PoseSearch/Source/Runtime/Private/PoseSearchFeatureChannel_Distance.cpp
--- a/Engine/Plugins/Animation/PoseSearch/Source/Runtime/Private/PoseSearchFeatureChannel_Distance.cpp
+++ b/Engine/Plugins/Animation/PoseSearch/Source/Runtime/Private/PoseSearchFeatureChannel_Distance.cpp
@@ -42,15 +42,8 @@ void UPoseSearchFeatureChannel_Distance::AddDependentChannels(UPoseSearchSchema*
 	{
 		if (SchemaOriginBoneIdx != RootSchemaBoneIdx || PermutationTimeType == EPermutationTimeType::UsePermutationTime)
 		{
-			if (bDefaultWithRootBone)
-			{
-				const EPermutationTimeType DependentChannelsPermutationTimeType = PermutationTimeType == EPermutationTimeType::UsePermutationTime ? EPermutationTimeType::UseSampleToPermutationTime : EPermutationTimeType::UseSampleTime;
-				UPoseSearchFeatureChannel_Position::FindOrAddToSchema(Schema, 0.f, OriginBone.BoneName, OriginRole, DependentChannelsPermutationTimeType);
-			}
-			else
-			{
-				// @todo: add bInjectAdditionalDebugChannels support for !bDefaultWithRootBone UPoseSearchFeatureChannel_Distance
-			}
+			const EPermutationTimeType DependentChannelsPermutationTimeType = PermutationTimeType == EPermutationTimeType::UsePermutationTime ? EPermutationTimeType::UseSampleToPermutationTime : EPermutationTimeType::UseSampleTime;
+			UPoseSearchFeatureChannel_Position::FindOrAddToSchema(Schema, 0.f, OriginBone.BoneName, OriginRole, DependentChannelsPermutationTimeType, bDefaultWithRootBone);
 		}
 	}
 }
diff --git a/Engine/Plugins/Animation/PoseSearch/Source/Runtime/Private/PoseSearchFeatureChannel_Heading.cpp b/Engine/Plugins/Animation/PoseSearch/Source/Runtime/Private/PoseSearchFeatureChannel_Heading.cpp
--- a/Engine/Plugins/Animation/PoseSearch/Source/Runtime/Private/PoseSearchFeatureChannel_Heading.cpp
+++ b/Engine/Plugins/Animation/PoseSearch/Source/Runtime/Private/PoseSearchFeatureChannel_Heading.cpp
@@ -20,22 +20,31 @@ UPoseSearchFeatureChannel_Heading::UPoseSearchFeatureChannel_Heading()
 	bUseBlueprintQueryOverride = Cast<UBlueprintGeneratedClass>(GetClass()) != nullptr;
 }
 
-void UPoseSearchFeatureChannel_Heading::FindOrAddToSchema(UPoseSearchSchema* Schema, float SampleTimeOffset, const FName& BoneName, const UE::PoseSearch::FRole& Role, EHeadingAxis HeadingAxis, EPermutationTimeType PermutationTimeType)
+void UPoseSearchFeatureChannel_Heading::FindOrAddToSchema(UPoseSearchSchema* Schema, float SampleTimeOffset, const FName& BoneName, const UE::PoseSearch::FRole& Role, EHeadingAxis HeadingAxis, EPermutationTimeType PermutationTimeType, bool bDefaultWithRootBone)
 {
-	if (!Schema->FindChannel([SampleTimeOffset, &BoneName, Role, HeadingAxis, PermutationTimeType](const UPoseSearchFeatureChannel* Channel) -> const UPoseSearchFeatureChannel_Heading*
+	using namespace UE::PoseSearch;
+
+	check(Schema);
+	const int8 SchemaBoneIdx = Schema->FindSchemaBoneIdx(BoneName, Role, bDefaultWithRootBone);
+	if (SchemaBoneIdx == InvalidSchemaBoneIdx)
+	{
+		UE_LOG(LogPoseSearch, Error, TEXT("UPoseSearchFeatureChannel_Heading::FindOrAddToSchema: couldn't find SchemaBoneIdx for Bone '%s' with Role '%s' in UPoseSearchSchema '%s'"), *BoneName.ToString(), *Role.ToString(), *GetNameSafe(Schema));
+		return;
+	}
+
+	if (!Schema->FindChannel([SchemaBoneIdx, SampleTimeOffset, Role, HeadingAxis, PermutationTimeType](const UPoseSearchFeatureChannel* Channel) -> const UPoseSearchFeatureChannel_Heading*
 		{
 			if (const UPoseSearchFeatureChannel_Heading* Heading = Cast<UPoseSearchFeatureChannel_Heading>(Channel))
 			{
-				// @todo: channels are already finalized, so we can use SchemaBoneIdx and SchemaOriginBoneIdx instead of Bone.BoneName and OriginBone.BoneName
-				if (Heading->Bone.BoneName == BoneName &&
-					Heading->OriginBone.BoneName == NAME_None && 
+				if (Heading->SchemaBoneIdx == SchemaBoneIdx &&
+					Heading->SchemaOriginBoneIdx == RootSchemaBoneIdx &&
 					Heading->SampleTimeOffset == SampleTimeOffset &&
 					Heading->OriginTimeOffset == 0.f &&
 					Heading->HeadingAxis == HeadingAxis &&
 					Heading->PermutationTimeType == PermutationTimeType &&
 					Heading->SampleRole == Role &&
 					Heading->OriginRole == Role &&
-					Heading->bDefaultWithRootBone)
+					!Heading->bDefaultWithRootBone)
 				{
 					return Heading;
 				}
@@ -43,9 +52,12 @@ void UPoseSearchFeatureChannel_Heading::FindOrAddToSchema(UPoseSearchSchema* Sch
 			return nullptr;
 		}))
 	{
+		const TConstArrayView<FBoneReference> BoneReferences = Schema->GetBoneReferences(Role);
+
 		UPoseSearchFeatureChannel_Heading* Heading = NewObject<UPoseSearchFeatureChannel_Heading>(Schema, NAME_None, RF_Transient);
-		Heading->Bone.BoneName = BoneName;
+		Heading->Bone.BoneName = BoneReferences.IsValidIndex(SchemaBoneIdx) ? BoneReferences[SchemaBoneIdx].BoneName : NAME_None;
 		Heading->SampleRole = Role;
+		Heading->OriginBone.BoneName = BoneReferences.IsValidIndex(RootSchemaBoneIdx) ? BoneReferences[RootSchemaBoneIdx].BoneName : NAME_None;
 		Heading->OriginRole = Role;
 #if WITH_EDITORONLY_DATA
 		Heading->Weight = 0.f;
@@ -54,6 +66,8 @@ void UPoseSearchFeatureChannel_Heading::FindOrAddToSchema(UPoseSearchSchema* Sch
 		Heading->SampleTimeOffset = SampleTimeOffset;
 		Heading->HeadingAxis = HeadingAxis;
 		Heading->PermutationTimeType = PermutationTimeType;
+		Heading->bDefaultWithRootBone = false;
+
 		Schema->AddTemporaryChannel(Heading);
 	}
 }
@@ -77,12 +91,12 @@ void UPoseSearchFeatureChannel_Heading::AddDependentChannels(UPoseSearchSchema*
 	if (Schema->bInjectAdditionalDebugChannels)
 	{
 		const EPermutationTimeType SamplePermutationTimeType = PermutationTimeType != EPermutationTimeType::UseSampleTime ? EPermutationTimeType::UseSampleToPermutationTime : EPermutationTimeType::UseSampleTime;
-		UPoseSearchFeatureChannel_Position::FindOrAddToSchema(Schema, SampleTimeOffset, Bone.BoneName, SampleRole, SamplePermutationTimeType);
+		UPoseSearchFeatureChannel_Position::FindOrAddToSchema(Schema, SampleTimeOffset, Bone.BoneName, SampleRole, SamplePermutationTimeType, bDefaultWithRootBone);
 		
-		// injecting 2 Heading channels to be able to reconstruct the Origin bone rotation
+		// injecting 2 Heading channels to be able to reconstruct the Bone rotation
 		const EPermutationTimeType OriginPermutationTimeType = PermutationTimeType == EPermutationTimeType::UsePermutationTime ? EPermutationTimeType::UseSampleToPermutationTime : EPermutationTimeType::UseSampleTime;
-		UPoseSearchFeatureChannel_Heading::FindOrAddToSchema(Schema, OriginTimeOffset, Bone.BoneName, OriginRole, EHeadingAxis::X, OriginPermutationTimeType);
-		UPoseSearchFeatureChannel_Heading::FindOrAddToSchema(Schema, OriginTimeOffset, Bone.BoneName, OriginRole, EHeadingAxis::Y, OriginPermutationTimeType);
+		UPoseSearchFeatureChannel_Heading::FindOrAddToSchema(Schema, OriginTimeOffset, Bone.BoneName, OriginRole, EHeadingAxis::X, OriginPermutationTimeType, bDefaultWithRootBone);
+		UPoseSearchFeatureChannel_Heading::FindOrAddToSchema(Schema, OriginTimeOffset, Bone.BoneName, OriginRole, EHeadingAxis::Y, OriginPermutationTimeType, bDefaultWithRootBone);
 	}
 }
 
diff --git a/Engine/Plugins/Animation/PoseSearch/Source/Runtime/Private/PoseSearchFeatureChannel_Position.cpp b/Engine/Plugins/Animation/PoseSearch/Source/Runtime/Private/PoseSearchFeatureChannel_Position.cpp
--- a/Engine/Plugins/Animation/PoseSearch/Source/Runtime/Private/PoseSearchFeatureChannel_Position.cpp
+++ b/Engine/Plugins/Animation/PoseSearch/Source/Runtime/Private/PoseSearchFeatureChannel_Position.cpp
@@ -19,21 +19,30 @@ UPoseSearchFeatureChannel_Position::UPoseSearchFeatureChannel_Position()
 	bUseBlueprintQueryOverride = Cast<UBlueprintGeneratedClass>(GetClass()) != nullptr;
 }
 
-void UPoseSearchFeatureChannel_Position::FindOrAddToSchema(UPoseSearchSchema* Schema, float SampleTimeOffset, const FName& BoneName, const UE::PoseSearch::FRole& Role, EPermutationTimeType PermutationTimeType)
+void UPoseSearchFeatureChannel_Position::FindOrAddToSchema(UPoseSearchSchema* Schema, float SampleTimeOffset, const FName& BoneName, const UE::PoseSearch::FRole& Role, EPermutationTimeType PermutationTimeType, bool bDefaultWithRootBone)
 {
-	if (!Schema->FindChannel([SampleTimeOffset, &BoneName, &Role, PermutationTimeType](const UPoseSearchFeatureChannel* Channel) -> const UPoseSearchFeatureChannel_Position*
+	using namespace UE::PoseSearch;
+
+	check(Schema);
+	const int8 SchemaBoneIdx = Schema->FindSchemaBoneIdx(BoneName, Role, bDefaultWithRootBone);
+	if (SchemaBoneIdx == InvalidSchemaBoneIdx)
+	{
+		UE_LOG(LogPoseSearch, Error, TEXT("UPoseSearchFeatureChannel_Position::FindOrAddToSchema: couldn't find SchemaBoneIdx for Bone '%s' with Role '%s' in UPoseSearchSchema '%s'"), *BoneName.ToString(), *Role.ToString(), *GetNameSafe(Schema));
+		return;
+	}
+
+	if (!Schema->FindChannel([SchemaBoneIdx, SampleTimeOffset, &Role, PermutationTimeType](const UPoseSearchFeatureChannel* Channel) -> const UPoseSearchFeatureChannel_Position*
 		{
 			if (const UPoseSearchFeatureChannel_Position* Position = Cast<UPoseSearchFeatureChannel_Position>(Channel))
 			{
-				// @todo: channels are already finalized, so we can use SchemaBoneIdx and SchemaOriginBoneIdx instead of Bone.BoneName and OriginBone.BoneName
-				if (Position->Bone.BoneName == BoneName &&
-					Position->OriginBone.BoneName == NAME_None &&
+				if (Position->SchemaBoneIdx == SchemaBoneIdx &&
+					Position->SchemaOriginBoneIdx == RootSchemaBoneIdx &&
 					Position->SampleTimeOffset == SampleTimeOffset &&
 					Position->OriginTimeOffset == 0.f &&
 					Position->PermutationTimeType == PermutationTimeType &&
 					Position->SampleRole == Role &&
 					Position->OriginRole == Role &&
-					Position->bDefaultWithRootBone && 
+					!Position->bDefaultWithRootBone &&
 					!Position->bNormalizeDisplacement)
 				{
 					return Position;
@@ -42,9 +51,12 @@ void UPoseSearchFeatureChannel_Position::FindOrAddToSchema(UPoseSearchSchema* Sc
 			return nullptr;
 		}))
 	{
+		const TConstArrayView<FBoneReference> BoneReferences = Schema->GetBoneReferences(Role);
+
 		UPoseSearchFeatureChannel_Position* Position = NewObject<UPoseSearchFeatureChannel_Position>(Schema, NAME_None, RF_Transient);
-		Position->Bone.BoneName = BoneName;
+		Position->Bone.BoneName = BoneReferences.IsValidIndex(SchemaBoneIdx) ? BoneReferences[SchemaBoneIdx].BoneName : NAME_None;
 		Position->SampleRole = Role;
+		Position->OriginBone.BoneName = BoneReferences.IsValidIndex(RootSchemaBoneIdx) ? BoneReferences[RootSchemaBoneIdx].BoneName : NAME_None;
 		Position->OriginRole = Role;
 #if WITH_EDITORONLY_DATA
 		Position->Weight = 0.f;
@@ -52,6 +64,8 @@ void UPoseSearchFeatureChannel_Position::FindOrAddToSchema(UPoseSearchSchema* Sc
 #endif // WITH_EDITORONLY_DATA
 		Position->SampleTimeOffset = SampleTimeOffset;
 		Position->PermutationTimeType = PermutationTimeType;
+		Position->bDefaultWithRootBone = false;
+
 		Schema->AddTemporaryChannel(Position);
 	}
 }
@@ -74,20 +88,10 @@ void UPoseSearchFeatureChannel_Position::AddDependentChannels(UPoseSearchSchema*
 {
 	using namespace UE::PoseSearch;
 
-	if (Schema->bInjectAdditionalDebugChannels)
+	if (Schema->bInjectAdditionalDebugChannels && (SchemaOriginBoneIdx != RootSchemaBoneIdx || PermutationTimeType == EPermutationTimeType::UsePermutationTime || OriginTimeOffset != 0.f))
 	{
-		if (SchemaOriginBoneIdx != RootSchemaBoneIdx || PermutationTimeType == EPermutationTimeType::UsePermutationTime)
-		{
-			if (bDefaultWithRootBone)
-			{
-				const EPermutationTimeType DependentChannelsPermutationTimeType = PermutationTimeType == EPermutationTimeType::UsePermutationTime ? EPermutationTimeType::UseSampleToPermutationTime : EPermutationTimeType::UseSampleTime;
-				UPoseSearchFeatureChannel_Position::FindOrAddToSchema(Schema, 0.f, OriginBone.BoneName, OriginRole, DependentChannelsPermutationTimeType);
-			}
-			else
-			{
-				// @todo: add bInjectAdditionalDebugChannels support for !bDefaultWithRootBone UPoseSearchFeatureChannel_Position
-			}
-		}
+		const EPermutationTimeType DependentChannelsPermutationTimeType = PermutationTimeType == EPermutationTimeType::UsePermutationTime ? EPermutationTimeType::UseSampleToPermutationTime : EPermutationTimeType::UseSampleTime;
+		UPoseSearchFeatureChannel_Position::FindOrAddToSchema(Schema, OriginTimeOffset, OriginBone.BoneName, OriginRole, DependentChannelsPermutationTimeType, bDefaultWithRootBone);
 	}
 }
 
diff --git a/Engine/Plugins/Animation/PoseSearch/Source/Runtime/Private/PoseSearchFeatureChannel_Velocity.cpp b/Engine/Plugins/Animation/PoseSearch/Source/Runtime/Private/PoseSearchFeatureChannel_Velocity.cpp
--- a/Engine/Plugins/Animation/PoseSearch/Source/Runtime/Private/PoseSearchFeatureChannel_Velocity.cpp
+++ b/Engine/Plugins/Animation/PoseSearch/Source/Runtime/Private/PoseSearchFeatureChannel_Velocity.cpp
@@ -40,7 +40,7 @@ void UPoseSearchFeatureChannel_Velocity::AddDependentChannels(UPoseSearchSchema*
 	if (Schema->bInjectAdditionalDebugChannels)
 	{
 		const EPermutationTimeType DependentChannelsPermutationTimeType = PermutationTimeType != EPermutationTimeType::UseSampleTime ? EPermutationTimeType::UseSampleToPermutationTime : EPermutationTimeType::UseSampleTime;
-		UPoseSearchFeatureChannel_Position::FindOrAddToSchema(Schema, SampleTimeOffset, Bone.BoneName, SampleRole, DependentChannelsPermutationTimeType);
+		UPoseSearchFeatureChannel_Position::FindOrAddToSchema(Schema, SampleTimeOffset, Bone.BoneName, SampleRole, DependentChannelsPermutationTimeType, bDefaultWithRootBone);
 	}
 }
 
diff --git a/Engine/Plugins/Animation/PoseSearch/Source/Runtime/Private/PoseSearchSchema.cpp b/Engine/Plugins/Animation/PoseSearch/Source/Runtime/Private/PoseSearchSchema.cpp
--- a/Engine/Plugins/Animation/PoseSearch/Source/Runtime/Private/PoseSearchSchema.cpp
+++ b/Engine/Plugins/Animation/PoseSearch/Source/Runtime/Private/PoseSearchSchema.cpp
@@ -146,6 +146,46 @@ int8 UPoseSearchSchema::AddBoneReference(const FBoneReference& BoneReference, co
 	return AddBoneReference(BoneReference, Role, true);
 }
 
+
+int8 UPoseSearchSchema::FindSchemaBoneIdx(const FName& BoneName, const UE::PoseSearch::FRole& Role, bool bDefaultWithRootBone) const
+{
+	using namespace UE::PoseSearch;
+
+	const FPoseSearchRoledSkeleton* RoledSkeleton = GetRoledSkeleton(Role);
+	if (!RoledSkeleton)
+	{
+		UE_LOG(LogPoseSearch, Error, TEXT("UPoseSearchSchema::FindSchemaBoneIdx: couldn't find data for the requested Role '%s' in UPoseSearchSchema '%s'"), *Role.ToString(), *GetNameSafe(this));
+		return InvalidSchemaBoneIdx;
+	}
+
+	if (BoneName == NAME_None)
+	{
+		if (bDefaultWithRootBone)
+		{
+			if (RoledSkeleton->BoneReferences.IsEmpty())
+			{
+				return InvalidSchemaBoneIdx;
+			}
+
+			return RootSchemaBoneIdx;
+		}
+		
+		return TrajectorySchemaBoneIdx;
+	}
+
+	for (int32 SchemaBoneIdx = 0; SchemaBoneIdx < RoledSkeleton->BoneReferences.Num(); ++SchemaBoneIdx)
+	{
+		const FBoneReference& BoneReference = RoledSkeleton->BoneReferences[SchemaBoneIdx];
+		if (BoneReference.BoneName == BoneName)
+		{
+			check(SchemaBoneIdx < 128);
+			return int8(SchemaBoneIdx);
+		}
+	}
+
+	return InvalidSchemaBoneIdx;
+}
+
 int8 UPoseSearchSchema::AddBoneReference(const FBoneReference& BoneReference, const UE::PoseSearch::FRole& Role, bool bDefaultWithRootBone)
 {
 	using namespace UE::PoseSearch;
diff --git a/Engine/Plugins/Animation/PoseSearch/Source/Runtime/Public/PoseSearch/PoseSearchAssetIndexer.h b/Engine/Plugins/Animation/PoseSearch/Source/Runtime/Public/PoseSearch/PoseSearchAssetIndexer.h
--- a/Engine/Plugins/Animation/PoseSearch/Source/Runtime/Public/PoseSearch/PoseSearchAssetIndexer.h
+++ b/Engine/Plugins/Animation/PoseSearch/Source/Runtime/Public/PoseSearch/PoseSearchAssetIndexer.h
@@ -112,24 +112,21 @@ public:
 #endif // ENABLE_ANIM_DEBUG
 
 private:
-	POSESEARCH_API int32 GetVectorIdx(int32 SampleIdx) const;
+	int32 GetVectorIdx(int32 SampleIdx) const;
 
 	// Returns the animation (world) space transform of the bone Schema.BoneReferences[SchemaBoneIdx] at time SampleTime
-	// bClamped will be true if SampleTime is outside the animation duration boundaries
-	POSESEARCH_API FTransform GetTransform(float SampleTime, const FRole& Role, bool& bClamped, int8 SchemaBoneIdx = RootSchemaBoneIdx);
-	POSESEARCH_API FTransform GetTransform(float SampleTime, int32 RoleIndex, bool& bClamped, const FBoneReference& BoneReference);
+	FTransform GetTransform(float SampleTime, const FRole& Role, int8 SchemaBoneIdx = RootSchemaBoneIdx);
+	FTransform GetTransform(float SampleTime, int32 RoleIndex, const FBoneReference& BoneReference);
 	
 	// Returns the value of float curve CurveName at time SampleTime
-	POSESEARCH_API float GetSampleCurveValueInternal(float SampleTime, const FName& CurveName, const FRole& Role);
+	float GetSampleCurveValueInternal(float SampleTime, const FName& CurveName, const FRole& Role);
 
 	// Returns the component space transform of the bone Schema.BoneReferences[SchemaBoneIdx] at time SampleTime
-	// bClamped will be true if SampleTime is outside the animation duration boundaries
-	POSESEARCH_API FTransform GetComponentSpaceTransform(float SampleTime, const FRole& Role, bool& bClamped, int8 SchemaBoneIdx = RootSchemaBoneIdx);
+	FTransform GetComponentSpaceTransform(float SampleTime, const FRole& Role, int8 SchemaBoneIdx = RootSchemaBoneIdx);
 
 	// Returns OutSamplePosition as the displacement of the bone Schema.BoneReferences[SchemaSampleBoneIdx] position at time SampleTime relative to the
 	// position of the bone Schema.BoneReferences[SchemaOriginBoneIdx] at time OriginTime in root bone space
-	// bClamped will be true if SampleTime or OriginTime are outside the animation duration boundaries
-	POSESEARCH_API bool GetSamplePositionInternal(FVector& OutSamplePosition, float SampleTime, float OriginTime, bool& bClamped, int8 SchemaSampleBoneIdx, int8 SchemaOriginBoneIdx, const FRole& SampleRole, const FRole& OriginRole, int32 SamplingAttributeId);
+	bool GetSamplePositionInternal(FVector& OutSamplePosition, float SampleTime, float OriginTime, int8 SchemaSampleBoneIdx, int8 SchemaOriginBoneIdx, const FRole& SampleRole, const FRole& OriginRole, int32 SamplingAttributeId);
 
 	struct FCachedCSPose : public FCSPose<FCompactHeapPose>
 	{
@@ -145,7 +142,6 @@ private:
 	struct FCachedEntry
 	{
 		float SampleTime = 0.f;
-		bool bClamped = false;
 
 		// RootTransform and ComponentSpacePose are stored mirrored in case SearchIndexAsset.IsMirrored
 		TArray<FTransform, TInlineAllocator<PreallocatedRolesNum>> RootTransform;
@@ -153,12 +149,11 @@ private:
 		TArray<FBlendedHeapCurve, TInlineAllocator<PreallocatedRolesNum>> Curves;
 	};
 
-	POSESEARCH_API void GetSampleInfo(float SampleTime, int32 RoleIndex, FTransform& OutRootTransform, float& OutClipTime, bool& bOutClamped) const;
-	POSESEARCH_API FTransform MirrorTransform(const FTransform& Transform, int32 RoleIndex) const;
-	POSESEARCH_API FCachedEntry& GetEntry(float SampleTime);
-	POSESEARCH_API FTransform CalculateComponentSpaceTransform(FCachedEntry& Entry, const FBoneReference& BoneReference, int32 RoleIndex);
+	FTransform MirrorTransform(const FTransform& Transform, int32 RoleIndex) const;
+	FCachedEntry& GetEntry(float SampleTime);
+	FTransform CalculateComponentSpaceTransform(FCachedEntry& Entry, const FBoneReference& BoneReference, int32 RoleIndex);
 
-	POSESEARCH_API void ComputeStats();
+	void ComputeStats();
 
 	TArray<FBoneContainer> BoneContainers;
 	TMap<float, FCachedEntry> CachedEntries;
diff --git a/Engine/Plugins/Animation/PoseSearch/Source/Runtime/Public/PoseSearch/PoseSearchDefines.h b/Engine/Plugins/Animation/PoseSearch/Source/Runtime/Public/PoseSearch/PoseSearchDefines.h
--- a/Engine/Plugins/Animation/PoseSearch/Source/Runtime/Public/PoseSearch/PoseSearchDefines.h
+++ b/Engine/Plugins/Animation/PoseSearch/Source/Runtime/Public/PoseSearch/PoseSearchDefines.h
@@ -43,7 +43,7 @@ static constexpr int32 PreallocatedSearchesNum = 4;
 static constexpr float FiniteDelta = 1 / 60.0f; // Time delta used for computing pose derivatives
 static constexpr int32 MaxNumberOfCollectedPoseCandidatesPerDatabase = 256;
 static constexpr int32 InvalidSearchId = INDEX_NONE; // @todo: it should be 0, since the search id is an hash!
-static constexpr int32 DatabaseIndexDerivedDataCacheKeyVersion = 2723; // bump this value to force reindex ALL the UPoseSearchDatabases, without having to change the global FDevSystemGuids::POSESEARCHDB_DERIVEDDATA_VER
+static constexpr int32 DatabaseIndexDerivedDataCacheKeyVersion = 2727; // bump this value to force reindex ALL the UPoseSearchDatabases, without having to change the global FDevSystemGuids::POSESEARCHDB_DERIVEDDATA_VER
 static constexpr int32 DefaultTickPriority = 100;
 
 template<uint32 Alignment = DEFAULT_ALIGNMENT> class TMemStackSetAllocator : public TSetAllocator<TSparseArrayAllocator<TMemStackAllocator<Alignment>, TMemStackAllocator<Alignment>>, TMemStackAllocator<Alignment>> {};
diff --git a/Engine/Plugins/Animation/PoseSearch/Source/Runtime/Public/PoseSearch/PoseSearchFeatureChannel_Heading.h b/Engine/Plugins/Animation/PoseSearch/Source/Runtime/Public/PoseSearch/PoseSearchFeatureChannel_Heading.h
--- a/Engine/Plugins/Animation/PoseSearch/Source/Runtime/Public/PoseSearch/PoseSearchFeatureChannel_Heading.h
+++ b/Engine/Plugins/Animation/PoseSearch/Source/Runtime/Public/PoseSearch/PoseSearchFeatureChannel_Heading.h
@@ -129,7 +129,7 @@ public:
 #endif // WITH_EDITORONLY_DATA
 
 	UE_API FVector GetAxis(const FQuat& Rotation) const;
-	static UE_API void FindOrAddToSchema(UPoseSearchSchema* Schema, float SampleTimeOffset, const FName& BoneName, const UE::PoseSearch::FRole& Role, EHeadingAxis HeadingAxis = EHeadingAxis::X, EPermutationTimeType PermutationTimeType = EPermutationTimeType::UseSampleTime);
+	static UE_API void FindOrAddToSchema(UPoseSearchSchema* Schema, float SampleTimeOffset, const FName& BoneName, const UE::PoseSearch::FRole& Role, EHeadingAxis HeadingAxis = EHeadingAxis::X, EPermutationTimeType PermutationTimeType = EPermutationTimeType::UseSampleTime, bool bDefaultWithRootBone = true);
 };
 
 #undef UE_API
diff --git a/Engine/Plugins/Animation/PoseSearch/Source/Runtime/Public/PoseSearch/PoseSearchFeatureChannel_Position.h b/Engine/Plugins/Animation/PoseSearch/Source/Runtime/Public/PoseSearch/PoseSearchFeatureChannel_Position.h
--- a/Engine/Plugins/Animation/PoseSearch/Source/Runtime/Public/PoseSearch/PoseSearchFeatureChannel_Position.h
+++ b/Engine/Plugins/Animation/PoseSearch/Source/Runtime/Public/PoseSearch/PoseSearchFeatureChannel_Position.h
@@ -123,7 +123,7 @@ public:
 	virtual FLinearColor GetDebugColor() const override { return DebugColor; }
 #endif // WITH_EDITORONLY_DATA
 
-	static UE_API void FindOrAddToSchema(UPoseSearchSchema* Schema, float SampleTimeOffset, const FName& BoneName, const UE::PoseSearch::FRole& Role, EPermutationTimeType PermutationTimeType = EPermutationTimeType::UseSampleTime);
+	static UE_API void FindOrAddToSchema(UPoseSearchSchema* Schema, float SampleTimeOffset, const FName& BoneName, const UE::PoseSearch::FRole& Role, EPermutationTimeType PermutationTimeType = EPermutationTimeType::UseSampleTime, bool bDefaultWithRootBone = true);
 };
 
 #undef UE_API
diff --git a/Engine/Plugins/Animation/PoseSearch/Source/Runtime/Public/PoseSearch/PoseSearchSchema.h b/Engine/Plugins/Animation/PoseSearch/Source/Runtime/Public/PoseSearch/PoseSearchSchema.h
--- a/Engine/Plugins/Animation/PoseSearch/Source/Runtime/Public/PoseSearch/PoseSearchSchema.h
+++ b/Engine/Plugins/Animation/PoseSearch/Source/Runtime/Public/PoseSearch/PoseSearchSchema.h
@@ -157,6 +157,9 @@ public:
 	UE_API int8 AddBoneReference(const FBoneReference& BoneReference, const UE::PoseSearch::FRole& Role);
 	UE_API int8 AddCurveReference(const FName& CurveReference, const UE::PoseSearch::FRole& Role);
 
+	// Experimental, this feature might be removed without warning, not for production use
+	int8 FindSchemaBoneIdx(const FName& BoneName, const UE::PoseSearch::FRole& Role, bool bDefaultWithRootBone) const;
+
 #if WITH_EDITOR
 	UE_API virtual void PostEditChangeProperty(struct FPropertyChangedEvent& PropertyChangedEvent) override;
 #endif // WITH_EDITOR
