PoseSearch - support for opening PoseSearchInteractionAssetEditor at the correct asset time when clicking the hyper link from the pose search debugger FAssetName column

diff --git a/Engine/Plugins/Animation/PoseSearch/Source/Editor/Private/PoseSearchDebuggerDatabaseColumns.h b/Engine/Plugins/Animation/PoseSearch/Source/Editor/Private/PoseSearchDebuggerDatabaseColumns.h
--- a/Engine/Plugins/Animation/PoseSearch/Source/Editor/Private/PoseSearchDebuggerDatabaseColumns.h
+++ b/Engine/Plugins/Animation/PoseSearch/Source/Editor/Private/PoseSearchDebuggerDatabaseColumns.h
@@ -2,17 +2,18 @@
 
 #pragma once
 
-#include "AnimPreviewInstance.h"
 #include "Animation/DebugSkelMeshComponent.h"
+#include "AnimPreviewInstance.h"
 #include "Editor.h"
+#include "Fonts/FontMeasure.h"
+#include "Framework/Application/SlateApplication.h"
 #include "IAnimationEditor.h"
+#include "IMultiAnimAssetEditor.h"
 #include "IPersonaToolkit.h"
 #include "PoseSearchDatabaseEditor.h"
 #include "PoseSearchDatabaseViewModel.h"
 #include "PoseSearchDebuggerDatabaseRowData.h"
 #include "PoseSearchDebuggerViewModel.h"
-#include "Fonts/FontMeasure.h"
-#include "Framework/Application/SlateApplication.h"
 #include "Preferences/PersonaOptions.h"
 #include "Styling/AppStyle.h"
 #include "Subsystems/AssetEditorSubsystem.h"
@@ -345,6 +346,13 @@ struct FAssetName : IColumn
 										PreviewComponent->PreviewInstance->SetPlaying(false);
 										PreviewComponent->PreviewInstance->SetBlendSpacePosition(RowData->BlendParameters);
 									}
+									else if (Editor->GetEditorName() == "PoseSearchInteractionAssetEditor")
+									{
+										IMultiAnimAssetEditor* MultiAnimAssetEditor = static_cast<IMultiAnimAssetEditor*>(Editor);
+
+										// Open asset paused and at specific time as seen on the pose search debugger.
+										MultiAnimAssetEditor->SetPreviewProperties(RowData->AssetTime, RowData->BlendParameters, false);
+									}
 								}
 							}
 						}
