PoseSearch - fix for crash in FConstraints::ProcessAll when one PSIA animation is null

#rnx

diff --git a/Engine/Plugins/Animation/PoseSearch/Source/Runtime/Private/PoseSearchConstraints.cpp b/Engine/Plugins/Animation/PoseSearch/Source/Runtime/Private/PoseSearchConstraints.cpp
--- a/Engine/Plugins/Animation/PoseSearch/Source/Runtime/Private/PoseSearchConstraints.cpp
+++ b/Engine/Plugins/Animation/PoseSearch/Source/Runtime/Private/PoseSearchConstraints.cpp
@@ -99,7 +99,12 @@ void FConstraints::ProcessAll(const UMultiAnimAsset* MultiAnimAsset, float Delta
 	{
 		const FRole& Role = Roles[AnimContextIndex];
 		const UAnimationAsset* Animation = MultiAnimAsset->GetAnimationAsset(Role);
-		check(Animation);
+
+		if (!Animation)
+		{
+			UE_LOG(LogPoseSearch, Error, TEXT("FConstraints::ProcessAll aborting because of invalid animation asset for Role '%s' in MultiAnimAsset '%s'"), *Role.ToString(), *MultiAnimAsset->GetName());
+			return;
+		}
 
 		// initializing FAnimationAssetSampler only for extracting AnimNotifyStates
 		FAnimationAssetSampler Sampler(Animation, FTransform::Identity, FVector::ZeroVector, FAnimationAssetSampler::DefaultRootTransformSamplingRate, false, false);
@@ -241,7 +246,7 @@ void FConstraints::ProcessAll(const UMultiAnimAsset* MultiAnimAsset, float Delta
 			{
 				const FVector Pos = RootBoneTransform.GetTranslation();
 				const FRotator Rot(RootBoneTransform.GetRotation());
-				UE_LOG(LogPoseSearch, Error, TEXT("FInteractionIsland::UpdateConstraints unsupported non identity root bone in %s at time %f Pos(%f, %f, %f), Rot(%f, %f, %f)"), *Animation->GetName(), AssetTime, Pos.X, Pos.Y, Pos.Z, Rot.Pitch, Rot.Yaw, Rot.Roll);
+				UE_LOG(LogPoseSearch, Error, TEXT("FConstraints::ProcessAll unsupported non identity root bone in %s at time %f Pos(%f, %f, %f), Rot(%f, %f, %f)"), *Animation->GetName(), AssetTime, Pos.X, Pos.Y, Pos.Z, Rot.Pitch, Rot.Yaw, Rot.Roll);
 			}
 			RootBoneTransform = Sampler.ExtractRootTransform(AssetTime);
 			ComponentSpacePoses[AnimContextIndex].InitPose(Poses[AnimContextIndex]);
