PoseSearch - logging error when indexing databases with schema containing channels with negative weights

#Synced-CL 48276774

diff --git a/Engine/Plugins/Animation/PoseSearch/Source/Runtime/Private/PoseSearchDerivedData.cpp b/Engine/Plugins/Animation/PoseSearch/Source/Runtime/Private/PoseSearchDerivedData.cpp
--- a/Engine/Plugins/Animation/PoseSearch/Source/Runtime/Private/PoseSearchDerivedData.cpp
+++ b/Engine/Plugins/Animation/PoseSearch/Source/Runtime/Private/PoseSearchDerivedData.cpp
@@ -511,7 +511,7 @@ static bool InitSearchIndexAssets(FSearchIndexBase& SearchIndex, const UPoseSear
 	return !bAnyErrors;
 }
 
-static void PreprocessSearchIndexWeights(FSearchIndex& SearchIndex, const UPoseSearchSchema* Schema, TConstArrayView<float> Deviation)
+static bool PreprocessSearchIndexWeights(FSearchIndex& SearchIndex, const UPoseSearchSchema* Schema, TConstArrayView<float> Deviation)
 {
 #if WITH_EDITORONLY_DATA
 	// Copy deviation into search index for weights display in editor
@@ -524,6 +524,18 @@ static void PreprocessSearchIndexWeights(FSearchIndex& SearchIndex, const UPoseS
 	for (const TObjectPtr<UPoseSearchFeatureChannel>& ChannelPtr : Schema->GetChannels())
 	{
 		ChannelPtr->FillWeights(SearchIndex.WeightsSqrt);
+
+		const int32 ChannelDataStart = ChannelPtr->GetChannelDataOffset();
+		const int32 ChannelDataEnd = ChannelDataStart + ChannelPtr->GetChannelCardinality();
+
+		for (int32 ChannelDataIndex = ChannelDataStart; ChannelDataIndex < ChannelDataEnd; ++ChannelDataIndex)
+		{
+			if (SearchIndex.WeightsSqrt[ChannelDataIndex] < 0.f)
+			{
+				UE_LOG(LogPoseSearch, Error, TEXT("Error in PreprocessSearchIndexWeights: channel '%s' in schema '%s' has a negative weight!"), *ChannelPtr->GetName(), *Schema->GetName());
+				return false;
+			}
+		}
 	}
 
 	EPoseSearchDataPreprocessor DataPreprocessor = Schema->DataPreprocessor;
@@ -532,7 +544,7 @@ static void PreprocessSearchIndexWeights(FSearchIndex& SearchIndex, const UPoseS
 		// normalizing user weights: the idea behind this step is to be able to compare poses from databases using different schemas
 		RowMajorVectorMap MapWeights(SearchIndex.WeightsSqrt.GetData(), 1, NumDimensions);
 		const float WeightsSum = MapWeights.sum();
-		if (!FMath::IsNearlyZero(WeightsSum))
+		if (WeightsSum > UE_SMALL_NUMBER)
 		{
 			MapWeights *= (1.f / WeightsSum);
 		}
@@ -549,9 +561,14 @@ static void PreprocessSearchIndexWeights(FSearchIndex& SearchIndex, const UPoseS
 		for (int32 Dimension = 0; Dimension != NumDimensions; ++Dimension)
 		{
 			// the idea here is to pre-multiply the weights by the inverse of the variance (proportional to the square of the deviation) to have a "weighted Mahalanobis" distance
-			SearchIndex.WeightsSqrt[Dimension] /= Deviation[Dimension];
+			if (Deviation[Dimension] > UE_SMALL_NUMBER)
+			{
+				SearchIndex.WeightsSqrt[Dimension] /= Deviation[Dimension];
+			}
 		}
 	}
+
+	return true;
 }
 
 // it calculates Mean, PCAValues, and PCAProjectionMatrix
@@ -2050,7 +2067,13 @@ void FPoseSearchDatabaseAsyncCacheTask::OnGetComplete(UE::DerivedData::FCacheGet
 				const TArray<float> Deviation = FMeanDeviationCalculator::Calculate(SearchIndexBases, Schemas);
 
 				// Building FSearchIndex
-				PreprocessSearchIndexWeights(SearchIndex, MainDatabaseSchema, Deviation);
+				if (!PreprocessSearchIndexWeights(SearchIndex, MainDatabaseSchema, Deviation))
+				{
+					UE_LOG(LogPoseSearch, Error, TEXT("%s - %s BuildIndex Failed while preprocessing weights for schema '%s'"), *LexToString(FullIndexKey.Hash), *MainDatabaseName, *MainDatabaseSchema->GetName());
+					ResetSearchIndex();
+					return;
+				}
+
 				if (Owner.IsCanceled())
 				{
 					UE_LOG(LogPoseSearch, Log, TEXT("%s - %s BuildIndex Cancelled"), *LexToString(FullIndexKey.Hash), *MainDatabaseName);
