PoseHistory: Added blueprint library functions to get bone data from pose history.

#Synced-CL 48363837

diff --git a/Engine/Plugins/Animation/PoseSearch/Source/Runtime/Private/PoseSearchHistoryCollectorAnimNodeLibrary.cpp b/Engine/Plugins/Animation/PoseSearch/Source/Runtime/Private/PoseSearchHistoryCollectorAnimNodeLibrary.cpp
--- a/Engine/Plugins/Animation/PoseSearch/Source/Runtime/Private/PoseSearchHistoryCollectorAnimNodeLibrary.cpp
+++ b/Engine/Plugins/Animation/PoseSearch/Source/Runtime/Private/PoseSearchHistoryCollectorAnimNodeLibrary.cpp
@@ -3,6 +3,8 @@
 #include "PoseSearch/PoseSearchHistoryCollectorAnimNodeLibrary.h"
 #include "PoseSearch/AnimNode_PoseSearchHistoryCollector.h"
 
+#include "Animation/Skeleton.h"
+
 #include UE_INLINE_GENERATED_CPP_BY_NAME(PoseSearchHistoryCollectorAnimNodeLibrary)
 
 FPoseSearchHistoryCollectorAnimNodeReference UPoseSearchHistoryCollectorAnimNodeLibrary::ConvertToPoseHistoryNode(const FAnimNodeReference& Node, EAnimNodeReferenceConversionResult& Result)
@@ -49,3 +51,56 @@ void UPoseSearchHistoryCollectorAnimNodeLibrary::SetPoseHistoryNodeTransformTraj
 		}
 	}
 }
+
+bool UPoseSearchHistoryCollectorAnimNodeLibrary::GetPoseHistoryBoneWorldTransform(FTransform& OutTransform, const FPoseSearchHistoryCollectorAnimNodeReference& Node, const USkeleton* Skeleton, const FName BoneName, const float Time)
+{
+	if (!Skeleton)
+	{
+		OutTransform = FTransform::Identity;
+		return false;
+	}
+
+	if (FAnimNode_PoseSearchHistoryCollector* PoseSearchHistoryCollectorNodePtr = Node.GetAnimNodePtr<FAnimNode_PoseSearchHistoryCollector>())
+	{
+		if (UE::PoseSearch::FPoseHistory* PoseHistory = PoseSearchHistoryCollectorNodePtr->GetPoseHistoryPtr())
+		{
+			const int32 BoneIndex = Skeleton->GetReferenceSkeleton().FindBoneIndex(BoneName);
+			if (BoneIndex != INDEX_NONE && PoseHistory->GetTransformAtTime(Time, OutTransform, Skeleton, BoneIndex, UE::PoseSearch::WorldSpaceIndexType))
+			{
+				return true;
+			}
+		}
+	}
+
+	OutTransform = FTransform::Identity;
+	return false;
+}
+
+bool UPoseSearchHistoryCollectorAnimNodeLibrary::GetPoseHistoryBoneWorldLinearVelocity(FVector& OutVelocity, const FPoseSearchHistoryCollectorAnimNodeReference& Node, const USkeleton* Skeleton, const FName BoneName, const float Time, const float DeltaTime)
+{
+	if (!Skeleton)
+	{
+		OutVelocity = FVector::ZeroVector;
+		return false;
+	}
+
+	if (FAnimNode_PoseSearchHistoryCollector* PoseSearchHistoryCollectorNodePtr = Node.GetAnimNodePtr<FAnimNode_PoseSearchHistoryCollector>())
+	{
+		if (UE::PoseSearch::FPoseHistory* PoseHistory = PoseSearchHistoryCollectorNodePtr->GetPoseHistoryPtr())
+		{
+			FTransform PrevTransform, CurrTransform;
+			const int32 BoneIndex = Skeleton->GetReferenceSkeleton().FindBoneIndex(BoneName);
+			if (DeltaTime > UE_SMALL_NUMBER && 
+				BoneIndex != INDEX_NONE && 
+				PoseHistory->GetTransformAtTime(Time - DeltaTime, PrevTransform, Skeleton, BoneIndex, UE::PoseSearch::WorldSpaceIndexType) &&
+				PoseHistory->GetTransformAtTime(Time, CurrTransform, Skeleton, BoneIndex, UE::PoseSearch::WorldSpaceIndexType))
+			{
+				OutVelocity = (CurrTransform.GetLocation() - PrevTransform.GetLocation()) / DeltaTime;
+				return true;
+			}
+		}
+	}
+
+	OutVelocity = FVector::ZeroVector;
+	return false;
+}
diff --git a/Engine/Plugins/Animation/PoseSearch/Source/Runtime/Public/PoseSearch/PoseSearchHistoryCollectorAnimNodeLibrary.h b/Engine/Plugins/Animation/PoseSearch/Source/Runtime/Public/PoseSearch/PoseSearchHistoryCollectorAnimNodeLibrary.h
--- a/Engine/Plugins/Animation/PoseSearch/Source/Runtime/Public/PoseSearch/PoseSearchHistoryCollectorAnimNodeLibrary.h
+++ b/Engine/Plugins/Animation/PoseSearch/Source/Runtime/Public/PoseSearch/PoseSearchHistoryCollectorAnimNodeLibrary.h
@@ -15,6 +15,7 @@ struct FPoseSearchQueryTrajectory;
 PRAGMA_ENABLE_DEPRECATION_WARNINGS
 struct FTransformTrajectory;
 struct FAnimNode_PoseSearchHistoryCollector;
+class USkeleton;
 
 USTRUCT(BlueprintType)
 struct FPoseSearchHistoryCollectorAnimNodeReference : public FAnimNodeReference
@@ -45,8 +46,14 @@ public:
 	UFUNCTION(BlueprintCallable, Category = "Animation|PoseHistory", meta = (BlueprintThreadSafe, DisplayName = "Set Pose History Node Trajectory"))
 	static UE_API void SetPoseHistoryNodeTransformTrajectory(const FPoseSearchHistoryCollectorAnimNodeReference& PoseSearchHistoryCollectorNode, const FTransformTrajectory& Trajectory);
 
-	UFUNCTION(BlueprintPure, Category = "Animation|Pose Search|Experimental", meta = (BlueprintThreadSafe))
+	UFUNCTION(BlueprintPure, Category = "Animation|PoseHistory", meta = (BlueprintThreadSafe))
 	static UE_API FPoseHistoryReference GetPoseHistoryReference(const FPoseSearchHistoryCollectorAnimNodeReference& Node);
+
+	UFUNCTION(BlueprintPure, Category = "Animation|PoseHistory", meta = (BlueprintThreadSafe))
+	static UE_API bool GetPoseHistoryBoneWorldTransform(FTransform& OutTransform, const FPoseSearchHistoryCollectorAnimNodeReference& Node, const USkeleton* Skeleton, const FName BoneName, const float Time);
+
+	UFUNCTION(BlueprintPure, Category = "Animation|PoseHistory", meta = (BlueprintThreadSafe))
+	static UE_API bool GetPoseHistoryBoneWorldLinearVelocity(FVector& OutVelocity, const FPoseSearchHistoryCollectorAnimNodeReference& Node, const USkeleton* Skeleton, const FName BoneName, const float Time, const float DeltaTime);
 };
 
 #undef UE_API
