PoseSearch - logging cleanup

#Synced-CL 48365056

diff --git a/Engine/Plugins/Animation/PoseSearch/Source/Editor/Private/PoseSearchDatabaseViewModel.cpp b/Engine/Plugins/Animation/PoseSearch/Source/Editor/Private/PoseSearchDatabaseViewModel.cpp
--- a/Engine/Plugins/Animation/PoseSearch/Source/Editor/Private/PoseSearchDatabaseViewModel.cpp
+++ b/Engine/Plugins/Animation/PoseSearch/Source/Editor/Private/PoseSearchDatabaseViewModel.cpp
@@ -160,7 +160,6 @@ bool FDatabasePreviewActor::SpawnPreviewActor(UWorld* World, const UPoseSearchDa
 		TrajectorySpeed[0] = NumPoses > 1 ? TrajectorySpeed[1] : 0.f;
 	}
 
-	UE_LOG(LogPoseSearchEditor, Log, TEXT("Spawned preview Actor: %s"), *GetNameSafe(ActorPtr.Get()));
 	return true;
 }
 
diff --git a/Engine/Plugins/Animation/PoseSearch/Source/Editor/Private/PoseSearchInteractionAssetEditor.cpp b/Engine/Plugins/Animation/PoseSearch/Source/Editor/Private/PoseSearchInteractionAssetEditor.cpp
--- a/Engine/Plugins/Animation/PoseSearch/Source/Editor/Private/PoseSearchInteractionAssetEditor.cpp
+++ b/Engine/Plugins/Animation/PoseSearch/Source/Editor/Private/PoseSearchInteractionAssetEditor.cpp
@@ -86,8 +86,6 @@ bool FInteractionAssetPreviewActor::SpawnPreviewActor(UWorld* World, const UPose
 		AnimInstance->SetBlendSpacePosition(BlendParameters);
 
 		ActorPtr->SetActorTransform(ActorTransform);
-
-		UE_LOG(LogPoseSearchEditor, Log, TEXT("Spawned preview Actor: %s"), *GetNameSafe(ActorPtr.Get()));
 	}
 	return true;
 }
diff --git a/Engine/Plugins/Animation/PoseSearch/Source/Runtime/Private/PoseSearchInteractionIsland.cpp b/Engine/Plugins/Animation/PoseSearch/Source/Runtime/Private/PoseSearchInteractionIsland.cpp
--- a/Engine/Plugins/Animation/PoseSearch/Source/Runtime/Private/PoseSearchInteractionIsland.cpp
+++ b/Engine/Plugins/Animation/PoseSearch/Source/Runtime/Private/PoseSearchInteractionIsland.cpp
@@ -946,7 +946,7 @@ void FInteractionIsland::FPostTickFunction::ExecuteTick(float DeltaTime, enum EL
 		{
 			// @todo: need to figure out why when creating a new island FPostTickFunction gets called twice (it's not a real issue rather than a performance hit)
 			//		  use GVarPoseSearchInteractionCacheIslands = false to debug the issue (it destroys the islands every frame)
-			UE_LOG(LogPoseSearch, Warning, TEXT("FInteractionIsland::FPostTickFunction::ExecuteTick, called twice before UPoseSearchInteractionSubsystem::Tick!"));
+			UE_LOG(LogPoseSearch, Log, TEXT("FInteractionIsland::FPostTickFunction::ExecuteTick, called twice before UPoseSearchInteractionSubsystem::Tick!"));
 		}
 		else
 		{
@@ -984,7 +984,7 @@ void FInteractionIsland::FClosingTickFunction::ExecuteTick(float DeltaTime, enum
 		{
 			// @todo: need to figure out why when creating a new island FClosingTickFunction gets called twice (it's not a real issue rather than a performance hit)
 			//		  use GVarPoseSearchInteractionCacheIslands = false to debug the issue (it destroys the islands every frame)
-			UE_LOG(LogPoseSearch, Warning, TEXT("FInteractionIsland::FClosingTickFunction::ExecuteTick, called twice before UPoseSearchInteractionSubsystem::Tick!"));
+			UE_LOG(LogPoseSearch, Log, TEXT("FInteractionIsland::FClosingTickFunction::ExecuteTick, called twice before UPoseSearchInteractionSubsystem::Tick!"));
 		}
 		else
 		{
