UAFPoseSearch - MMI constraints support

#Synced-CL 48365056

diff --git a/Engine/Plugins/Animation/PoseSearch/Source/Runtime/Private/PoseSearchConstraints.cpp b/Engine/Plugins/Animation/PoseSearch/Source/Runtime/Private/PoseSearchConstraints.cpp
--- a/Engine/Plugins/Animation/PoseSearch/Source/Runtime/Private/PoseSearchConstraints.cpp
+++ b/Engine/Plugins/Animation/PoseSearch/Source/Runtime/Private/PoseSearchConstraints.cpp
@@ -273,10 +273,10 @@ void FConstraints::ProcessAll(const UMultiAnimAsset* MultiAnimAsset, float Delta
 							const FColor Color = DrawAnimContext == AnimContext ? FColor::Red : FColor::Blue;
 
 							// current characters poses
-							UE_VLOG_SEGMENT(DrawAnimContext, LogName, Display, CurrentBoneWorldLocation, CurrentParentBoneWorldLocation, Color, TEXT(""));
+							UE_VLOG_SEGMENT(DrawAnimContext, LogName, VeryVerbose, CurrentBoneWorldLocation, CurrentParentBoneWorldLocation, Color, TEXT(""));
 
 							// sampled poses from UAnimationAsset
-							UE_VLOG_SEGMENT(DrawAnimContext, LogName, Display, AnimBoneWorldLocation, AnimParentBoneWorldLocation, Color, TEXT(""));
+							UE_VLOG_SEGMENT(DrawAnimContext, LogName, VeryVerbose, AnimBoneWorldLocation, AnimParentBoneWorldLocation, Color, TEXT(""));
 						}
 					}
 				}
@@ -367,17 +367,21 @@ void FConstraints::ProcessAll(const UMultiAnimAsset* MultiAnimAsset, float Delta
 							const UObject* DrawAnimContext = AnimContexts[DrawAnimContextIndex];
 
 							// current characters poses
-							UE_VLOG_SEGMENT(DrawAnimContext, LogName, Display, CurrentFromTransform.GetTranslation(), CurrentToTransform.GetTranslation(), FColor::Yellow, TEXT(""));
+							UE_VLOG_SEGMENT(DrawAnimContext, LogName, VeryVerbose, CurrentFromTransform.GetTranslation(), CurrentToTransform.GetTranslation(), FColor::Yellow, TEXT(""));
 
 							// sampled poses from UAnimationAsset
-							UE_VLOG_SEGMENT(DrawAnimContext, LogName, Display, AnimFromTransform.GetTranslation(), AnimToTransform.GetTranslation(), FColor::Yellow, TEXT(""));
+							UE_VLOG_SEGMENT(DrawAnimContext, LogName, VeryVerbose, AnimFromTransform.GetTranslation(), AnimToTransform.GetTranslation(), FColor::Yellow, TEXT(""));
 
 							// wanted transforms
-							UE_VLOG_SEGMENT(DrawAnimContext, LogName, Display, CurrentFromTransform.GetTranslation(), WantedCurrentToTransform.GetTranslation(), FColor::Green, TEXT(""));
-							UE_VLOG_SEGMENT(DrawAnimContext, LogName, Display, CurrentToTransform.GetTranslation(), WantedCurrentFromTransform.GetTranslation(), FColor::Green, TEXT(""));
+							UE_VLOG_SEGMENT(DrawAnimContext, LogName, VeryVerbose, CurrentFromTransform.GetTranslation(), WantedCurrentToTransform.GetTranslation(), FColor::Green, TEXT(""));
+							UE_VLOG_SEGMENT(DrawAnimContext, LogName, VeryVerbose, CurrentToTransform.GetTranslation(), WantedCurrentFromTransform.GetTranslation(), FColor::Green, TEXT(""));
 
-							UE_VLOG_SEGMENT(DrawAnimContext, LogName, Display, CurrentFromTransform.GetTranslation(), Constraint.FromSocketTransform.GetTranslation(), FColor::Orange, TEXT(""));
-							UE_VLOG_SEGMENT(DrawAnimContext, LogName, Display, CurrentToTransform.GetTranslation(), Constraint.ToSocketTransform.GetTranslation(), FColor::Orange, TEXT(""));
+							UE_VLOG_SEGMENT(DrawAnimContext, LogName, VeryVerbose, CurrentFromTransform.GetTranslation(), Constraint.FromSocketTransform.GetTranslation(), FColor::Orange, TEXT(""));
+							UE_VLOG_SEGMENT(DrawAnimContext, LogName, VeryVerbose, CurrentToTransform.GetTranslation(), Constraint.ToSocketTransform.GetTranslation(), FColor::Orange, TEXT(""));
+
+							const FColor DesiredReachColor = (FLinearColor::Red* (1.f - Constraint.DesiredReach) + FLinearColor::Green * Constraint.DesiredReach).ToFColor(true);
+							UE_VLOG_SPHERE(DrawAnimContext, LogName, Display, Constraint.FromSocketTransform.GetTranslation(), 2.f, DesiredReachColor, TEXT(""));
+							UE_VLOG_SPHERE(DrawAnimContext, LogName, Display, Constraint.ToSocketTransform.GetTranslation(), 2.f, DesiredReachColor, TEXT(""));
 						}
 					}
 #endif // WITH_EDITORONLY_DATA && ENABLE_ANIM_DEBUG && ENABLE_VISUAL_LOG
diff --git a/Engine/Plugins/Experimental/UAF/UAFPoseSearch/Source/UAFPoseSearch/Private/RigUnit_GetMotionMatchInteractionConstraint.cpp b/Engine/Plugins/Experimental/UAF/UAFPoseSearch/Source/UAFPoseSearch/Private/RigUnit_GetMotionMatchInteractionConstraint.cpp
new file mode 100644
--- /dev/null
+++ b/Engine/Plugins/Experimental/UAF/UAFPoseSearch/Source/UAFPoseSearch/Private/RigUnit_GetMotionMatchInteractionConstraint.cpp
@@ -0,0 +1,11 @@
+// Copyright Epic Games, Inc. All Rights Reserved.
+
+#include "RigUnit_GetMotionMatchInteractionConstraint.h"
+#include "PoseSearch/PoseSearchInteractionLibrary.h"
+
+#include UE_INLINE_GENERATED_CPP_BY_NAME(RigUnit_GetMotionMatchInteractionConstraint)
+
+FRigUnit_GetMotionMatchInteractionConstraint_Execute()
+{
+	UPoseSearchInteractionLibrary::GetMotionMatchInteractionConstraint(ExecuteContext.GetOwningObject(), SocketName, DesiredReach, Transform, IsValid, bCompareOwningActors);
+}
diff --git a/Engine/Plugins/Experimental/UAF/UAFPoseSearch/Source/UAFPoseSearch/Private/RigUnit_GetMotionMatchInteractionConstraint.h b/Engine/Plugins/Experimental/UAF/UAFPoseSearch/Source/UAFPoseSearch/Private/RigUnit_GetMotionMatchInteractionConstraint.h
new file mode 100644
--- /dev/null
+++ b/Engine/Plugins/Experimental/UAF/UAFPoseSearch/Source/UAFPoseSearch/Private/RigUnit_GetMotionMatchInteractionConstraint.h
@@ -0,0 +1,36 @@
+// Copyright Epic Games, Inc. All Rights Reserved.
+
+#pragma once
+
+#include "Graph/RigUnit_AnimNextBase.h"
+#include "RigUnit_GetMotionMatchInteractionConstraint.generated.h"
+
+/** Gets the motion matching interaction constraint property associated to SocketName */
+USTRUCT(meta=(DisplayName="Get Motion Match Interaction Constraint", Category="Animation Graph", NodeColor="0, 1, 1"))
+struct FRigUnit_GetMotionMatchInteractionConstraint : public FRigUnit_AnimNextBase
+{
+	GENERATED_BODY()
+
+	RIGVM_METHOD()
+	void Execute();
+
+	// Socket name associated with the constraint
+	UPROPERTY(EditAnywhere, Category = "Graph", meta = (Input))
+	FName SocketName = NAME_None;
+
+	// search the constraint for the owning actor instead of using the AnimContext
+	UPROPERTY(EditAnywhere, Category = "Graph", meta = (Input))
+	bool bCompareOwningActors = false;
+	
+	// output desired reach for the constraint
+	UPROPERTY(EditAnywhere, Category = "Graph", meta = (Output))
+	float DesiredReach = 0.f;
+	
+	// output transform for the constraint
+	UPROPERTY(EditAnywhere, Category = "Graph", meta = (Output))
+	FTransform Transform = FTransform::Identity;
+	
+	// validity of the requested constraint
+	UPROPERTY(EditAnywhere, Category = "Graph", meta = (Output))
+	bool IsValid = false;
+};
