PoseSearch - MM keep interaction alive fix

#Synced-CL 48424704

diff --git a/Engine/Plugins/Animation/PoseSearch/Source/Runtime/Private/AnimNode_MotionMatching.cpp b/Engine/Plugins/Animation/PoseSearch/Source/Runtime/Private/AnimNode_MotionMatching.cpp
--- a/Engine/Plugins/Animation/PoseSearch/Source/Runtime/Private/AnimNode_MotionMatching.cpp
+++ b/Engine/Plugins/Animation/PoseSearch/Source/Runtime/Private/AnimNode_MotionMatching.cpp
@@ -159,10 +159,11 @@ void FAnimNode_MotionMatching::UpdateAssetPlayer(const FAnimationUpdateContext&
 		const UObject* AnimContext = Context.AnimInstanceProxy->GetAnimInstanceObject();
 		check(AnimContext);
 		// MotionMatchInteraction synchronizes and integrates MotionMatchingState.SearchResult internally, so there's no need to update it with the blendstack information
-		const bool bIsInteraction = UPoseSearchInteractionLibrary::MotionMatchInteraction(MotionMatchingState.SearchResult, Availabilities, AnimContext, FName(), PoseHistory, bValidateResultAgainstAvailabilities, bKeepInteractionAlive, BlendTime);
+		UPoseSearchInteractionLibrary::MotionMatchInteraction(MotionMatchingState.SearchResult, Availabilities,
+			AnimContext, FName(), PoseHistory, bValidateResultAgainstAvailabilities, bKeepInteractionAlive, BlendTime, Context.GetDeltaTime());
 
 		// performing the regular single character motion matching search in case there's no MM interaction
-		if (!bIsInteraction)
+		if (!MotionMatchingState.SearchResult.bIsInteraction)
 		{
 			// We adjust the motion matching state asset time to the current player node's asset time. This is done 
 			// because the player node may have ticked more or less time than we expected due to variable dt or the 
diff --git a/Engine/Plugins/Animation/PoseSearch/Source/Runtime/Private/PoseSearchInteractionLibrary.cpp b/Engine/Plugins/Animation/PoseSearch/Source/Runtime/Private/PoseSearchInteractionLibrary.cpp
--- a/Engine/Plugins/Animation/PoseSearch/Source/Runtime/Private/PoseSearchInteractionLibrary.cpp
+++ b/Engine/Plugins/Animation/PoseSearch/Source/Runtime/Private/PoseSearchInteractionLibrary.cpp
@@ -37,11 +37,10 @@ void UPoseSearchInteractionLibrary::MotionMatchInteraction(FPoseSearchBlueprintR
 	}
 }
 
-bool UPoseSearchInteractionLibrary::MotionMatchInteraction(FPoseSearchBlueprintResult& InOutResult, const TArrayView<const FPoseSearchInteractionAvailability> Availabilities, const UObject* AnimContext, FName PoseHistoryName, const UE::PoseSearch::IPoseHistory* PoseHistory, bool bValidateResultAgainstAvailabilities, bool bKeepInteractionAlive, float BlendTime)
+void UPoseSearchInteractionLibrary::MotionMatchInteraction(FPoseSearchBlueprintResult& InOutResult, const TArrayView<const FPoseSearchInteractionAvailability> Availabilities, const UObject* AnimContext, FName PoseHistoryName, const UE::PoseSearch::IPoseHistory* PoseHistory, bool bValidateResultAgainstAvailabilities, bool bKeepInteractionAlive, float BlendTime, float DeltaTime)
 {
 	using namespace UE::PoseSearch;
 	
-	bool bIsInteraction = false;
 	if (!Availabilities.IsEmpty())
 	{
 		CheckInteractionThreadSafety(AnimContext);
@@ -52,50 +51,45 @@ bool UPoseSearchInteractionLibrary::MotionMatchInteraction(FPoseSearchBlueprintR
 
 		if (SearchResult.SelectedAnim)
 		{
-			bIsInteraction = true;
 			InOutResult = SearchResult;
+			return;
 		}
 	}
 
 	// performing the regular single character motion matching search in case there's no MM interaction
-	if (!bIsInteraction)
+	if (!InOutResult.bIsInteraction)
 	{
-		if (!InOutResult.bIsInteraction)
-		{
-			// do nothing
-		}
-		else if (bKeepInteractionAlive)
-		{
-			// checking if the kept alive interaction has reached the end of animation
-			if (!InOutResult.SelectedDatabase || INDEX_NONE == InOutResult.SelectedDatabase->GetPoseIndex(InOutResult.SelectedAnim.Get(), InOutResult.SelectedTime, InOutResult.bIsMirrored, InOutResult.BlendParameters))
-			{
-				InOutResult = FPoseSearchBlueprintResult();
-			}
-			else
-			{
-				// letting the interaction animation run until its length minus blend time
-				// (to avoid having to blend from a frozen animation that reached its end for the entire duration of the blend)
-				const UAnimationAsset* AnimationAssetForRole = InOutResult.GetAnimationAssetForRole();
-				if (!AnimationAssetForRole || InOutResult.SelectedTime >= (AnimationAssetForRole->GetPlayLength() - BlendTime))
-				{
-					InOutResult = FPoseSearchBlueprintResult();
-				}
-				else
-				{
-					// we're keeping alive only the animation part of the search result
-					InOutResult.ActorRootTransforms.Reset();
-					InOutResult.ActorRootBoneTransforms.Reset();
-					InOutResult.AnimContexts.Reset();
-				}
-			}
-		}
-		else
-		{
-			InOutResult = FPoseSearchBlueprintResult();
-		}
+		// do nothing
+		return;
+	}
+	
+	if (!bKeepInteractionAlive)
+	{
+		InOutResult = FPoseSearchBlueprintResult();
+		return;
+	}
+
+	// checking if the kept alive interaction has reached the end of animation
+	if (!InOutResult.SelectedDatabase || INDEX_NONE == InOutResult.SelectedDatabase->GetPoseIndex(InOutResult.SelectedAnim.Get(), InOutResult.SelectedTime, InOutResult.bIsMirrored, InOutResult.BlendParameters))
+	{
+		InOutResult = FPoseSearchBlueprintResult();
+		return;
+	}
+
+	// letting the interaction animation run until its length minus blend time
+	// (to avoid having to blend from a frozen animation that reached its end for the entire duration of the blend)
+	const UAnimationAsset* AnimationAssetForRole = InOutResult.GetAnimationAssetForRole();
+	if (!AnimationAssetForRole || InOutResult.SelectedTime >= (AnimationAssetForRole->GetPlayLength() - BlendTime))
+	{
+		InOutResult = FPoseSearchBlueprintResult();
+		return;
 	}
 
-	return bIsInteraction;
+	// we're keeping alive only the animation part of the search result
+	InOutResult.ActorRootTransforms.Reset();
+	InOutResult.ActorRootBoneTransforms.Reset();
+	InOutResult.AnimContexts.Reset();
+	InOutResult.SelectedTime += DeltaTime;
 }
 
 void UPoseSearchInteractionLibrary::CalculateFullAlignedTransforms(const FPoseSearchBlueprintResult& Result, float TimeOffset, bool bWarpUsingRootBone, TArray<FTransform>& OutFullAlignedTransforms)
diff --git a/Engine/Plugins/Animation/PoseSearch/Source/Runtime/Public/PoseSearch/PoseSearchInteractionLibrary.h b/Engine/Plugins/Animation/PoseSearch/Source/Runtime/Public/PoseSearch/PoseSearchInteractionLibrary.h
--- a/Engine/Plugins/Animation/PoseSearch/Source/Runtime/Public/PoseSearch/PoseSearchInteractionLibrary.h
+++ b/Engine/Plugins/Animation/PoseSearch/Source/Runtime/Public/PoseSearch/PoseSearchInteractionLibrary.h
@@ -49,8 +49,10 @@ public:
 	UFUNCTION(Experimental, BlueprintCallable, Category = "Animation|Pose Search", meta = (BlueprintThreadSafe))
 	static UE_API void CalculateNoAlignedTransform(const FPoseSearchBlueprintResult& CurrentResult, float TimeOffset, bool bWarpUsingRootBone, FTransform& OutNoAlignedTransform);
 
+	// Experimental, this feature might be removed without warning, not for production use
 	static UE_API void MotionMatchInteraction(FPoseSearchBlueprintResult& Result, const TArrayView<const FPoseSearchInteractionAvailability> Availabilities, const UObject* AnimContext, FName PoseHistoryName, const UE::PoseSearch::IPoseHistory* PoseHistory, bool bValidateResultAgainstAvailabilities = true);
-	static UE_API bool MotionMatchInteraction(FPoseSearchBlueprintResult& InOutResult, const TArrayView<const FPoseSearchInteractionAvailability> Availabilities, const UObject* AnimContext, FName PoseHistoryName, const UE::PoseSearch::IPoseHistory* PoseHistory, bool bValidateResultAgainstAvailabilities, bool bKeepInteractionAlive, float BlendTime);
+	// Experimental, this feature might be removed without warning, not for production use
+	static UE_API void MotionMatchInteraction(FPoseSearchBlueprintResult& InOutResult, const TArrayView<const FPoseSearchInteractionAvailability> Availabilities, const UObject* AnimContext, FName PoseHistoryName, const UE::PoseSearch::IPoseHistory* PoseHistory, bool bValidateResultAgainstAvailabilities, bool bKeepInteractionAlive, float BlendTime, float DeltaTime);
 
 	UFUNCTION(Experimental, BlueprintPure, Category = "Animation|Pose Search", meta = (BlueprintThreadSafe))
 	static UE_API FPoseSearchContinuingProperties GetMontageContinuingProperties(UAnimInstance* AnimInstance);
diff --git a/Engine/Plugins/Experimental/UAF/UAFPoseSearch/Source/UAFPoseSearch/Private/MotionMatchingTrait.cpp b/Engine/Plugins/Experimental/UAF/UAFPoseSearch/Source/UAFPoseSearch/Private/MotionMatchingTrait.cpp
--- a/Engine/Plugins/Experimental/UAF/UAFPoseSearch/Source/UAFPoseSearch/Private/MotionMatchingTrait.cpp
+++ b/Engine/Plugins/Experimental/UAF/UAFPoseSearch/Source/UAFPoseSearch/Private/MotionMatchingTrait.cpp
@@ -146,14 +146,14 @@ PRAGMA_ENABLE_DEPRECATION_WARNINGS
 		}
 #endif // WITH_EDITOR
 
-
-
+		const float DeltaTime = TraitState.GetDeltaTime();
 		const TArrayView<const FPoseSearchInteractionAvailability> Availabilities = SharedData->GetAvailabilities(Binding);
 		// MotionMatchInteraction synchronizes and integrates MotionMatchingState.SearchResult internally, so there's no need to update it with the blendstack information
-		const bool bIsInteraction = UPoseSearchInteractionLibrary::MotionMatchInteraction(MotionMatchingState.SearchResult, Availabilities, AnimContext, FName(), PoseHistory, SharedData->GetbValidateResultAgainstAvailabilities(Binding), SharedData->GetbKeepInteractionAlive(Binding), SharedData->BlendArguments.BlendTime);
+		UPoseSearchInteractionLibrary::MotionMatchInteraction(MotionMatchingState.SearchResult, Availabilities, AnimContext, FName(), PoseHistory,
+			SharedData->GetbValidateResultAgainstAvailabilities(Binding), SharedData->GetbKeepInteractionAlive(Binding), SharedData->BlendArguments.BlendTime, DeltaTime);
 
 		// performing the regular single character motion matching search in case there's no MM interaction
-		if (!bIsInteraction)
+		if (!MotionMatchingState.SearchResult.bIsInteraction)
 		{
 			// We adjust the MotionMatchingInteractionState time to the current player node's asset time. This is done 
 			// because the player node may have ticked more or less time than we expected due to variable dt or the 
@@ -207,7 +207,6 @@ PRAGMA_ENABLE_DEPRECATION_WARNINGS
 			}
 			else
 			{
-				const float DeltaTime = TraitState.GetDeltaTime();
 				FFloatInterval PoseSearchPlayRate = SharedData->GetPlayRate(Binding);
 #if ENABLE_ANIM_DEBUG
 				bool bPoseSearchPlayRateEnabled = GVarAnimNextMotionMatchingPlayRateEnabled == UE::Private::EPlayRateState::Enabled || GVarAnimNextMotionMatchingPlayRateEnabled == UE::Private::EPlayRateState::PoseSearchOnly;
@@ -395,12 +394,14 @@ PRAGMA_ENABLE_DEPRECATION_WARNINGS
 
 		PublishResults(Binding);
 
-		if (bIsInteraction)
+		if (MotionMatchingState.SearchResult.bIsInteraction && SharedData->GetbEnableWarp(Binding))
 		{
 			if (const UMultiAnimAsset* MultiAnimAsset = Cast<UMultiAnimAsset>(MotionMatchingState.SearchResult.SelectedAnim))
 			{
 				const int32 CurrentResultRoleIndex = GetRoleIndex(MultiAnimAsset, MotionMatchingState.SearchResult.Role);
-				if (CurrentResultRoleIndex != INDEX_NONE)
+				if (CurrentResultRoleIndex != INDEX_NONE &&
+					MotionMatchingState.SearchResult.ActorRootTransforms.Num() == MotionMatchingState.SearchResult.AnimContexts.Num() &&
+					MotionMatchingState.SearchResult.ActorRootBoneTransforms.Num() == MotionMatchingState.SearchResult.AnimContexts.Num())
 				{
 					const int32 NumRoles = MultiAnimAsset->GetNumRoles();
 					TArray<FTransform, TInlineAllocator<PreallocatedRolesNum>> FullAlignedTransforms;
