PoseSearch - MMI broad phase hash grid

#Synced-CL 48618295

diff --git a/Engine/Plugins/Animation/PoseSearch/Source/Runtime/Private/PoseSearchInteractionSubsystem.cpp b/Engine/Plugins/Animation/PoseSearch/Source/Runtime/Private/PoseSearchInteractionSubsystem.cpp
--- a/Engine/Plugins/Animation/PoseSearch/Source/Runtime/Private/PoseSearchInteractionSubsystem.cpp
+++ b/Engine/Plugins/Animation/PoseSearch/Source/Runtime/Private/PoseSearchInteractionSubsystem.cpp
@@ -10,16 +10,23 @@
 #include "PoseSearch/PoseSearchRole.h"
 #include "PoseSearch/PoseSearchSchema.h"
 #include "PoseSearch/Trace/PoseSearchTraceLogger.h"
+#include "Util/GridIndexing2.h"
 #include "VisualLogger/VisualLogger.h"
 
 #include UE_INLINE_GENERATED_CPP_BY_NAME(PoseSearchInteractionSubsystem)
 
 namespace UE::PoseSearch
 {
-static bool GVarPoseSearchInteractionSortAnimContextInfosByUniqueCommonRoleFilter = true;
+enum EGenerateAnimContextInfosMethod
+{
+	BruteForce,
+	UniqueCommonRoleFilter,
+	HashGrid
+};
+static int32 GVarPoseSearchInteractionGenerateAnimContextInfosMethod = HashGrid;
 #if !NO_CVARS
-static FAutoConsoleVariableRef CVarPoseSearchInteractionSortAnimContextInfosByUniqueCommonRoleFilter(TEXT("a.PoseSearchInteraction.SortAnimContextInfosByUniqueCommonRoleFilter"),
-	GVarPoseSearchInteractionSortAnimContextInfosByUniqueCommonRoleFilter, TEXT("Sort AnimContextInfos By UniqueCommonRoleFilter (it improves broadphase performance when actors publish availabilities with a single role filter)"));
+static FAutoConsoleVariableRef CVarPoseSearchInteractionGenerateAnimContextInfosMethod(TEXT("a.PoseSearchInteraction.GenerateAnimContextInfosMethod"), GVarPoseSearchInteractionGenerateAnimContextInfosMethod,
+	TEXT("GenerateAnimContextInfosMethod: 0) brute force, all vs all no partitioning, 1) Sort AnimContextInfos by unique common role filter, it improves broadphase performance when actors publish availabilities with a single role filter, 2) use hash grid partitioning"));
 
 static bool GVarPoseSearchInteractionEnabled = true;
 static FAutoConsoleVariableRef CVarPoseSearchInteractionEnabled(TEXT("a.PoseSearchInteraction.Enabled"), GVarPoseSearchInteractionEnabled, TEXT("Enable/Disable Pose Search Interaction"));
@@ -86,6 +93,11 @@ struct FAnimContextInfo
 		return Location;
 	}
 
+	FVector2f GetLocation2f() const
+	{
+		return FVector2f(Location.X, Location.Y);
+	}
+
 	float GetAvailabilitiesMaxBroadPhaseRadius() const
 	{
 		return AvailabilitiesMaxBroadPhaseRadius;
@@ -615,6 +627,7 @@ void UPoseSearchInteractionSubsystem::GenerateAnimContextInfosAndTagToDatabases(
 {
 	QUICK_SCOPE_CYCLE_COUNTER(STAT_UPoseSearchInteractionSubsystem_GenerateAnimContextInfos);
 
+	using namespace UE::Geometry;
 	using namespace UE::PoseSearch;
 
 	const UWorld* SubsystemWorld = GetWorld();
@@ -635,15 +648,15 @@ void UPoseSearchInteractionSubsystem::GenerateAnimContextInfosAndTagToDatabases(
 #endif // DO_CHECK
 		
 		// adding AnimContext to AnimContextsWithAvailabilities only if at least one availability has a valid database or has a valid tag
-		for (const FInteractionAvailability& InteractionAvailabilityEx : AnimContextAvailabilities.Availabilities)
+		for (const FInteractionAvailability& InteractionAvailability : AnimContextAvailabilities.Availabilities)
 		{
 			bool bAnyValidAvailability = false;
-			if (const UPoseSearchDatabase* Database = InteractionAvailabilityEx.Database.Get())
+			if (const UPoseSearchDatabase* Database = InteractionAvailability.Database.Get())
 			{
 				check(Database->Schema);
-				if (InteractionAvailabilityEx.IsTagValid())
+				if (InteractionAvailability.IsTagValid())
 				{
-					TagToDatabases.FindOrAdd(InteractionAvailabilityEx.Tag).AddUnique(InteractionAvailabilityEx.Database);
+					TagToDatabases.FindOrAdd(InteractionAvailability.Tag).AddUnique(InteractionAvailability.Database);
 				}
 			}
 		}
@@ -666,7 +679,30 @@ void UPoseSearchInteractionSubsystem::GenerateAnimContextInfosAndTagToDatabases(
 			}
 		};
 
-	if (GVarPoseSearchInteractionSortAnimContextInfosByUniqueCommonRoleFilter)
+	const int32 NumAnimContextInfos = AnimContextsAvailabilities.Num();
+	AnimContextInfos.SetNum(NumAnimContextInfos);
+	for (int32 AnimContextInfoIndex = 0; AnimContextInfoIndex < NumAnimContextInfos; ++AnimContextInfoIndex)
+	{
+		AnimContextInfos[AnimContextInfoIndex].Init(AnimContextsAvailabilities[AnimContextInfoIndex]);
+	}
+
+	switch (GVarPoseSearchInteractionGenerateAnimContextInfosMethod)
+	{
+	default:
+	case EGenerateAnimContextInfosMethod::BruteForce:
+	{
+		// solving the broad phase using the AnimContextInfos
+		for (int32 AnimContextInfoIndexA = 0; AnimContextInfoIndexA < NumAnimContextInfos - 1; ++AnimContextInfoIndexA)
+		{
+			for (int32 AnimContextInfoIndexB = AnimContextInfoIndexA + 1; AnimContextInfoIndexB < NumAnimContextInfos; ++AnimContextInfoIndexB)
+			{
+				EvaluateForNearbyAnimContextInfo(AnimContextInfoIndexA, AnimContextInfoIndexB);
+			}
+		}
+		break;
+	}
+	
+	case EGenerateAnimContextInfosMethod::UniqueCommonRoleFilter:
 	{
 		typedef TMap<const FRole, int32, TInlineSetAllocator<16, TMemStackSetAllocator<>>> FCommonRoleFilterToIndex;
 		FCommonRoleFilterToIndex CommonRoleFilterToIndex;
@@ -677,11 +713,8 @@ void UPoseSearchInteractionSubsystem::GenerateAnimContextInfosAndTagToDatabases(
 		// adding RoledAnimContextInfoIndexes[0] containing all the indexes for those AnimContextInfos returning nullptr from GetUniqueCommonRoleFilter
 		RoledAnimContextInfoIndexes.AddDefaulted();
 
-		const int32 NumAnimContextInfos = AnimContextsAvailabilities.Num();
-		AnimContextInfos.SetNum(NumAnimContextInfos);
 		for (int32 AnimContextInfoIndex = 0; AnimContextInfoIndex < NumAnimContextInfos; ++AnimContextInfoIndex)
 		{
-			AnimContextInfos[AnimContextInfoIndex].Init(AnimContextsAvailabilities[AnimContextInfoIndex]);
 			if (const FRole* UniqueCommonRoleFilter = AnimContextInfos[AnimContextInfoIndex].GetUniqueCommonRoleFilter())
 			{
 				if (int32* UniqueCommonRoleFilterIndex = CommonRoleFilterToIndex.Find(*UniqueCommonRoleFilter))
@@ -740,26 +773,64 @@ void UPoseSearchInteractionSubsystem::GenerateAnimContextInfosAndTagToDatabases(
 				}
 			}
 		}
-
-		
+		break;
 	}
-	else
+
+	case EGenerateAnimContextInfosMethod::HashGrid:
 	{
-		const int32 NumAnimContextInfos = AnimContextsAvailabilities.Num();
-		AnimContextInfos.SetNum(NumAnimContextInfos);
+		float GlobalMaxRadius = 0.f;
 		for (int32 AnimContextInfoIndex = 0; AnimContextInfoIndex < NumAnimContextInfos; ++AnimContextInfoIndex)
 		{
-			AnimContextInfos[AnimContextInfoIndex].Init(AnimContextsAvailabilities[AnimContextInfoIndex]);
+			GlobalMaxRadius = FMath::Max(GlobalMaxRadius, AnimContextInfos[AnimContextInfoIndex].GetAvailabilitiesMaxBroadPhaseRadius());
 		}
 
-		// solving the broad phase using the AnimContextInfos
-		for (int32 AnimContextInfoIndexA = 0; AnimContextInfoIndexA < NumAnimContextInfos - 1; ++AnimContextInfoIndexA)
+		// creating an hash grid
+		TScaleGridIndexer2<float> Indexer;
+		// using a minimum cell size of twice the GlobalMaxRadius to minimize the number of cells to search for neighbors
+		// AnimContextInfos (max of 4 cells), and clamping it to 1cm to avoid division by zero inside TScaleGridIndexer2
+		Indexer.CellSize = FMath::Max(1.f, GlobalMaxRadius * 2.f);
+
+		typedef TMultiMap<FVector2i, int, TMemStackSetAllocator<>> FHashGrid;
+		FHashGrid Hash;
+		Hash.Reserve(AnimContextsAvailabilitiesNum);
+
+		for (int32 AnimContextInfoIndex = 0; AnimContextInfoIndex < NumAnimContextInfos; ++AnimContextInfoIndex)
 		{
-			for (int32 AnimContextInfoIndexB = AnimContextInfoIndexA + 1; AnimContextInfoIndexB < NumAnimContextInfos; ++AnimContextInfoIndexB)
+			const FVector2i idx = Indexer.ToGrid(AnimContextInfos[AnimContextInfoIndex].GetLocation2f());
+			Hash.Add(idx, AnimContextInfoIndex);
+		}
+
+		// no need to evaluate the last AnimContextInfoIndex since the condition OtherAnimContextInfoIndex > AnimContextInfoIndex will fail anyways
+		for (int32 AnimContextInfoIndex = 0; AnimContextInfoIndex < NumAnimContextInfos - 1; ++AnimContextInfoIndex)
+		{
+			const float Radius = AnimContextInfos[AnimContextInfoIndex].GetAvailabilitiesMaxBroadPhaseRadius();
+			const FVector& Location = AnimContextInfos[AnimContextInfoIndex].GetLocation();
+			const FVector2f QueryPoint(Location.X, Location.Y);
+
+			const FVector2i min_idx = Indexer.ToGrid(QueryPoint - Radius * FVector2f::One());
+			const FVector2i max_idx = Indexer.ToGrid(QueryPoint + Radius * FVector2f::One());
+
+			const float RadiusSquared = Radius * Radius;
+
+			for (int yi = min_idx.Y; yi <= max_idx.Y; yi++)
 			{
-				EvaluateForNearbyAnimContextInfo(AnimContextInfoIndexA, AnimContextInfoIndexB);
+				for (int xi = min_idx.X; xi <= max_idx.X; xi++)
+				{
+					const FVector2i idx(xi, yi);
+					for (FHashGrid::TConstKeyIterator It = Hash.CreateConstKeyIterator(idx); It; ++It)
+					{
+						const int32 OtherAnimContextInfoIndex = It->Value;
+						// to avoid adding duplicates we evaluate only pairs of indexes where OtherAnimContextInfoIndex > AnimContextInfoIndex
+						if (OtherAnimContextInfoIndex > AnimContextInfoIndex)
+						{
+							EvaluateForNearbyAnimContextInfo(AnimContextInfoIndex, OtherAnimContextInfoIndex);
+						}
+					}
+				}
 			}
 		}
+		break;
+	}
 	}
 }
 
