Fixup API macros

#rnx

#Synced-CL 48757204

diff --git a/Engine/Plugins/Animation/IKRig/Source/IKRig/Public/Retargeter/RetargetOps/AlignPoleVectorOp.h b/Engine/Plugins/Animation/IKRig/Source/IKRig/Public/Retargeter/RetargetOps/AlignPoleVectorOp.h
--- a/Engine/Plugins/Animation/IKRig/Source/IKRig/Public/Retargeter/RetargetOps/AlignPoleVectorOp.h
+++ b/Engine/Plugins/Animation/IKRig/Source/IKRig/Public/Retargeter/RetargetOps/AlignPoleVectorOp.h
@@ -87,8 +87,8 @@ struct FRetargetPoleVectorSettings
 	bool MaintainOffset = false;
 	
 	UE_API bool operator==(const FRetargetPoleVectorSettings& Other) const;
-	UE_API FName GetName() const { return TargetChainName; };
-	UE_API void SetName(const FName InName) { TargetChainName = InName; };
+	FName GetName() const { return TargetChainName; };
+	void SetName(const FName InName) { TargetChainName = InName; };
 };
 
 USTRUCT(BlueprintType, meta = (DisplayName = "Pole Vector Alignment Settings"))
diff --git a/Engine/Plugins/Animation/IKRig/Source/IKRig/Public/Retargeter/RetargetOps/FKChainsOp.h b/Engine/Plugins/Animation/IKRig/Source/IKRig/Public/Retargeter/RetargetOps/FKChainsOp.h
--- a/Engine/Plugins/Animation/IKRig/Source/IKRig/Public/Retargeter/RetargetOps/FKChainsOp.h
+++ b/Engine/Plugins/Animation/IKRig/Source/IKRig/Public/Retargeter/RetargetOps/FKChainsOp.h
@@ -225,8 +225,8 @@ struct FRetargetFKChainSettings
 	double TranslationAlpha = 1.0f;
 
 	IKRIG_API bool operator==(const FRetargetFKChainSettings& Other) const;
-	IKRIG_API FName GetName() const { return TargetChainName; };
-	IKRIG_API void SetName(const FName InName) { TargetChainName = InName; };
+	FName GetName() const { return TargetChainName; };
+	void SetName(const FName InName) { TargetChainName = InName; };
 };
 
 USTRUCT(BlueprintType, meta = (DisplayName = "Retarget Chains Settings"))
diff --git a/Engine/Plugins/Animation/IKRig/Source/IKRig/Public/Retargeter/RetargetOps/IKChainsOp.h b/Engine/Plugins/Animation/IKRig/Source/IKRig/Public/Retargeter/RetargetOps/IKChainsOp.h
--- a/Engine/Plugins/Animation/IKRig/Source/IKRig/Public/Retargeter/RetargetOps/IKChainsOp.h
+++ b/Engine/Plugins/Animation/IKRig/Source/IKRig/Public/Retargeter/RetargetOps/IKChainsOp.h
@@ -134,8 +134,8 @@ struct FRetargetIKChainSettings
 	double Extension = 1.0f;
 
 	UE_API bool operator==(const FRetargetIKChainSettings& Other) const;
-	UE_API FName GetName() const { return TargetChainName; };
-	UE_API void SetName(const FName InName) { TargetChainName = InName; };
+	FName GetName() const { return TargetChainName; };
+	void SetName(const FName InName) { TargetChainName = InName; };
 
 	/** Custom serialization to handle loading old assets */
 	bool Serialize(FArchive& Ar);
diff --git a/Engine/Plugins/Animation/IKRig/Source/IKRig/Public/Retargeter/RetargetOps/SpeedPlantingOp.h b/Engine/Plugins/Animation/IKRig/Source/IKRig/Public/Retargeter/RetargetOps/SpeedPlantingOp.h
--- a/Engine/Plugins/Animation/IKRig/Source/IKRig/Public/Retargeter/RetargetOps/SpeedPlantingOp.h
+++ b/Engine/Plugins/Animation/IKRig/Source/IKRig/Public/Retargeter/RetargetOps/SpeedPlantingOp.h
@@ -27,9 +27,9 @@ struct FRetargetSpeedPlantingSettings
 	UPROPERTY(EditAnywhere, BlueprintReadWrite, BlueprintReadWrite, Category = "Plant IK by Speed", meta = (ClampMin = "0.0", ClampMax = "100.0", UIMin = "0.0", UIMax = "100.0"))
 	FName SpeedCurveName;
 
-	UE_API bool operator==(const FRetargetSpeedPlantingSettings& Other) const { return SpeedCurveName == Other.SpeedCurveName; };
-	UE_API FName GetName() const { return TargetChainName; };
-	UE_API void SetName(const FName InName) { TargetChainName = InName; };
+	bool operator==(const FRetargetSpeedPlantingSettings& Other) const { return SpeedCurveName == Other.SpeedCurveName; };
+	FName GetName() const { return TargetChainName; };
+	void SetName(const FName InName) { TargetChainName = InName; };
 };
 
 USTRUCT(BlueprintType, meta = (DisplayName = "Speed Plant Settings"))
diff --git a/Engine/Plugins/Animation/IKRig/Source/IKRig/Public/Retargeter/RetargetOps/StretchChainOp.h b/Engine/Plugins/Animation/IKRig/Source/IKRig/Public/Retargeter/RetargetOps/StretchChainOp.h
--- a/Engine/Plugins/Animation/IKRig/Source/IKRig/Public/Retargeter/RetargetOps/StretchChainOp.h
+++ b/Engine/Plugins/Animation/IKRig/Source/IKRig/Public/Retargeter/RetargetOps/StretchChainOp.h
@@ -1,4 +1,4 @@
-﻿// Copyright Epic Games, Inc. All Rights Reserved.
+// Copyright Epic Games, Inc. All Rights Reserved.
 #pragma once
 
 #include "Retargeter/IKRetargetChainMapping.h"
@@ -62,8 +62,8 @@ struct FRetargetStretchChainSettings
 	double ScaleChainLength = 1.0f;
 	
 	bool operator==(const FRetargetStretchChainSettings& Other) const;
-	UE_API FName GetName() const { return TargetChainName; };
-	UE_API void SetName(const FName InName) { TargetChainName = InName; };
+	FName GetName() const { return TargetChainName; };
+	void SetName(const FName InName) { TargetChainName = InName; };
 };
 
 USTRUCT(BlueprintType, meta = (DisplayName = "Stretch Chain Settings"))
diff --git a/Engine/Plugins/Animation/IKRig/Source/IKRig/Public/Retargeter/RetargetOps/StrideWarpingOp.h b/Engine/Plugins/Animation/IKRig/Source/IKRig/Public/Retargeter/RetargetOps/StrideWarpingOp.h
--- a/Engine/Plugins/Animation/IKRig/Source/IKRig/Public/Retargeter/RetargetOps/StrideWarpingOp.h
+++ b/Engine/Plugins/Animation/IKRig/Source/IKRig/Public/Retargeter/RetargetOps/StrideWarpingOp.h
@@ -36,8 +36,8 @@ struct FRetargetStrideWarpChainSettings
 	bool EnableStrideWarping = true;
 
 	UE_API bool operator==(const FRetargetStrideWarpChainSettings& Other) const;
-	UE_API FName GetName() const { return TargetChainName; };
-	UE_API void SetName(const FName InName) { TargetChainName = InName; };
+	FName GetName() const { return TargetChainName; };
+	void SetName(const FName InName) { TargetChainName = InName; };
 };
 
 USTRUCT(BlueprintType, meta = (DisplayName = "Stride Warp Settings"))
diff --git a/Engine/Plugins/Animation/PoseSearch/Source/Runtime/Public/PoseSearch/PoseSearchFilter.h b/Engine/Plugins/Animation/PoseSearch/Source/Runtime/Public/PoseSearch/PoseSearchFilter.h
--- a/Engine/Plugins/Animation/PoseSearch/Source/Runtime/Public/PoseSearch/PoseSearchFilter.h
+++ b/Engine/Plugins/Animation/PoseSearch/Source/Runtime/Public/PoseSearch/PoseSearchFilter.h
@@ -65,27 +65,27 @@ struct FPoseSearchFilter
 };
 
 // debug FPoseSearchFilter to test InitializeData / DeinitializeData by padding the data by 1 byte
-USTRUCT(Experimental)
-struct POSESEARCH_API FPoseSearchFilterPadding : public FPoseSearchFilter
+USTRUCT(MinimalAPI, Experimental)
+struct FPoseSearchFilterPadding : public FPoseSearchFilter
 {
 	GENERATED_BODY()
 
-	virtual int32 GetRequiredDataSize(const UPoseSearchFeatureChannel* Channel) const override;
+	UE_API virtual int32 GetRequiredDataSize(const UPoseSearchFeatureChannel* Channel) const override;
 };
 
 // FPoseSearchFilterMaxPosition is invalid if the query for the associated UPoseSearchFeatureChannel_Position is too far from the candidate pose
-USTRUCT(Experimental)
-struct POSESEARCH_API FPoseSearchFilterMaxPosition : public FPoseSearchFilter
+USTRUCT(MinimalAPI, Experimental)
+struct FPoseSearchFilterMaxPosition : public FPoseSearchFilter
 {
 	GENERATED_BODY()
 
-	virtual int32 GetRequiredDataSize(const UPoseSearchFeatureChannel* Channel) const override;
-	virtual bool InitializeData(TArrayView<uint8> Data, const UPoseSearchFeatureChannel* Channel, TConstArrayView<float> QueryValues, UE::PoseSearch::FSearchContext& SearchContext) const override;
-	virtual void DeinitializeData(TArrayView<uint8> Data, const UPoseSearchFeatureChannel* Channel, UE::PoseSearch::FSearchContext& SearchContext) const override;
-	virtual bool IsFilterValid(TConstArrayView<float> PoseValues, TConstArrayView<float> QueryValues, int32 PoseIdx, const UE::PoseSearch::FPoseMetadata& Metadata, TConstArrayView<uint8> Data, const UPoseSearchFeatureChannel* Channel) const override;
+	UE_API virtual int32 GetRequiredDataSize(const UPoseSearchFeatureChannel* Channel) const override;
+	UE_API virtual bool InitializeData(TArrayView<uint8> Data, const UPoseSearchFeatureChannel* Channel, TConstArrayView<float> QueryValues, UE::PoseSearch::FSearchContext& SearchContext) const override;
+	UE_API virtual void DeinitializeData(TArrayView<uint8> Data, const UPoseSearchFeatureChannel* Channel, UE::PoseSearch::FSearchContext& SearchContext) const override;
+	UE_API virtual bool IsFilterValid(TConstArrayView<float> PoseValues, TConstArrayView<float> QueryValues, int32 PoseIdx, const UE::PoseSearch::FPoseMetadata& Metadata, TConstArrayView<uint8> Data, const UPoseSearchFeatureChannel* Channel) const override;
 	
 #if ENABLE_DRAW_DEBUG
-	virtual void DebugDraw(const UE::PoseSearch::FDebugDrawParams& DrawParams, TConstArrayView<float> PoseValues, const UPoseSearchFeatureChannel* Channel) const override;
+	UE_API virtual void DebugDraw(const UE::PoseSearch::FDebugDrawParams& DrawParams, TConstArrayView<float> PoseValues, const UPoseSearchFeatureChannel* Channel) const override;
 #endif // ENABLE_DRAW_DEBUG
 
 	// During pose selection if the squared distance between query versus candidate poses for the associated 
@@ -103,18 +103,18 @@ private:
 };
 
 // FPoseSearchFilterWedge is invalid if the query for the associated UPoseSearchFeatureChannel_Position doesn't land in the defined wedge relative to the trajectory
-USTRUCT(Experimental)
-struct POSESEARCH_API FPoseSearchFilterWedge : public FPoseSearchFilter
+USTRUCT(MinimalAPI, Experimental)
+struct FPoseSearchFilterWedge : public FPoseSearchFilter
 {
 	GENERATED_BODY()
 
-	virtual int32 GetRequiredDataSize(const UPoseSearchFeatureChannel* Channel) const override;
-	virtual bool InitializeData(TArrayView<uint8> Data, const UPoseSearchFeatureChannel* Channel, TConstArrayView<float> QueryValues, UE::PoseSearch::FSearchContext& SearchContext) const override;
-	virtual void DeinitializeData(TArrayView<uint8> Data, const UPoseSearchFeatureChannel* Channel, UE::PoseSearch::FSearchContext& SearchContext) const override;
-	virtual bool IsFilterValid(TConstArrayView<float> PoseValues, TConstArrayView<float> QueryValues, int32 PoseIdx, const UE::PoseSearch::FPoseMetadata& Metadata, TConstArrayView<uint8> Data, const UPoseSearchFeatureChannel* Channel) const override;
+	UE_API virtual int32 GetRequiredDataSize(const UPoseSearchFeatureChannel* Channel) const override;
+	UE_API virtual bool InitializeData(TArrayView<uint8> Data, const UPoseSearchFeatureChannel* Channel, TConstArrayView<float> QueryValues, UE::PoseSearch::FSearchContext& SearchContext) const override;
+	UE_API virtual void DeinitializeData(TArrayView<uint8> Data, const UPoseSearchFeatureChannel* Channel, UE::PoseSearch::FSearchContext& SearchContext) const override;
+	UE_API virtual bool IsFilterValid(TConstArrayView<float> PoseValues, TConstArrayView<float> QueryValues, int32 PoseIdx, const UE::PoseSearch::FPoseMetadata& Metadata, TConstArrayView<uint8> Data, const UPoseSearchFeatureChannel* Channel) const override;
 
 #if ENABLE_DRAW_DEBUG
-	virtual void DebugDraw(const UE::PoseSearch::FDebugDrawParams& DrawParams, TConstArrayView<float> PoseValues, const UPoseSearchFeatureChannel* Channel) const override;
+	UE_API virtual void DebugDraw(const UE::PoseSearch::FDebugDrawParams& DrawParams, TConstArrayView<float> PoseValues, const UPoseSearchFeatureChannel* Channel) const override;
 #endif // ENABLE_DRAW_DEBUG
 
 	UPROPERTY(EditAnywhere, Category = "Experimental|Filter", meta = (ClampToMinMaxLimits))
diff --git a/Engine/Plugins/Experimental/Animation/ContextualAnimation/Source/ContextualAnimation/Public/ContextualAnimTypes.h b/Engine/Plugins/Experimental/Animation/ContextualAnimation/Source/ContextualAnimation/Public/ContextualAnimTypes.h
--- a/Engine/Plugins/Experimental/Animation/ContextualAnimation/Source/ContextualAnimation/Public/ContextualAnimTypes.h
+++ b/Engine/Plugins/Experimental/Animation/ContextualAnimation/Source/ContextualAnimation/Public/ContextualAnimTypes.h
@@ -631,7 +631,7 @@ struct FContextualAnimSceneBindings
 
 	UE_API void SetOverrideProvider(const TScriptInterface <IContextualAnimOverrideInterface>& Provider);
 	
-	UE_API const TScriptInterface <IContextualAnimOverrideInterface>& GetOverrideProvider() const { return OverrideProviderInterface;}
+	const TScriptInterface <IContextualAnimOverrideInterface>& GetOverrideProvider() const { return OverrideProviderInterface;}
 
 	/*
 	* Updates Each Bindings Overrides, Including Animation, Alignment IK, ETC 
diff --git a/Engine/Plugins/Experimental/Animation/RigMapper/Source/RigMapper/Public/RigMapperDefinition.h b/Engine/Plugins/Experimental/Animation/RigMapper/Source/RigMapper/Public/RigMapperDefinition.h
--- a/Engine/Plugins/Experimental/Animation/RigMapper/Source/RigMapper/Public/RigMapperDefinition.h
+++ b/Engine/Plugins/Experimental/Animation/RigMapper/Source/RigMapper/Public/RigMapperDefinition.h
@@ -327,9 +327,9 @@ public:
 
 	virtual ERigMapperFeatureType GetFeatureType() const override { return ERigMapperFeatureType::NullOutput; }
 
-	UE_API virtual void RelinkFeature(const FString& InNewInput, const int32 InputIndex) override {};
+	virtual void RelinkFeature(const FString& InNewInput, const int32 InputIndex) override {};
 
-	UE_API virtual void AddNewInputLink(const FString& InNewInput) override {};
+	virtual void AddNewInputLink(const FString& InNewInput) override {};
 };
 
 USTRUCT(BlueprintType)
diff --git a/Engine/Plugins/Experimental/Animation/RigMapper/Source/RigMapper/Public/RigMapperUtils.h b/Engine/Plugins/Experimental/Animation/RigMapper/Source/RigMapper/Public/RigMapperUtils.h
--- a/Engine/Plugins/Experimental/Animation/RigMapper/Source/RigMapper/Public/RigMapperUtils.h
+++ b/Engine/Plugins/Experimental/Animation/RigMapper/Source/RigMapper/Public/RigMapperUtils.h
@@ -9,7 +9,7 @@
 class URigMapperDefinition;
 enum class ERigMapperFeatureType : uint8;
 
-class UE_API FRigMapperUtils
+class FRigMapperUtils
 {
 public:
 	/**
@@ -20,7 +20,7 @@ public:
 	 * @param InFeatureType		Type of the feature being generated to include/replace as a prefix.
 	 * @return					Returns the unique name.
 	 */
-	static FString GenerateUniqueFeatureName(const TWeakObjectPtr<URigMapperDefinition> InDefinition, FString InDesiredName, ERigMapperFeatureType InFeatureType);
+	static UE_API FString GenerateUniqueFeatureName(const TWeakObjectPtr<URigMapperDefinition> InDefinition, FString InDesiredName, ERigMapperFeatureType InFeatureType);
 };
 
-#undef UE_API
\ No newline at end of file
+#undef UE_API
diff --git a/Engine/Plugins/Experimental/MessageBusTester/Source/MessageBusTester/Public/IMessageBusTesterLogger.h b/Engine/Plugins/Experimental/MessageBusTester/Source/MessageBusTester/Public/IMessageBusTesterLogger.h
--- a/Engine/Plugins/Experimental/MessageBusTester/Source/MessageBusTester/Public/IMessageBusTesterLogger.h
+++ b/Engine/Plugins/Experimental/MessageBusTester/Source/MessageBusTester/Public/IMessageBusTesterLogger.h
@@ -22,7 +22,7 @@ struct FMessageBusTesterLogEntry
 };
 
 
-class MESSAGEBUSTESTER_API IMessageBusTesterLogger
+class IMessageBusTesterLogger
 {
 public:
 
diff --git a/Engine/Plugins/Experimental/MessageBusTester/Source/MessageBusTester/Public/MessageBusTesterCommon.h b/Engine/Plugins/Experimental/MessageBusTester/Source/MessageBusTester/Public/MessageBusTesterCommon.h
--- a/Engine/Plugins/Experimental/MessageBusTester/Source/MessageBusTester/Public/MessageBusTesterCommon.h
+++ b/Engine/Plugins/Experimental/MessageBusTester/Source/MessageBusTester/Public/MessageBusTesterCommon.h
@@ -104,8 +104,8 @@ struct FTesterDiscoveryMessage : public FTesterBaseMessage
 	GENERATED_BODY()
 
 public:
-	UE_API FTesterDiscoveryMessage() = default;
-	UE_API FTesterDiscoveryMessage(FTesterInstanceDescriptor&& InDescriptor)
+	FTesterDiscoveryMessage() = default;
+	FTesterDiscoveryMessage(FTesterInstanceDescriptor&& InDescriptor)
 		: Descriptor(MoveTemp(InDescriptor))
 	{}
 
diff --git a/Engine/Plugins/PCG/Source/PCG/Public/PCGModule.h b/Engine/Plugins/PCG/Source/PCG/Public/PCGModule.h
--- a/Engine/Plugins/PCG/Source/PCG/Public/PCGModule.h
+++ b/Engine/Plugins/PCG/Source/PCG/Public/PCGModule.h
@@ -69,8 +69,8 @@ public:
 	static const FPCGGraphExecutionRegistry& GetConstGraphExecutionRegistry() { return GetPCGModuleChecked().GraphExecutionRegistry; }
 	static FPCGGraphExecutionRegistry& GetGraphExecutionRegistry() { return GetPCGModuleChecked().GraphExecutionRegistry; }
 
-	PCG_API static const FPCGDataViewRegistry& GetConstPCGDataViewRegistry() { return GetPCGModuleChecked().DataViewRegistry; }
-	PCG_API static FPCGDataViewRegistry& GetMutablePCGDataViewRegistry() { return GetPCGModuleChecked().DataViewRegistry; }
+	static const FPCGDataViewRegistry& GetConstPCGDataViewRegistry() { return GetPCGModuleChecked().DataViewRegistry; }
+	static FPCGDataViewRegistry& GetMutablePCGDataViewRegistry() { return GetPCGModuleChecked().DataViewRegistry; }
 
 	static const FPCGChangeTrackingRegistry& GetConstChangeTrackingRegistry() { return GetPCGModuleChecked().ChangeTrackingRegistry; }
 	static FPCGChangeTrackingRegistry& GetMutableChangeTrackingRegistry() { return GetPCGModuleChecked().ChangeTrackingRegistry; }
diff --git a/Engine/Source/Runtime/AnimGraphRuntime/Public/IAnimGraphRuntime_SequencerMixerTargetConnector.h b/Engine/Source/Runtime/AnimGraphRuntime/Public/IAnimGraphRuntime_SequencerMixerTargetConnector.h
--- a/Engine/Source/Runtime/AnimGraphRuntime/Public/IAnimGraphRuntime_SequencerMixerTargetConnector.h
+++ b/Engine/Source/Runtime/AnimGraphRuntime/Public/IAnimGraphRuntime_SequencerMixerTargetConnector.h
@@ -6,7 +6,7 @@
 
 struct FPoseContext;
 
-class ANIMGRAPHRUNTIME_API IAnimGraphRuntime_SequencerMixerTargetConnector : public IModularFeature
+class IAnimGraphRuntime_SequencerMixerTargetConnector : public IModularFeature
 {
 public:
 
@@ -28,4 +28,4 @@ public:
 	virtual ~IAnimGraphRuntime_SequencerMixerTargetConnector() = default;
 	
 	virtual void ApplySequencerMixedPose(FPoseContext& Output, FName InTargetName) const = 0;
-};
\ No newline at end of file
+};
diff --git a/Engine/Source/Runtime/SignalProcessing/Public/DSP/SphericalHarmonicCalculator.h b/Engine/Source/Runtime/SignalProcessing/Public/DSP/SphericalHarmonicCalculator.h
--- a/Engine/Source/Runtime/SignalProcessing/Public/DSP/SphericalHarmonicCalculator.h
+++ b/Engine/Source/Runtime/SignalProcessing/Public/DSP/SphericalHarmonicCalculator.h
@@ -39,7 +39,7 @@ public:
 		Swap(InOutVector.X,InOutVector.Y);
 	}
 	
-	static SIGNALPROCESSING_API int64 OrderToNumChannels(const int32 InOrder)
+	static int64 OrderToNumChannels(const int32 InOrder)
 	{
 		check(InOrder >= 0);
 		return (InOrder+1) * (InOrder+1);
