Fix issue where GetPoseIndex could return an invalid pose index when there are multiple sampling ranges on the same animation asset
#Synced-CL 48841652

diff --git a/Engine/Plugins/Animation/PoseSearch/Source/Runtime/Private/PoseSearchDatabase.cpp b/Engine/Plugins/Animation/PoseSearch/Source/Runtime/Private/PoseSearchDatabase.cpp
--- a/Engine/Plugins/Animation/PoseSearch/Source/Runtime/Private/PoseSearchDatabase.cpp
+++ b/Engine/Plugins/Animation/PoseSearch/Source/Runtime/Private/PoseSearchDatabase.cpp
@@ -942,18 +942,20 @@ int32 UPoseSearchDatabase::GetPoseIndex(const UObject* AnimationAsset, float Ani
 			{
 				const float BlendParametersSquaredLength = (BlendParameters - SearchIndexAsset.GetBlendParameters()).SquaredLength();
 
-				// using <= so we don't have to check for PoseIdx == INDEX_NONE, since any float will be smaller or equal than UE_MAX_FLT
 				if (BlendParametersSquaredLength <= MinSquaredLength)
 				{
-					MinSquaredLength = BlendParametersSquaredLength;
-
 					const FPoseSearchDatabaseAnimationAssetBase* DatabaseAnimationAssetBase = GetDatabaseAnimationAsset(SearchIndexAsset);
 
 					check(DatabaseAnimationAssetBase);
 					check(DatabaseAnimationAssetBase->GetAnimationAsset() == AnimationAsset);
 
 					const float RealAssetTime = AnimationAssetTime * SearchIndexAsset.GetToRealTimeFactor();
-					PoseIdx = SearchIndexAsset.GetPoseIndexFromTime(RealAssetTime, SampleRate);
+					const int32 CurrentIdx = SearchIndexAsset.GetPoseIndexFromTime(RealAssetTime, SampleRate);
+					if (CurrentIdx != INDEX_NONE)
+					{
+						MinSquaredLength = BlendParametersSquaredLength;
+						PoseIdx = CurrentIdx;
+					}
 				}
 			}
 		}
