PoseSearch - database cooking fix

#Synced-CL 49553855

diff --git a/Engine/Plugins/Animation/PoseSearch/Source/Runtime/Private/Chooser/PoseSearchChooserColumn.cpp b/Engine/Plugins/Animation/PoseSearch/Source/Runtime/Private/Chooser/PoseSearchChooserColumn.cpp
--- a/Engine/Plugins/Animation/PoseSearch/Source/Runtime/Private/Chooser/PoseSearchChooserColumn.cpp
+++ b/Engine/Plugins/Animation/PoseSearch/Source/Runtime/Private/Chooser/PoseSearchChooserColumn.cpp
@@ -227,6 +227,11 @@ void FPoseSearchColumn::DeleteRows(TArrayView<int> RowIndices)
 {
 	using namespace UE::PoseSearch;
 
+	if (RowIndices.IsEmpty())
+	{
+		return;
+	}
+
 	check(InternalDatabase);
 	InternalDatabase->Modify();
 	
@@ -242,6 +247,8 @@ void FPoseSearchColumn::DeleteRows(TArrayView<int> RowIndices)
 		FallbackDatabaseAnimationAsset->AnimAsset = GetReferencedObject(ChooserColumn_SpecialIndex_Fallback);
 	}
 
+	InternalDatabase->PostEditChange();
+
 	InternalDatabase->NotifySynchronizeWithExternalDependencies();
 	CheckConsistency();
 }
diff --git a/Engine/Plugins/Animation/PoseSearch/Source/Runtime/Private/PoseSearchDatabase.cpp b/Engine/Plugins/Animation/PoseSearch/Source/Runtime/Private/PoseSearchDatabase.cpp
--- a/Engine/Plugins/Animation/PoseSearch/Source/Runtime/Private/PoseSearchDatabase.cpp
+++ b/Engine/Plugins/Animation/PoseSearch/Source/Runtime/Private/PoseSearchDatabase.cpp
@@ -1248,6 +1248,9 @@ void UPoseSearchDatabase::SynchronizeWithExternalDependencies(TConstArrayView<UA
 		Modify();
 
 		DatabaseAnimationAssets = NewDatabaseAnimationAssets;
+
+		PostEditChange();
+
 		NotifySynchronizeWithExternalDependencies();
 	}
 }
