PoseSearch - fix for UPoseSearchLibrary::MotionMatch crashing when reusing an invalid continuing pose from a previously evaluated schema

#Synced-CL 49553855

diff --git a/Engine/Plugins/Animation/PoseSearch/Source/Runtime/Private/PoseSearchLibrary.cpp b/Engine/Plugins/Animation/PoseSearch/Source/Runtime/Private/PoseSearchLibrary.cpp
--- a/Engine/Plugins/Animation/PoseSearch/Source/Runtime/Private/PoseSearchLibrary.cpp
+++ b/Engine/Plugins/Animation/PoseSearch/Source/Runtime/Private/PoseSearchLibrary.cpp
@@ -60,19 +60,21 @@ namespace UE::PoseSearch
 	// used to cache the continuing pose search results
 	struct FCachedContinuingPoseSearchResults : public TMap<const UObject*, FSearchResult, TInlineSetAllocator<16, TMemStackSetAllocator<>>>
 	{
-		const FSearchResult& CheckedAdd(const UObject* Object, const FSearchResult& SearchResult)
+		void CheckedAdd(const UObject* Object, const FSearchResult& SearchResult)
 		{
 			check(Object);
 			check(!Find(Object));
+			check(SearchResult.IsValid());
 
 			FSearchResult& NewSearchResult = Add(Object);
 			NewSearchResult = SearchResult;
-			return NewSearchResult;
 		}
 
 		const FSearchResult& FindOrAdd(const UObject* Object, const FSearchResult& SearchResult)
 		{
 			check(Object);
+			check(SearchResult.IsValid());
+
 			if (const FSearchResult* FoundSearchResult = Find(Object))
 			{
 				return *FoundSearchResult;
@@ -1335,23 +1337,28 @@ void UPoseSearchLibrary::MotionMatch(
 		const bool bInvalidatingContinuingPose = IsInvalidatingContinuingPose(ContinuingProperties.InterruptMode, Database, AssetsToSearchPerDatabase);
 		if (!bInvalidatingContinuingPose)
 		{
-			// reconstructing and caching all the required continuing pose search results
-			FSearchResult DatabaseContinuingPoseSearchResult;
-			DatabaseContinuingPoseSearchResult.SetAssetTime(ContinuingProperties.PlayingAssetAccumulatedTime);
-			DatabaseContinuingPoseSearchResult.PoseIdx = Database->GetPoseIndex(ContinuingProperties.PlayingAsset.Get(), ContinuingProperties.PlayingAssetAccumulatedTime, ContinuingProperties.bIsPlayingAssetMirrored, ContinuingProperties.PlayingAssetBlendParameters);
-			DatabaseContinuingPoseSearchResult.Database = Database;
-			CachedContinuingPoseSearchResults.CheckedAdd(Database, DatabaseContinuingPoseSearchResult);
-
-			// adding the continuing pose search redult relative to the schema - first instance of the DatabaseContinuingPoseSearchResult 
-			// (used to gather the continuing pose search values adopted to create the MM query - relaive to the schema, NOT the database)
-			const FSearchResult& SchemaContinuingPoseSearchResult = CachedContinuingPoseSearchResults.FindOrAdd(Database->Schema, DatabaseContinuingPoseSearchResult);
-
-			const bool bForceInterrupt = IsForceInterrupt(ContinuingProperties.InterruptMode, Database, AssetsToSearchPerDatabase);
-			const bool bCanAdvance = DatabaseContinuingPoseSearchResult.PoseIdx != INDEX_NONE;
-			if (bCanAdvance && !bForceInterrupt)
+			const int32 ContinuingPoseIdx = Database->GetPoseIndex(ContinuingProperties.PlayingAsset.Get(), ContinuingProperties.PlayingAssetAccumulatedTime, ContinuingProperties.bIsPlayingAssetMirrored, ContinuingProperties.PlayingAssetBlendParameters);
+
+			if (ContinuingPoseIdx != INDEX_NONE)
 			{
-				SearchContext.UpdateContinuingPoseSearchResult(DatabaseContinuingPoseSearchResult, SchemaContinuingPoseSearchResult);
-				Database->SearchContinuingPose(SearchContext, SearchResults);
+				// reconstructing and caching all the required continuing pose search results
+				FSearchResult DatabaseContinuingPoseSearchResult;
+				DatabaseContinuingPoseSearchResult.SetAssetTime(ContinuingProperties.PlayingAssetAccumulatedTime);
+				DatabaseContinuingPoseSearchResult.PoseIdx = ContinuingPoseIdx;
+				DatabaseContinuingPoseSearchResult.Database = Database;
+
+				CachedContinuingPoseSearchResults.CheckedAdd(Database, DatabaseContinuingPoseSearchResult);
+
+				// adding the continuing pose search result relative to the schema - first instance of the DatabaseContinuingPoseSearchResult 
+				// (used to gather the continuing pose search values adopted to create the MM query - relative to the schema, NOT the database)
+				const FSearchResult& SchemaContinuingPoseSearchResult = CachedContinuingPoseSearchResults.FindOrAdd(Database->Schema, DatabaseContinuingPoseSearchResult);
+
+				const bool bForceInterrupt = IsForceInterrupt(ContinuingProperties.InterruptMode, Database, AssetsToSearchPerDatabase);
+				if (!bForceInterrupt)
+				{
+					SearchContext.UpdateContinuingPoseSearchResult(DatabaseContinuingPoseSearchResult, SchemaContinuingPoseSearchResult);
+					Database->SearchContinuingPose(SearchContext, SearchResults);
+				}
 			}
 		}
 	}
