PoseSearch - adding MMI OnInteractionStartedDelegate OnInteractionContinuingDelegate OnInteractionEndedDelegate

#Synced-CL 49608181

diff --git a/Engine/Plugins/Animation/PoseSearch/Source/Runtime/Private/PoseSearchInteractionIsland.cpp b/Engine/Plugins/Animation/PoseSearch/Source/Runtime/Private/PoseSearchInteractionIsland.cpp
--- a/Engine/Plugins/Animation/PoseSearch/Source/Runtime/Private/PoseSearchInteractionIsland.cpp
+++ b/Engine/Plugins/Animation/PoseSearch/Source/Runtime/Private/PoseSearchInteractionIsland.cpp
@@ -683,7 +683,7 @@ void FValidInteractionSearch::AddReferencedObjects(FReferenceCollector& Collecto
 
 // FInteractionSearchContext
 ///////////////////////////////////////////////////////////
-void FValidInteractionSearches::Update(TConstArrayView<FInteractionIsland*> Islands)
+void FValidInteractionSearches::Update(const UPoseSearchInteractionSubsystem* InteractionSubsystem)
 {
 #if ENABLE_ANIM_DEBUG
 	UE_MT_SCOPED_WRITE_ACCESS(ThreadSafeCounter);
@@ -697,7 +697,7 @@ void FValidInteractionSearches::Update(TConstArrayView<FInteractionIsland*> Isla
 
 	FMemMark Mark(FMemStack::Get());
 	TMap<uint32, bool, TInlineSetAllocator<32, TMemStackSetAllocator<>>> ContinuingOrNewValidSearchContexts;
-	for (const FInteractionIsland* Island : Islands)
+	for (const FInteractionIsland* Island : InteractionSubsystem->GetIslands())
 	{
 		if (Island->IsInitialized())
 		{
@@ -728,7 +728,7 @@ void FValidInteractionSearches::Update(TConstArrayView<FInteractionIsland*> Isla
 	{
 		if (!ContinuingOrNewValidSearchContexts.Find(It->Key))
 		{
-			OnInteractionEnd(It->Value);
+			OnInteractionEnd(It->Value, InteractionSubsystem);
 			It.RemoveCurrent();
 		}
 	}
@@ -739,38 +739,50 @@ void FValidInteractionSearches::Update(TConstArrayView<FInteractionIsland*> Isla
 		const bool bIsContinuingValidSearchContexts = ContinuingOrNewValidSearchContexts.FindChecked(It->Key);
 		if (bIsContinuingValidSearchContexts)
 		{
-			OnInteractionContinuing(It->Value);
+			OnInteractionContinuing(It->Value, InteractionSubsystem);
 		}
 		else
 		{
-			OnInteractionStart(It->Value);
+			OnInteractionStart(It->Value, InteractionSubsystem);
 		}
 	}
 }
 
-void FValidInteractionSearches::OnInteractionStart(FValidInteractionSearch& ValidInteractionSearch)
+void FValidInteractionSearches::OnInteractionStart(FValidInteractionSearch& ValidInteractionSearch, const UPoseSearchInteractionSubsystem* InteractionSubsystem)
 {
+	check(InteractionSubsystem);
+
 #if ENABLE_VISUAL_LOG
 	ValidInteractionSearch.VLogContext(FColor::Blue);
 #endif
 
 	ValidInteractionSearch.DoDisableCollisions();
+
+	InteractionSubsystem->OnInteractionStartedDelegate.Broadcast(ValidInteractionSearch);
 }
 
-void FValidInteractionSearches::OnInteractionContinuing(FValidInteractionSearch& ValidInteractionSearch)
+void FValidInteractionSearches::OnInteractionContinuing(FValidInteractionSearch& ValidInteractionSearch, const UPoseSearchInteractionSubsystem* InteractionSubsystem)
 {
+	check(InteractionSubsystem);
+
 #if ENABLE_VISUAL_LOG
 	ValidInteractionSearch.VLogContext(FColor::Green);
 #endif
+
+	InteractionSubsystem->OnInteractionContinuingDelegate.Broadcast(ValidInteractionSearch);
 }
 
-void FValidInteractionSearches::OnInteractionEnd(FValidInteractionSearch& ValidInteractionSearch)
+void FValidInteractionSearches::OnInteractionEnd(FValidInteractionSearch& ValidInteractionSearch, const UPoseSearchInteractionSubsystem* InteractionSubsystem)
 {
+	check(InteractionSubsystem);
+
 #if ENABLE_VISUAL_LOG
 	ValidInteractionSearch.VLogContext(FColor::Black);
 #endif // ENABLE_VISUAL_LOG
 
 	ValidInteractionSearch.UndoDisableCollisions();
+
+	InteractionSubsystem->OnInteractionEndedDelegate.Broadcast(ValidInteractionSearch);
 }
 
 const FValidInteractionSearch* FValidInteractionSearches::Get(const FInteractionSearchContext& SearchContext) const
diff --git a/Engine/Plugins/Animation/PoseSearch/Source/Runtime/Private/PoseSearchInteractionSubsystem.cpp b/Engine/Plugins/Animation/PoseSearch/Source/Runtime/Private/PoseSearchInteractionSubsystem.cpp
--- a/Engine/Plugins/Animation/PoseSearch/Source/Runtime/Private/PoseSearchInteractionSubsystem.cpp
+++ b/Engine/Plugins/Animation/PoseSearch/Source/Runtime/Private/PoseSearchInteractionSubsystem.cpp
@@ -577,14 +577,20 @@ UPoseSearchInteractionSubsystem* UPoseSearchInteractionSubsystem::GetSubsystem_A
 {
 	if (AnimContext)
 	{
-		if (UWorld* World = AnimContext->GetWorld())
+		return GetSubsystem_AnyThread(AnimContext->GetWorld());
+	}
+	return nullptr;
+}
+
+UPoseSearchInteractionSubsystem* UPoseSearchInteractionSubsystem::GetSubsystem_AnyThread(const UWorld* World)
+{
+	if (World)
+	{
+		// We expect the subsystem to be already created from the GameThread.
+		// We don't create the subsystem from any thread
+		if (World->HasSubsystem<UPoseSearchInteractionSubsystem>())
 		{
-			// We expect the subsystem to be already created from the GameThread.
-			// We don't create the subsystem from any thread
-			if (World->HasSubsystem<UPoseSearchInteractionSubsystem>())
-			{
-				return World->GetSubsystem<UPoseSearchInteractionSubsystem>();
-			}
+			return World->GetSubsystem<UPoseSearchInteractionSubsystem>();
 		}
 	}
 	return nullptr;
@@ -598,7 +604,7 @@ void UPoseSearchInteractionSubsystem::Initialize(FSubsystemCollectionBase& Colle
 
 void UPoseSearchInteractionSubsystem::Deinitialize()
 {
-	ValidInteractionSearches.Update(Islands);
+	ValidInteractionSearches.Update(this);
 	DestroyAllIslands();
 	Super::Deinitialize();
 }
@@ -1402,7 +1408,7 @@ void UPoseSearchInteractionSubsystem::Tick(float DeltaSeconds)
 	Super::Tick(DeltaSeconds);
 
 	FMemMark Mark(FMemStack::Get());
-	ValidInteractionSearches.Update(Islands);
+	ValidInteractionSearches.Update(this);
 
 	if (ShouldRegenerateAllIslands())
 	{
diff --git a/Engine/Plugins/Animation/PoseSearch/Source/Runtime/Public/PoseSearch/PoseSearchInteractionIsland.h b/Engine/Plugins/Animation/PoseSearch/Source/Runtime/Public/PoseSearch/PoseSearchInteractionIsland.h
--- a/Engine/Plugins/Animation/PoseSearch/Source/Runtime/Public/PoseSearch/PoseSearchInteractionIsland.h
+++ b/Engine/Plugins/Animation/PoseSearch/Source/Runtime/Public/PoseSearch/PoseSearchInteractionIsland.h
@@ -113,16 +113,16 @@ private:
 
 struct FValidInteractionSearches
 {
-	void Update(TConstArrayView<FInteractionIsland*> Islands);
+	void Update(const UPoseSearchInteractionSubsystem* InteractionSubsystem);
 	const FValidInteractionSearch* Get(const FInteractionSearchContext& SearchContext) const;
 	// returns true if the tuple of AnimContexts (e.g. pair of AnimInstance(s) ~ Characters) from the SearchContext is used in any FValidInteractionSearch with any assigned roles
 	bool IsThereAnyValidSearchWithTheSameContexts(const FInteractionSearchContext& SearchContext) const;
 	void AddReferencedObjects(FReferenceCollector& Collector);
 
 private:
-	static void OnInteractionStart(FValidInteractionSearch& ValidInteractionSearch);
-	static void OnInteractionContinuing(FValidInteractionSearch& ValidInteractionSearch);
-	static void OnInteractionEnd(FValidInteractionSearch& ValidInteractionSearch);
+	static void OnInteractionStart(FValidInteractionSearch& ValidInteractionSearch, const UPoseSearchInteractionSubsystem* InteractionSubsystem);
+	static void OnInteractionContinuing(FValidInteractionSearch& ValidInteractionSearch, const UPoseSearchInteractionSubsystem* InteractionSubsystem);
+	static void OnInteractionEnd(FValidInteractionSearch& ValidInteractionSearch, const UPoseSearchInteractionSubsystem* InteractionSubsystem);
 
 	// mapping from FValidInteractionSearch::GetHashForEquivalence to FValidInteractionSearch
 	TMap<uint32, FValidInteractionSearch> Searches;
diff --git a/Engine/Plugins/Animation/PoseSearch/Source/Runtime/Public/PoseSearch/PoseSearchInteractionSubsystem.h b/Engine/Plugins/Animation/PoseSearch/Source/Runtime/Public/PoseSearch/PoseSearchInteractionSubsystem.h
--- a/Engine/Plugins/Animation/PoseSearch/Source/Runtime/Public/PoseSearch/PoseSearchInteractionSubsystem.h
+++ b/Engine/Plugins/Animation/PoseSearch/Source/Runtime/Public/PoseSearch/PoseSearchInteractionSubsystem.h
@@ -20,6 +20,8 @@ struct FAnimContextInfos;
 struct FTagToDatabases;
 } // namespace UE::PoseSearch
 
+DECLARE_MULTICAST_DELEGATE_OneParam(FOnPoseSearchInteractionEvent, const UE::PoseSearch::FValidInteractionSearch&);
+
 // World subsystem accepting the publication of characters (via their AnimInstance(s)) FPoseSearchInteractionAvailability, representing the characters willingness to partecipate in an 
 // interaction with other characters from the next frame forward via Query_AnyThread method.
 //
@@ -79,6 +81,7 @@ class UPoseSearchInteractionSubsystem : public UTickableWorldSubsystem
 
 public:
 	static UE_API UPoseSearchInteractionSubsystem* GetSubsystem_AnyThread(const UObject* AnimInstance);
+	static UE_API UPoseSearchInteractionSubsystem* GetSubsystem_AnyThread(const UWorld* World);
 
 	UE_API virtual void Initialize(FSubsystemCollectionBase& Collection) override;
 	UE_API virtual void Deinitialize() override;
@@ -112,6 +115,12 @@ public:
 	// for example when using this API from a different ABP from a different skeletal mesh component
 	bool GetConstraint(const UObject* AnimContext, FName SocketName, float& OutDesiredReach, FTransform& OutTransform, bool bCompareOwningActors = false);
 
+	TConstArrayView<UE::PoseSearch::FInteractionIsland*> GetIslands() const { return Islands; }
+
+	FOnPoseSearchInteractionEvent OnInteractionStartedDelegate;
+	FOnPoseSearchInteractionEvent OnInteractionContinuingDelegate;
+	FOnPoseSearchInteractionEvent OnInteractionEndedDelegate;
+
 private:
 #if ENABLE_ANIM_DEBUG
 	friend struct UE::PoseSearch::FInteractionValidator;
