PoseSearch - fixed a check() failure in UPoseSearchInteractionSubsystem that asserted Availabilities must not be empty

#Synced-CL 49608181

diff --git a/Engine/Plugins/Animation/PoseSearch/Source/Runtime/Private/PoseSearchInteractionSubsystem.cpp b/Engine/Plugins/Animation/PoseSearch/Source/Runtime/Private/PoseSearchInteractionSubsystem.cpp
--- a/Engine/Plugins/Animation/PoseSearch/Source/Runtime/Private/PoseSearchInteractionSubsystem.cpp
+++ b/Engine/Plugins/Animation/PoseSearch/Source/Runtime/Private/PoseSearchInteractionSubsystem.cpp
@@ -615,16 +615,19 @@ void UPoseSearchInteractionSubsystem::AddAvailabilities(const TArrayView<const F
 
 	check(AnimContext && AnimContext->GetWorld() && AnimContext->GetWorld() == GetWorld());
 
-	const int32 ReservedAnimContextsAvailabilitiesIndex = AnimContextsAvailabilitiesIndex.Add(1);
-
-	if (ReservedAnimContextsAvailabilitiesIndex < AnimContextsAvailabilitiesBuffer.Num())
+	if (!Availabilities.IsEmpty())
 	{
-		FInteractionAnimContextAvailabilities& AnimContextAvailabilities = AnimContextsAvailabilitiesBuffer[ReservedAnimContextsAvailabilitiesIndex];
-		AnimContextAvailabilities.AnimContext = AnimContext;
-		AnimContextAvailabilities.Availabilities.Reset();
-		for (const FPoseSearchInteractionAvailability& Availability : Availabilities)
+		const int32 ReservedAnimContextsAvailabilitiesIndex = AnimContextsAvailabilitiesIndex.Add(1);
+
+		if (ReservedAnimContextsAvailabilitiesIndex < AnimContextsAvailabilitiesBuffer.Num())
 		{
-			AnimContextAvailabilities.Availabilities.AddDefaulted_GetRef().Init(Availability, PoseHistoryName, PoseHistory);
+			FInteractionAnimContextAvailabilities& AnimContextAvailabilities = AnimContextsAvailabilitiesBuffer[ReservedAnimContextsAvailabilitiesIndex];
+			AnimContextAvailabilities.AnimContext = AnimContext;
+			AnimContextAvailabilities.Availabilities.Reset();
+			for (const FPoseSearchInteractionAvailability& Availability : Availabilities)
+			{
+				AnimContextAvailabilities.Availabilities.AddDefaulted_GetRef().Init(Availability, PoseHistoryName, PoseHistory);
+			}
 		}
 	}
 }
