PoseSearch - filtering assets by their index in the database, rather than just by their object, to support PoseSearch column filtering by index

#Synced-CL 49608181

diff --git a/Engine/Plugins/Animation/PoseSearch/Source/Editor/Private/Chooser/PoseSearchColumnEditor.cpp b/Engine/Plugins/Animation/PoseSearch/Source/Editor/Private/Chooser/PoseSearchColumnEditor.cpp
--- a/Engine/Plugins/Animation/PoseSearch/Source/Editor/Private/Chooser/PoseSearchColumnEditor.cpp
+++ b/Engine/Plugins/Animation/PoseSearch/Source/Editor/Private/Chooser/PoseSearchColumnEditor.cpp
@@ -24,11 +24,7 @@ TSharedRef<SWidget> CreatePoseSearchColumnWidget(UChooserTable* Chooser, FChoose
 {
 	FPoseSearchColumn* PoseSearchColumn = static_cast<FPoseSearchColumn*>(Column);
 	
-	if (RowIndex == UE::ChooserEditor::ColumnWidget_SpecialIndex_Fallback)
-	{
-		return SNullWidget::NullWidget;
-	}
-	else if (RowIndex == UE::ChooserEditor::ColumnWidget_SpecialIndex_Header)
+	if (RowIndex == UE::ChooserEditor::ColumnWidget_SpecialIndex_Header)
 	{
 		// @todo: avoid reopening the asset picker after tyhe first user choice, or prevent the FPoseSearchColumn creation if no schemas has been selected!
 		if (!PoseSearchColumn->GetDatabaseSchema())
@@ -71,7 +67,7 @@ TSharedRef<SWidget> CreatePoseSearchColumnWidget(UChooserTable* Chooser, FChoose
 
 		// create column header widget
 		const FSlateBrush* ColumnIcon = FAppStyle::Get().GetBrush("Icons.Search");
-		const FText ColumnTooltip = LOCTEXT("Pose Match Tooltip", "Pose Match: Selects a single result based on the animation with the best matching pose, and outputs the StartTime for the frame with that pose. Animation Assets must contain \"Pose Match Branch In\" Notify State. AutoPopulate will fill in Column data with result Animation Assets.");
+		const FText ColumnTooltip = LOCTEXT("Pose Match Tooltip", "Pose Match: Selects results based on the animations with the best matching poses scored via this column \"Internal Database\", and outputs the StartTime of the best matching animation. AutoPopulate will fill in Column data in case it mismatch the Chooser Table \"Result\" column.");
 		const FText ColumnName = LOCTEXT("Pose Match","Pose Match");
         		
 		TSharedPtr<SWidget> DebugWidget = nullptr;
diff --git a/Engine/Plugins/Animation/PoseSearch/Source/Runtime/Private/Chooser/PoseSearchChooserColumn.cpp b/Engine/Plugins/Animation/PoseSearch/Source/Runtime/Private/Chooser/PoseSearchChooserColumn.cpp
--- a/Engine/Plugins/Animation/PoseSearch/Source/Runtime/Private/Chooser/PoseSearchChooserColumn.cpp
+++ b/Engine/Plugins/Animation/PoseSearch/Source/Runtime/Private/Chooser/PoseSearchChooserColumn.cpp
@@ -465,25 +465,23 @@ void FPoseSearchColumn::Filter(FChooserEvaluationContext& Context, const FChoose
 	{
 		PoseSearchColumnScratchArea->PoseHistory = PoseHistory;
 
-		FStackAssetSet AssetsToConsider;
+		FStackDatabaseToAssetIndexes AssetIndexesToConsiderPerDatabase;
+		FStackAssetIndexes& AssetIndexesToConsider = AssetIndexesToConsiderPerDatabase.Add(InternalDatabase.Get());
+
 		for (const FChooserIndexArray::FIndexData& IndexData : IndexListIn)
 		{
 			const int32 RowIndex = IndexData.Index;
 			const int32 DatabaseAssetIndex = GetDatabaseAssetIndex(RowIndex);
-			const FPoseSearchDatabaseAnimationAsset* Asset = InternalDatabase->GetDatabaseAnimationAsset(DatabaseAssetIndex);
-			if (ensure(Asset))
+
+			// pruning out eventual nested chooser or invalid assets not supported by the database
+			if (InternalDatabase->GetAnimationAsset(DatabaseAssetIndex))
 			{
-				const UObject* AnimationAsset = Asset->GetAnimationAsset();
-#if DO_CHECK && WITH_EDITOR
-				// making sure InternalDatabase and the OuterChooser are still in synch
-				const UObject* ReferencedObject = GetReferencedObject(RowIndex);
-				check(ReferencedObject == AnimationAsset);
-#endif // DO_CHECK && WITH_EDITOR
-				AssetsToConsider.Add(AnimationAsset);
+				check(DatabaseAssetIndex >= 0 && DatabaseAssetIndex < InternalDatabase->GetNumAnimationAssets());
+				AssetIndexesToConsider.Add(DatabaseAssetIndex);
 			}
 		}
 
-		if (!AssetsToConsider.IsEmpty())
+		if (!AssetIndexesToConsider.IsEmpty())
 		{
 			const FFloatInterval PoseJumpThresholdTime(0.f, 0.f);
 			const UObject* DatabasesToSearch[] = { InternalDatabase.Get() };
@@ -491,7 +489,7 @@ void FPoseSearchColumn::Filter(FChooserEvaluationContext& Context, const FChoose
 
 			const FRole Role = InternalDatabase->Schema ? InternalDatabase->Schema->GetDefaultRole() : DefaultRole;
 			SearchContext.AddRole(Role, &Context, PoseHistory);
-			SearchContext.SetAssetsToConsiderSet(&AssetsToConsider);
+			SearchContext.SetAssetIndexesToConsiderPerDatabase(&AssetIndexesToConsiderPerDatabase);
 
 			FPoseSearchContinuingProperties ContinuingProperties;
 			if (PoseSearchColumnScratchArea->ChooserPlayerSettings)
@@ -513,6 +511,17 @@ void FPoseSearchColumn::Filter(FChooserEvaluationContext& Context, const FChoose
 			{
 				FSearchResults_Single SearchResults;
 				UPoseSearchLibrary::MotionMatch(SearchContext, DatabasesToSearch, ContinuingProperties, SearchResults);
+#if DO_CHECK
+				SearchResults.IterateOverSearchResults([&AssetIndexesToConsider](FSearchResult& SearchResult)
+					{
+						if (const FSearchIndexAsset* SearchIndexAsset = SearchResult.GetSearchIndexAsset())
+						{
+							ensure(AssetIndexesToConsider.Contains(SearchIndexAsset->GetSourceAssetIdx()));
+						}
+						return false;
+					});
+#endif // DO_CHECK
+
 				BestSearchResult = SearchResults.GetBestResult();
 
 				if (const FSearchIndexAsset* SearchIndexAsset = BestSearchResult.GetSearchIndexAsset())
@@ -530,12 +539,23 @@ void FPoseSearchColumn::Filter(FChooserEvaluationContext& Context, const FChoose
 #if WITH_EDITOR
 					TracedActiveColumnCosts.Add({ RowIndex, BestSearchResult.PoseCost });
 #endif // WITH_EDITOR
-				}
+				}				
 			}
 			else
 			{
 				FSearchResults_AssetBests SearchResults;
 				UPoseSearchLibrary::MotionMatch(SearchContext, DatabasesToSearch, ContinuingProperties, SearchResults);
+#if DO_CHECK
+				SearchResults.IterateOverSearchResults([&AssetIndexesToConsider](FSearchResult& SearchResult)
+					{
+						if (const FSearchIndexAsset* SearchIndexAsset = SearchResult.GetSearchIndexAsset())
+						{
+							ensure(AssetIndexesToConsider.Contains(SearchIndexAsset->GetSourceAssetIdx()));
+						}
+						return false;
+					});
+#endif // DO_CHECK
+
 				BestSearchResult = SearchResults.GetBestResult();
 				if (MaxNumberOfResults > 0)
 				{
@@ -585,79 +605,74 @@ void FPoseSearchColumn::SetOutputs(FChooserEvaluationContext& Context, int RowIn
 	check(PoseSearchColumnScratchArea->DebugOwner == this);
 #endif // DO_CHECK
 
+	FSearchResult BestSearchResult;
+
 	if (RowIndex == ChooserColumn_SpecialIndex_Fallback && PoseSearchColumnScratchArea->PoseHistory)
 	{
 		// making sure that if we search again with the fallback column we didn't find ANY previous result!
 		check(PoseSearchColumnScratchArea->SearchResults.IsEmpty());
 
 		// do the search again only with the fallback row..
-		FStackAssetSet AssetsToConsider;
+		FStackDatabaseToAssetIndexes AssetIndexesToConsiderPerDatabase;
+		FStackAssetIndexes& AssetIndexesToConsider = AssetIndexesToConsiderPerDatabase.Add(InternalDatabase.Get());
+
 		const int32 DatabaseAssetIndex = GetDatabaseAssetIndex(RowIndex);
-		const FPoseSearchDatabaseAnimationAsset* Asset = InternalDatabase->GetDatabaseAnimationAsset(DatabaseAssetIndex);
-		if (ensure(Asset))
+		
+		// pruning out eventual nested chooser or invalid assets not supported by the database
+		if (InternalDatabase->GetAnimationAsset(DatabaseAssetIndex))
 		{
-			if (const UObject* AnimationAsset = Asset->GetAnimationAsset())
+			check(DatabaseAssetIndex >= 0 && DatabaseAssetIndex < InternalDatabase->GetNumAnimationAssets());
+			AssetIndexesToConsider.Add(DatabaseAssetIndex);
+
+			const FFloatInterval PoseJumpThresholdTime(0.f, 0.f);
+			const UObject* DatabasesToSearch[] = { InternalDatabase.Get() };
+			FSearchContext SearchContext(0.f, PoseJumpThresholdTime, FPoseSearchEvent());
+
+			const FRole Role = InternalDatabase->Schema ? InternalDatabase->Schema->GetDefaultRole() : DefaultRole;
+			SearchContext.AddRole(Role, &Context, PoseSearchColumnScratchArea->PoseHistory);
+			SearchContext.SetAssetIndexesToConsiderPerDatabase(&AssetIndexesToConsiderPerDatabase);
+
+			FPoseSearchContinuingProperties ContinuingProperties;
+			if (PoseSearchColumnScratchArea->ChooserPlayerSettings)
 			{
-#if DO_CHECK && WITH_EDITOR
-				// making sure InternalDatabase and the OuterChooser are still in synch
-				const UObject* ReferencedObject = GetReferencedObject(RowIndex);
-				check(ReferencedObject == AnimationAsset);
-#endif // DO_CHECK && WITH_EDITOR
-				AssetsToConsider.Add(AnimationAsset);
-
-				const FFloatInterval PoseJumpThresholdTime(0.f, 0.f);
-				const UObject* DatabasesToSearch[] = { InternalDatabase.Get() };
-				FSearchContext SearchContext(0.f, PoseJumpThresholdTime, FPoseSearchEvent());
-
-				const FRole Role = InternalDatabase->Schema ? InternalDatabase->Schema->GetDefaultRole() : DefaultRole;
-				SearchContext.AddRole(Role, &Context, PoseSearchColumnScratchArea->PoseHistory);
-				SearchContext.SetAssetsToConsiderSet(&AssetsToConsider);
-
-				FPoseSearchContinuingProperties ContinuingProperties;
-				if (PoseSearchColumnScratchArea->ChooserPlayerSettings)
-				{
-					ContinuingProperties.PlayingAsset = PoseSearchColumnScratchArea->ChooserPlayerSettings->PlayingAsset;
-					ContinuingProperties.PlayingAssetAccumulatedTime = PoseSearchColumnScratchArea->ChooserPlayerSettings->PlayingAssetAccumulatedTime;
-					ContinuingProperties.bIsPlayingAssetMirrored = PoseSearchColumnScratchArea->ChooserPlayerSettings->bIsPlayingAssetMirrored;
-					ContinuingProperties.PlayingAssetBlendParameters = PoseSearchColumnScratchArea->ChooserPlayerSettings->PlayingAssetBlendParameters;
+				ContinuingProperties.PlayingAsset = PoseSearchColumnScratchArea->ChooserPlayerSettings->PlayingAsset;
+				ContinuingProperties.PlayingAssetAccumulatedTime = PoseSearchColumnScratchArea->ChooserPlayerSettings->PlayingAssetAccumulatedTime;
+				ContinuingProperties.bIsPlayingAssetMirrored = PoseSearchColumnScratchArea->ChooserPlayerSettings->bIsPlayingAssetMirrored;
+				ContinuingProperties.PlayingAssetBlendParameters = PoseSearchColumnScratchArea->ChooserPlayerSettings->PlayingAssetBlendParameters;
 
-					uint8 InterruptModeValue = 0;
-					if (InterruptMode.IsValid() && InterruptMode.Get<FChooserParameterEnumBase>().GetValue(Context, InterruptModeValue))
-					{
-						ContinuingProperties.InterruptMode = (EPoseSearchInterruptMode)InterruptModeValue;
-					}
+				uint8 InterruptModeValue = 0;
+				if (InterruptMode.IsValid() && InterruptMode.Get<FChooserParameterEnumBase>().GetValue(Context, InterruptModeValue))
+				{
+					ContinuingProperties.InterruptMode = (EPoseSearchInterruptMode)InterruptModeValue;
 				}
+			}
 
-				FSearchResults_Single SearchResults;
-				UPoseSearchLibrary::MotionMatch(SearchContext, DatabasesToSearch, ContinuingProperties, SearchResults);
-				const FSearchResult BestSearchResult = SearchResults.GetBestResult();
+			FSearchResults_Single SearchResults;
+			UPoseSearchLibrary::MotionMatch(SearchContext, DatabasesToSearch, ContinuingProperties, SearchResults);
 
-				if (const FSearchIndexAsset* SearchIndexAsset = BestSearchResult.GetSearchIndexAsset())
+#if DO_CHECK
+			SearchResults.IterateOverSearchResults([&AssetIndexesToConsider](FSearchResult& SearchResult)
 				{
-					check(BestSearchResult.Database == InternalDatabase);
-
-					// @todo: handle duplicate entries, by providing to the UPoseSearchLibrary::MotionMatch a SourceAssetIdxsToConsider rather than AssetsToConsider
-					const int32 SearchResultSourceAssetIdx = SearchIndexAsset->GetSourceAssetIdx();
-					check(SearchResultSourceAssetIdx == 0 && RowIndex == GetRowIndex(SearchResultSourceAssetIdx));
+					if (const FSearchIndexAsset* SearchIndexAsset = SearchResult.GetSearchIndexAsset())
+					{
+						ensure(AssetIndexesToConsider.Contains(SearchIndexAsset->GetSourceAssetIdx()));
+					}
+					return false;
+				});
+#endif // DO_CHECK
 
-					StartTimeValue = BestSearchResult.GetAssetTime();
-					CostValue = BestSearchResult.PoseCost;
-					bMirrorValue = SearchIndexAsset->IsMirrored();
-					bForceBlendToValue = BestSearchResult.bIsContinuingPoseSearch;
+			BestSearchResult = SearchResults.GetBestResult();
+			if (const FSearchIndexAsset* SearchIndexAsset = BestSearchResult.GetSearchIndexAsset())
+			{
+				check(BestSearchResult.Database == InternalDatabase);
 
-					if (const FPoseIndicesHistory* PoseIndicesHistory = PoseSearchColumnScratchArea->PoseHistory->GetPoseIndicesHistory())
-					{
-						// @todo: integrate DeltaTime into SearchContext, and implement it for UAF as well
-						float DeltaTime = FiniteDelta;
-						if (const UAnimInstance* AnimInstance = Cast<UAnimInstance>(GetAnimContext(&Context)))
-						{
-							DeltaTime = AnimInstance->GetDeltaSeconds();
-						}
+				const int32 SearchResultSourceAssetIdx = SearchIndexAsset->GetSourceAssetIdx();
+				ensure(SearchResultSourceAssetIdx == 0 && RowIndex == GetRowIndex(SearchResultSourceAssetIdx));
 
-						// const casting here is safe since we're in the thread owning the pose history, and it's the correct place to update the previously selected poses
-						const_cast<FPoseIndicesHistory*>(PoseIndicesHistory)->Update(BestSearchResult, DeltaTime, PoseReselectHistory);
-					}
-				}
+				StartTimeValue = BestSearchResult.GetAssetTime();
+				CostValue = BestSearchResult.PoseCost;
+				bMirrorValue = SearchIndexAsset->IsMirrored();
+				bForceBlendToValue = BestSearchResult.bIsContinuingPoseSearch;
 			}
 		}
 	}
@@ -676,27 +691,29 @@ void FPoseSearchColumn::SetOutputs(FChooserEvaluationContext& Context, int RowIn
 					bMirrorValue = SearchIndexAsset->IsMirrored();
 					bForceBlendToValue = SearchResult.bIsContinuingPoseSearch;
 
-					if (PoseSearchColumnScratchArea->PoseHistory)
-					{
-						if (const FPoseIndicesHistory* PoseIndicesHistory = PoseSearchColumnScratchArea->PoseHistory->GetPoseIndicesHistory())
-						{
-							// @todo: integrate DeltaTime into SearchContext, and implement it for UAF as well
-							float DeltaTime = FiniteDelta;
-							if (const UAnimInstance* AnimInstance = Cast<UAnimInstance>(GetAnimContext(&Context)))
-							{
-								DeltaTime = AnimInstance->GetDeltaSeconds();
-							}
-
-							// const casting here is safe since we're in the thread owning the pose history, and it's the correct place to update the previously selected poses
-							const_cast<FPoseIndicesHistory*>(PoseIndicesHistory)->Update(SearchResult, DeltaTime, PoseReselectHistory);
-						}
-					}
+					BestSearchResult = SearchResult;
 					break;
 				}
 			}
 		}
 	}
 
+	if (PoseSearchColumnScratchArea->PoseHistory)
+	{
+		if (const FPoseIndicesHistory* PoseIndicesHistory = PoseSearchColumnScratchArea->PoseHistory->GetPoseIndicesHistory())
+		{
+			// @todo: integrate DeltaTime into SearchContext, and implement it for UAF as well
+			float DeltaTime = FiniteDelta;
+			if (const UAnimInstance* AnimInstance = Cast<UAnimInstance>(GetAnimContext(&Context)))
+			{
+				DeltaTime = AnimInstance->GetDeltaSeconds();
+			}
+
+			// const casting here is safe since we're in the thread owning the pose history, and it's the correct place to update the previously selected poses
+			const_cast<FPoseIndicesHistory*>(PoseIndicesHistory)->Update(BestSearchResult, DeltaTime, PoseReselectHistory);
+		}
+	}
+
 	if (const FChooserParameterFloatBase* StartTime = OutputStartTime.GetPtr<FChooserParameterFloatBase>())
 	{
 		StartTime->SetValue(Context, StartTimeValue);
diff --git a/Engine/Plugins/Animation/PoseSearch/Source/Runtime/Private/PoseSearchDatabase.cpp b/Engine/Plugins/Animation/PoseSearch/Source/Runtime/Private/PoseSearchDatabase.cpp
--- a/Engine/Plugins/Animation/PoseSearch/Source/Runtime/Private/PoseSearchDatabase.cpp
+++ b/Engine/Plugins/Animation/PoseSearch/Source/Runtime/Private/PoseSearchDatabase.cpp
@@ -922,7 +922,7 @@ const UE::PoseSearch::FSearchIndex& UPoseSearchDatabase::GetSearchIndex() const
 	return SearchIndexPrivate;
 }
 
-int32 UPoseSearchDatabase::GetPoseIndex(const UObject* AnimationAsset, float AnimationAssetTime, bool bMirrored, const FVector& BlendParameters) const
+int32 UPoseSearchDatabase::GetPoseIndex(const UObject* AnimationAsset, float AnimationAssetTime, bool bMirrored, const FVector& BlendParameters, const UE::PoseSearch::FStackAssetIndexes* AssetIndexesToConsider) const
 {
 	using namespace UE::PoseSearch;
 
@@ -938,23 +938,28 @@ int32 UPoseSearchDatabase::GetPoseIndex(const UObject* AnimationAsset, float Ani
 		for (int32 AssetIndex : AssetIndexesForSourceAsset)
 		{
 			const FSearchIndexAsset& SearchIndexAsset = SearchIndex.Assets[AssetIndex];
-			if (SearchIndexAsset.IsMirrored() == bMirrored)
+			const bool bIsAssetIndexValid = !AssetIndexesToConsider || AssetIndexesToConsider->Contains(SearchIndexAsset.GetSourceAssetIdx());
+			if (bIsAssetIndexValid)
 			{
-				const float BlendParametersSquaredLength = (BlendParameters - SearchIndexAsset.GetBlendParameters()).SquaredLength();
-
-				if (BlendParametersSquaredLength <= MinSquaredLength)
+				if (SearchIndexAsset.IsMirrored() == bMirrored)
 				{
-					const FPoseSearchDatabaseAnimationAssetBase* DatabaseAnimationAssetBase = GetDatabaseAnimationAsset(SearchIndexAsset);
-
-					check(DatabaseAnimationAssetBase);
-					check(DatabaseAnimationAssetBase->GetAnimationAsset() == AnimationAsset);
+					// @todo: optimize this code, since there's no need to calculate BlendParametersSquaredLength if AssetIndexesForSourceAsset.Num() is 1
+					const float BlendParametersSquaredLength = (BlendParameters - SearchIndexAsset.GetBlendParameters()).SquaredLength();
 
-					const float RealAssetTime = AnimationAssetTime * SearchIndexAsset.GetToRealTimeFactor();
-					const int32 CurrentIdx = SearchIndexAsset.GetPoseIndexFromTime(RealAssetTime, SampleRate);
-					if (CurrentIdx != INDEX_NONE)
+					if (BlendParametersSquaredLength <= MinSquaredLength)
 					{
-						MinSquaredLength = BlendParametersSquaredLength;
-						PoseIdx = CurrentIdx;
+						const FPoseSearchDatabaseAnimationAssetBase* DatabaseAnimationAssetBase = GetDatabaseAnimationAsset(SearchIndexAsset);
+
+						check(DatabaseAnimationAssetBase);
+						check(DatabaseAnimationAssetBase->GetAnimationAsset() == AnimationAsset);
+
+						const float RealAssetTime = AnimationAssetTime * SearchIndexAsset.GetToRealTimeFactor();
+						const int32 CurrentIdx = SearchIndexAsset.GetPoseIndexFromTime(RealAssetTime, SampleRate);
+						if (CurrentIdx != INDEX_NONE)
+						{
+							MinSquaredLength = BlendParametersSquaredLength;
+							PoseIdx = CurrentIdx;
+						}
 					}
 				}
 			}
@@ -1541,8 +1546,10 @@ void UPoseSearchDatabase::PopulateSelectableAssetIdx(FSelectableAssetIdx& Select
 
 void UPoseSearchDatabase::PopulateSelectableAssetIdx(FSelectableAssetIdx& SelectableAssetIdx, const UE::PoseSearch::FStackAssetSet* AssetsToConsider) const
 {
+	check(AssetsToConsider);
+
 	SelectableAssetIdx.Reset();
-	if (AssetsToConsider && !AssetsToConsider->IsEmpty())
+	if (!AssetsToConsider->IsEmpty())
 	{
 		for (const UObject* AssetToConsider : *AssetsToConsider)
 		{
@@ -1565,6 +1572,41 @@ void UPoseSearchDatabase::PopulateSelectableAssetIdx(FSelectableAssetIdx& Select
 	}
 }
 
+void UPoseSearchDatabase::PopulateSelectableAssetIdx(FSelectableAssetIdx& SelectableAssetIdx, const UE::PoseSearch::FStackDatabaseToAssetIndexes* AssetIndexesToConsiderPerDatabase) const
+{
+	using namespace UE::PoseSearch;
+	
+	check(AssetIndexesToConsiderPerDatabase);
+
+	const int32 NumAnimationAssets = GetNumAnimationAssets();
+
+	SelectableAssetIdx.Reset();
+	if (const FStackAssetIndexes* AssetIndexesToConsider = AssetIndexesToConsiderPerDatabase->Find(this))
+	{
+		check(AssetIndexesToConsider->Num() <= NumAnimationAssets);
+
+		// if AssetIndexesToConsider->Num() == NumAnimationAssets we want to search the full database and and empty SelectableAssetIdx does that
+		if (AssetIndexesToConsider->Num() != NumAnimationAssets)
+		{
+			for (int32 AssetIndexToConsider : *AssetIndexesToConsider)
+			{
+				if (const TArray<int32>* IndexesForSourceAsset = CachedAssetMap.Find(DatabaseAnimationAssets[AssetIndexToConsider].AnimAsset))
+				{
+					// @todo: add a mapping like CachedAssetMap for asset indexes to SearchIndexPrivate.Assets indexes to avoid this cumbersome filtering
+					for (int32 IndexForSourceAsset : *IndexesForSourceAsset)
+					{
+						if (SearchIndexPrivate.Assets[IndexForSourceAsset].GetSourceAssetIdx() == AssetIndexToConsider)
+						{
+							SelectableAssetIdx.Add(IndexForSourceAsset);
+						}
+					}
+				}
+			}
+			SelectableAssetIdx.Sort();
+		}
+	}
+}
+
 UE::PoseSearch::FSearchResult UPoseSearchDatabase::Search(UE::PoseSearch::FSearchContext& SearchContext) const
 {
 	using namespace UE::PoseSearch;
@@ -1817,11 +1859,15 @@ void UPoseSearchDatabase::SearchPCAKDTree(UE::PoseSearch::FSearchContext& Search
 		{
 			PopulateSelectableAssetIdx(SelectableAssetIdx, SearchContext.GetInternalDeprecatedAssetsToConsider());
 		}
-		else
+		else if (SearchContext.GetAssetsToConsiderSet())
 		PRAGMA_ENABLE_DEPRECATION_WARNINGS
 		{
 			PopulateSelectableAssetIdx(SelectableAssetIdx, SearchContext.GetAssetsToConsiderSet());
 		}
+		else if (SearchContext.GetAssetIndexesToConsiderPerDatabase())
+		{
+			PopulateSelectableAssetIdx(SelectableAssetIdx, SearchContext.GetAssetIndexesToConsiderPerDatabase());
+		}
 
 		FNonSelectableIdx NonSelectableIdx;
 		PopulateNonSelectableIdx(NonSelectableIdx, SearchContext
@@ -2171,11 +2217,15 @@ void UPoseSearchDatabase::SearchVPTree(UE::PoseSearch::FSearchContext& SearchCon
 		{
 			PopulateSelectableAssetIdx(SelectableAssetIdx, SearchContext.GetInternalDeprecatedAssetsToConsider());
 		}
-		else
+		else if (SearchContext.GetAssetsToConsiderSet())
 		PRAGMA_ENABLE_DEPRECATION_WARNINGS
 		{
 			PopulateSelectableAssetIdx(SelectableAssetIdx, SearchContext.GetAssetsToConsiderSet());
 		}
+		else if (SearchContext.GetAssetIndexesToConsiderPerDatabase())
+		{
+			PopulateSelectableAssetIdx(SelectableAssetIdx, SearchContext.GetAssetIndexesToConsiderPerDatabase());
+		}
 
 		// @todo: implement filtering within the VPTree as KDTree does
 		FNonSelectableIdx NonSelectableIdx;
@@ -2287,11 +2337,15 @@ void UPoseSearchDatabase::SearchBruteForce(UE::PoseSearch::FSearchContext& Searc
 		{
 			PopulateSelectableAssetIdx(SelectableAssetIdx, SearchContext.GetInternalDeprecatedAssetsToConsider());
 		}
-		else
+		else if (SearchContext.GetAssetsToConsiderSet())
 		PRAGMA_ENABLE_DEPRECATION_WARNINGS
 		{
 			PopulateSelectableAssetIdx(SelectableAssetIdx, SearchContext.GetAssetsToConsiderSet());
 		}
+		else if (SearchContext.GetAssetIndexesToConsiderPerDatabase())
+		{
+			PopulateSelectableAssetIdx(SelectableAssetIdx, SearchContext.GetAssetIndexesToConsiderPerDatabase());
+		}
 
 		FNonSelectableIdx NonSelectableIdx;
 		PopulateNonSelectableIdx(NonSelectableIdx, SearchContext
@@ -2483,11 +2537,15 @@ void UPoseSearchDatabase::SearchEvent(UE::PoseSearch::FSearchContext& SearchCont
 			{
 				PopulateSelectableAssetIdx(SelectableAssetIdx, SearchContext.GetInternalDeprecatedAssetsToConsider());
 			}
-			else
+			else if (SearchContext.GetAssetsToConsiderSet())
 			PRAGMA_ENABLE_DEPRECATION_WARNINGS
 			{
 				PopulateSelectableAssetIdx(SelectableAssetIdx, SearchContext.GetAssetsToConsiderSet());
 			}
+			else if (SearchContext.GetAssetIndexesToConsiderPerDatabase())
+			{
+				PopulateSelectableAssetIdx(SelectableAssetIdx, SearchContext.GetAssetIndexesToConsiderPerDatabase());
+			}
 
 			FNonSelectableIdx NonSelectableIdx;
 			if (EventToSearch.bEnablePoseFilters)
diff --git a/Engine/Plugins/Animation/PoseSearch/Source/Runtime/Private/PoseSearchLibrary.cpp b/Engine/Plugins/Animation/PoseSearch/Source/Runtime/Private/PoseSearchLibrary.cpp
--- a/Engine/Plugins/Animation/PoseSearch/Source/Runtime/Private/PoseSearchLibrary.cpp
+++ b/Engine/Plugins/Animation/PoseSearch/Source/Runtime/Private/PoseSearchLibrary.cpp
@@ -98,9 +98,98 @@ namespace UE::PoseSearch
 		}
 	};
 
+	static void FilterByAssetsToConsider(FSearchContext& SearchContext, FStackDatabaseToAssetIndexes& AssetIndexesToSearchPerDatabase)
+	{
+		// intersecting AssetsToSearchPerDatabase with asset to consider
+		PRAGMA_DISABLE_DEPRECATION_WARNINGS
+		if (!SearchContext.GetInternalDeprecatedAssetsToConsider().IsEmpty())
+		{
+			// backward compatible path to support deprecated API FSearchContext::SetAssetsToConsider
+			for (TPair<const UPoseSearchDatabase*, FStackAssetIndexes>& AssetIndexesToSearchPerDatabasePair : AssetIndexesToSearchPerDatabase)
+			{
+				const UPoseSearchDatabase* Database = AssetIndexesToSearchPerDatabasePair.Key;
+				check(Database);
+
+				FStackAssetIndexes& AssetIndexesToSearchForDatabase = AssetIndexesToSearchPerDatabasePair.Value;
+				FStackAssetIndexes IntersectedAssetIndexesToSearchForDatabase;
+
+				const FSearchIndex& SearchIndex = Database->GetSearchIndex();
+				for (const UObject* AssetToConsider : SearchContext.GetInternalDeprecatedAssetsToConsider())
+				{
+					for (int32 DatabaseAnimationAssetIndex : Database->GetAssetIndexesForSourceAsset(AssetToConsider))
+					{
+						const FSearchIndexAsset& SearchIndexAsset = SearchIndex.Assets[DatabaseAnimationAssetIndex];
+						const int32 SourceAssetIdx = SearchIndexAsset.GetSourceAssetIdx();
+						if (AssetIndexesToSearchForDatabase.Find(SourceAssetIdx))
+						{
+							IntersectedAssetIndexesToSearchForDatabase.Add(SourceAssetIdx);
+						}
+					}
+				}
+
+				AssetIndexesToSearchForDatabase = IntersectedAssetIndexesToSearchForDatabase;
+			}
+		}
+		PRAGMA_ENABLE_DEPRECATION_WARNINGS
+		else if (const FStackAssetSet* AssetsToConsider = SearchContext.GetAssetsToConsiderSet())
+		{
+			if (!AssetsToConsider->IsEmpty())
+			{
+				for (TPair<const UPoseSearchDatabase*, FStackAssetIndexes>& AssetIndexesToSearchPerDatabasePair : AssetIndexesToSearchPerDatabase)
+				{
+					const UPoseSearchDatabase* Database = AssetIndexesToSearchPerDatabasePair.Key;
+					check(Database);
+
+					FStackAssetIndexes& AssetIndexesToSearchForDatabase = AssetIndexesToSearchPerDatabasePair.Value;
+					FStackAssetIndexes IntersectedAssetIndexesToSearchForDatabase;
+
+					const FSearchIndex& SearchIndex = Database->GetSearchIndex();
+					for (const UObject* AssetToConsider : *AssetsToConsider)
+					{
+						for (int32 DatabaseAnimationAssetIndex : Database->GetAssetIndexesForSourceAsset(AssetToConsider))
+						{
+							const FSearchIndexAsset& SearchIndexAsset = SearchIndex.Assets[DatabaseAnimationAssetIndex];
+							const int32 SourceAssetIdx = SearchIndexAsset.GetSourceAssetIdx();
+							if (AssetIndexesToSearchForDatabase.Find(SourceAssetIdx))
+							{
+								IntersectedAssetIndexesToSearchForDatabase.Add(SourceAssetIdx);
+							}
+						}
+					}
+
+					AssetIndexesToSearchForDatabase = IntersectedAssetIndexesToSearchForDatabase;
+				}
+			}
+		}
+		else if (const FStackDatabaseToAssetIndexes* AssetIndexesToConsiderPerDatabase = SearchContext.GetAssetIndexesToConsiderPerDatabase())
+		{
+			for (TPair<const UPoseSearchDatabase*, FStackAssetIndexes>& AssetIndexesToSearchPerDatabasePair : AssetIndexesToSearchPerDatabase)
+			{
+				const UPoseSearchDatabase* Database = AssetIndexesToSearchPerDatabasePair.Key;
+				check(Database);
+
+				FStackAssetIndexes& AssetIndexesToSearchForDatabase = AssetIndexesToSearchPerDatabasePair.Value;
+				FStackAssetIndexes IntersectedAssetIndexesToSearchForDatabase;
+
+				if (const FStackAssetIndexes* AssetIndexesToConsider = AssetIndexesToConsiderPerDatabase->Find(Database))
+				{
+					for (int32 DatabaseAnimationAssetIndex : *AssetIndexesToConsider)
+					{
+						if (AssetIndexesToSearchForDatabase.Find(DatabaseAnimationAssetIndex))
+						{
+							check(DatabaseAnimationAssetIndex < Database->GetNumAnimationAssets());
+							IntersectedAssetIndexesToSearchForDatabase.Add(DatabaseAnimationAssetIndex);
+						}
+					}
+				}
+
+				AssetIndexesToSearchForDatabase = IntersectedAssetIndexesToSearchForDatabase;
+			}
+		}
+	}
+
 	// this function adds AssetToSearch to the search of Database
-	// returns bAsyncBuildIndexInProgress
-	static bool AddToSearchForDatabase(FAssetsToSearchPerDatabase& AssetsToSearchPerDatabase, const UObject* AssetToSearch, const UPoseSearchDatabase* Database, bool bContainsIsMandatory)
+	static void AddToSearchForDatabase(FStackDatabaseToAssetIndexes& AssetIndexesToSearchPerDatabase, const UObject* AssetToSearch, const UPoseSearchDatabase* Database, bool bContainsIsMandatory, FSearchContext& SearchContext)
 	{
 		// making sure AssetToSearch is not a databases! later on we could add support for nested databases, but currently we don't support that
 		check(Cast<const UPoseSearchDatabase>(AssetToSearch) == nullptr);
@@ -110,50 +199,41 @@ namespace UE::PoseSearch
 		if (EAsyncBuildIndexResult::Success != FAsyncPoseSearchDatabasesManagement::RequestAsyncBuildIndex(Database, ERequestAsyncBuildFlag::ContinueRequest))
 		{
 			// database is still indexing... moving on
-			return true;
+			SearchContext.SetAsyncBuildIndexInProgress();
 		}
+		else
 #endif // WITH_EDITOR
-
-		if (!Database->Contains(AssetToSearch))
 		{
-			if (bContainsIsMandatory)
+			const TConstArrayView<int32> AssetIndexesForSourceAssetToSearch = Database->GetAssetIndexesForSourceAsset(AssetToSearch);
+			if (AssetIndexesForSourceAssetToSearch.IsEmpty())
 			{
-				UE_LOG(LogPoseSearch, Error, TEXT("improperly setup UAnimSequenceBase. Database %s doesn't contain UAnimSequenceBase %s"), *Database->GetName(), *AssetToSearch->GetName());
+				if (bContainsIsMandatory)
+				{
+					UE_LOG(LogPoseSearch, Error, TEXT("improperly setup UAnimSequenceBase. Database %s doesn't contain UAnimSequenceBase %s"), *Database->GetName(), *AssetToSearch->GetName());
+				}
 			}
-			return false;
-		}
-
-		if (FAssetsToSearchPerDatabasePair* AssetsToSearchPerDatabasePair = AssetsToSearchPerDatabase.Find(Database))
-		{
-			// an empty FStackAssetSet associated to Database means we need to search ALL the assets, so we don't need to add this AssetToSearch
-			FStackAssetSet& AssetsToSearch = AssetsToSearchPerDatabasePair->Value;
-			if (!AssetsToSearch.IsEmpty())
+			else
 			{
-				AssetsToSearch.Add(AssetToSearch);
+				const FSearchIndex& SearchIndex = Database->GetSearchIndex();
+				FStackAssetIndexes& AssetIndexesToSearch = AssetIndexesToSearchPerDatabase.FindOrAdd(Database);
+				for (int32 AssetIndexToSearch : AssetIndexesForSourceAssetToSearch)
+				{
+					const FSearchIndexAsset& SearchIndexAsset = SearchIndex.Assets[AssetIndexToSearch];
+					AssetIndexesToSearch.Add(SearchIndexAsset.GetSourceAssetIdx());
+				}
 			}
 		}
-		else
-		{
-			// no need to AddUnique since it's the first one
-			FAssetsToSearchPerDatabasePair& NewAssetsToSearchPerDatabasePair = AssetsToSearchPerDatabase.AddDefaulted_GetRef();
-			NewAssetsToSearchPerDatabasePair.Key = Database;
-			NewAssetsToSearchPerDatabasePair.Value.Add(AssetToSearch);
-		}
-
-		return false;
 	}
 
 	// this function is looking for UPoseSearchDatabase(s) to search for the input AssetToSearch:
 	// if AssetToSearch is a database search it ALL,
 	// if it's a sequence containing UAnimNotifyState_PoseSearchBranchIn, we add to the search of the database UAnimNotifyState_PoseSearchBranchIn::Database the asset AssetToSearch
-	// returns bAsyncBuildIndexInProgress
-	static bool AddToSearch(FAssetsToSearchPerDatabase& AssetsToSearchPerDatabase, const UObject* AssetToSearch, bool bUsePoseSearchBranchIn)
+	static void AddToSearch(FStackDatabaseToAssetIndexes& AssetIndexesToSearchPerDatabase, const UObject* AssetToSearch, bool bUsePoseSearchBranchIn, FSearchContext& SearchContext)
 	{
 		if (bUsePoseSearchBranchIn)
 		{
 			if (const UAnimSequenceBase* SequenceBase = Cast<const UAnimSequenceBase>(AssetToSearch))
 			{
-				bool bAsyncBuildIndexInProgress = false;
 				for (const FAnimNotifyEvent& NotifyEvent : SequenceBase->Notifies)
 				{
 					if (const UAnimNotifyState_PoseSearchBranchIn* PoseSearchBranchIn = Cast<UAnimNotifyState_PoseSearchBranchIn>(NotifyEvent.NotifyStateClass))
@@ -165,13 +245,9 @@ namespace UE::PoseSearch
 						}
 
 						// we just skip indexing databases to keep the experience as smooth as possible
-						if (AddToSearchForDatabase(AssetsToSearchPerDatabase, SequenceBase, PoseSearchBranchIn->Database, true))
-						{
-							bAsyncBuildIndexInProgress = true;
-						}
+						AddToSearchForDatabase(AssetIndexesToSearchPerDatabase, SequenceBase, PoseSearchBranchIn->Database, true, SearchContext);
 					}
 				}
-				return bAsyncBuildIndexInProgress;
 			}
 		}
 
@@ -180,185 +256,55 @@ namespace UE::PoseSearch
 #if WITH_EDITOR
 			if (EAsyncBuildIndexResult::Success != FAsyncPoseSearchDatabasesManagement::RequestAsyncBuildIndex(Database, ERequestAsyncBuildFlag::ContinueRequest))
 			{
-				return true;
-			}
-#endif // WITH_EDITOR
-
-			// we already added Database to AssetsToSearchPerDatabase, so it already successfully passed RequestAsyncBuildIndex
-			if (FAssetsToSearchPerDatabasePair* AssetsToSearchPerDatabasePair = AssetsToSearchPerDatabase.Find(Database))
-			{
-				// an empty FStackAssetSet associated to Database means we need to search ALL the assets
-				FStackAssetSet& AssetsToSearch = AssetsToSearchPerDatabasePair->Value;
-				AssetsToSearch.Reset();
+				SearchContext.SetAsyncBuildIndexInProgress();
 			}
 			else
+#endif // WITH_EDITOR
 			{
-				// an empty FStackAssetSet associated to Database means we need to search ALL the assets
-				AssetsToSearchPerDatabase.AddDefaulted_GetRef().Key = Database;
+				FStackAssetIndexes& AssetIndexesToSearch = AssetIndexesToSearchPerDatabase.FindOrAdd(Database);
+				const int32 NumAnimationAssets = Database->GetNumAnimationAssets();
+				for (int32 AnimationAssetIndex = 0; AnimationAssetIndex < NumAnimationAssets; ++AnimationAssetIndex)
+				{
+					AssetIndexesToSearch.Add(AnimationAssetIndex);
+				}
 			}
-			return false;
 		}
-
-		return false;
 	}
 
-	static void PopulateContinuingPoseSearches(const FPoseSearchContinuingProperties& ContinuingProperties, const TArrayView<const UObject*> AssetsToSearch, FSearchContext& SearchContext, FAssetsToSearchPerDatabase& ContinuingPoseAssetsToSearchPerDatabase, bool bUsePoseSearchBranchIn)
+	static void PopulateContinuingPoseSearches(const FPoseSearchContinuingProperties& ContinuingProperties, const TArrayView<const UObject*> AssetsToSearch, FSearchContext& SearchContext, FStackDatabaseToAssetIndexes& ContinuingPoseAssetIndexesToSearchPerDatabase, bool bUsePoseSearchBranchIn)
 	{
 		if (const UObject* PlayingAnimationAsset = ContinuingProperties.PlayingAsset.Get())
 		{
-			// checking if PlayingAnimationAsset can be considered or filtered out
-			bool bCanBeConsidered = true;
-			
-			// if !AssetsToConsider or AssetsToConsider->IsEmpty(), we need to consider all the assets!
-			PRAGMA_DISABLE_DEPRECATION_WARNINGS
-			if (!SearchContext.GetInternalDeprecatedAssetsToConsider().IsEmpty())
-			{
-				// backward compatible path to support deprecated API FSearchContext::SetAssetsToConsider
-				if (!SearchContext.GetInternalDeprecatedAssetsToConsider().Contains(PlayingAnimationAsset))
-				{
-					bCanBeConsidered = false;
-				}
-			}
-			else
-			PRAGMA_ENABLE_DEPRECATION_WARNINGS
+			// checking if any of the AssetsToSearch (databases) or ContinuingProperties.PlayingAssetDatabase contain PlayingAnimationAsset
+			if (ContinuingProperties.PlayingAssetDatabase)
 			{
-				if (const FStackAssetSet* AssetsToConsider = SearchContext.GetAssetsToConsiderSet())
-				{
-					if (!AssetsToConsider->IsEmpty() && !AssetsToConsider->Contains(PlayingAnimationAsset))
-					{
-						bCanBeConsidered = false;
-					}
-				}
+				AddToSearchForDatabase(ContinuingPoseAssetIndexesToSearchPerDatabase, PlayingAnimationAsset, ContinuingProperties.PlayingAssetDatabase, false, SearchContext);
 			}
 
-			if (bCanBeConsidered)
+			for (const UObject* AssetToSearch : AssetsToSearch)
 			{
-				// checking if any of the AssetsToSearch (databases) or ContinuingProperties.PlayingAssetDatabase contain PlayingAnimationAsset
-				if (ContinuingProperties.PlayingAssetDatabase)
+				if (const UPoseSearchDatabase* Database = Cast<UPoseSearchDatabase>(AssetToSearch))
 				{
-					if (AddToSearchForDatabase(ContinuingPoseAssetsToSearchPerDatabase, PlayingAnimationAsset, ContinuingProperties.PlayingAssetDatabase, false))
-					{
-#if WITH_EDITOR
-						SearchContext.SetAsyncBuildIndexInProgress();
-#endif // WITH_EDITOR
-					}
+					// since it cannot be a database we can directly add it to ContinuingPoseAssetsToSearchPerDatabase
+					AddToSearchForDatabase(ContinuingPoseAssetIndexesToSearchPerDatabase, PlayingAnimationAsset, Database, false, SearchContext);
 				}
+			}
 
-				for (const UObject* AssetToSearch : AssetsToSearch)
-				{
-					if (const UPoseSearchDatabase* Database = Cast<UPoseSearchDatabase>(AssetToSearch))
-					{
-						// since it cannot be a database we can directly add it to ContinuingPoseAssetsToSearchPerDatabase
-						if (AddToSearchForDatabase(ContinuingPoseAssetsToSearchPerDatabase, PlayingAnimationAsset, Database, false))
-						{
-#if WITH_EDITOR
-							SearchContext.SetAsyncBuildIndexInProgress();
-#endif // WITH_EDITOR
-						}
-					}
-				}
+			// checking if PlayingAnimationAsset has an associated database
+			AddToSearch(ContinuingPoseAssetIndexesToSearchPerDatabase, PlayingAnimationAsset, bUsePoseSearchBranchIn, SearchContext);
 
-				// checking if PlayingAnimationAsset has an associated database
-				if (AddToSearch(ContinuingPoseAssetsToSearchPerDatabase, PlayingAnimationAsset, bUsePoseSearchBranchIn))
-				{
-#if WITH_EDITOR
-					SearchContext.SetAsyncBuildIndexInProgress();
-#endif // WITH_EDITOR
-				}
-			}
+			FilterByAssetsToConsider(SearchContext, ContinuingPoseAssetIndexesToSearchPerDatabase);
 		}
 	}
 	
-	static void PopulateSearches(const TArrayView<const UObject*> AssetsToSearch, FSearchContext& SearchContext, FAssetsToSearchPerDatabase& AssetsToSearchPerDatabase, bool bUsePoseSearchBranchIn)
+	static void PopulateSearches(const TArrayView<const UObject*> AssetsToSearch, FSearchContext& SearchContext, FStackDatabaseToAssetIndexes& AssetIndexesToSearchPerDatabase, bool bUsePoseSearchBranchIn)
 	{
 		for (const UObject* AssetToSearch : AssetsToSearch)
 		{
-			if (AddToSearch(AssetsToSearchPerDatabase, AssetToSearch, bUsePoseSearchBranchIn))
-			{
-#if WITH_EDITOR
-				SearchContext.SetAsyncBuildIndexInProgress();
-#endif // WITH_EDITOR
-			}
+			AddToSearch(AssetIndexesToSearchPerDatabase, AssetToSearch, bUsePoseSearchBranchIn, SearchContext);
 		}
 
-		// intersecting AssetsToSearchPerDatabase with asset to consider
-		PRAGMA_DISABLE_DEPRECATION_WARNINGS
-		if (!SearchContext.GetInternalDeprecatedAssetsToConsider().IsEmpty())
-		{
-			// backward compatible path to support deprecated API FSearchContext::SetAssetsToConsider
-			for (FAssetsToSearchPerDatabasePair& AssetsToSearchPerDatabasePair : AssetsToSearchPerDatabase)
-			{
-				const UPoseSearchDatabase* Database = AssetsToSearchPerDatabasePair.Key;
-				check(Database);
-
-				FStackAssetSet& AssetsToSearchForDatabase = AssetsToSearchPerDatabasePair.Value;
-				if (!AssetsToSearchForDatabase.IsEmpty())
-				{
-					// doing an intersection between AssetsToSearchForDatabase and asset to consider
-					FStackAssetSet AssetsToSearchForDatabaseIntersection;
-					AssetsToSearchForDatabaseIntersection.Reserve(AssetsToSearchForDatabase.Num());
-					for (const UObject* AssetToSearchForDatabase : AssetsToSearchForDatabase)
-					{
-						if (SearchContext.GetInternalDeprecatedAssetsToConsider().Contains(AssetToSearchForDatabase))
-						{
-							AssetsToSearchForDatabaseIntersection.Add(AssetToSearchForDatabase);
-						}
-					}
-					AssetsToSearchForDatabase = AssetsToSearchForDatabaseIntersection;
-				}
-				else
-				{
-					// since all the database assets can be searched, we need to add only the one we can consider:
-					for (const UObject* AssetToConsider : SearchContext.GetInternalDeprecatedAssetsToConsider())
-					{
-						if (Database->Contains(AssetToConsider))
-						{
-							AssetsToSearchForDatabase.Add(AssetToConsider);
-						}
-					}
-				}
-			}
-		}
-		PRAGMA_ENABLE_DEPRECATION_WARNINGS
-		else if (const FStackAssetSet* AssetsToConsider = SearchContext.GetAssetsToConsiderSet())
-		{
-			if (!AssetsToConsider->IsEmpty())
-			{
-				for (FAssetsToSearchPerDatabasePair& AssetsToSearchPerDatabasePair : AssetsToSearchPerDatabase)
-				{
-					const UPoseSearchDatabase* Database = AssetsToSearchPerDatabasePair.Key;
-					check(Database);
-
-					FStackAssetSet& AssetsToSearchForDatabase = AssetsToSearchPerDatabasePair.Value;
-					if (!AssetsToSearchForDatabase.IsEmpty())
-					{
-						// doing an intersection between AssetsToSearchForDatabase and asset to consider
-						FStackAssetSet AssetsToSearchForDatabaseIntersection;
-						AssetsToSearchForDatabaseIntersection.Reserve(AssetsToSearchForDatabase.Num());
-						for (const UObject* AssetToSearchForDatabase : AssetsToSearchForDatabase)
-						{
-							if (AssetsToConsider->Contains(AssetToSearchForDatabase))
-							{
-								AssetsToSearchForDatabaseIntersection.Add(AssetToSearchForDatabase);
-							}
-						}
-						AssetsToSearchForDatabase = AssetsToSearchForDatabaseIntersection;
-					}
-					else
-					{
-						// since all the database assets can be searched, we need to add only the one we can consider:
-						// @todo: should we just do a copy? AssetsToSearchForDatabase = AssetsToConsider??
-						for (const UObject* AssetToConsider : *AssetsToConsider)
-						{
-							if (Database->Contains(AssetToConsider))
-							{
-								AssetsToSearchForDatabase.Add(AssetToConsider);
-							}
-						}
-					}
-				}
-			}
-		}
+		FilterByAssetsToConsider(SearchContext, AssetIndexesToSearchPerDatabase);
 	}
 
 	// @todo: refine this logic. Currently if AssetsToSearch contains ONLY UPoseSearchDatabase we don't have to look for other databases referenced by other assets UAnimNotifyState_PoseSearchBranchIn(s)
@@ -1310,34 +1256,39 @@ void UPoseSearchLibrary::MotionMatch(
 {
 	using namespace UE::PoseSearch;
 
-	const FStackAssetSet* CurrentAssetToConsider = nullptr;
+	const FStackAssetSet* CurrentAssetsToConsider = nullptr;
+	const FStackDatabaseToAssetIndexes* CurrentAssetIndexesToConsiderPerDatabase = nullptr;
 	PRAGMA_DISABLE_DEPRECATION_WARNINGS
-	if (SearchContext.GetInternalDeprecatedAssetsToConsider().IsEmpty())
+	TArray<const UObject*> InternalDeprecatedAssetsToConsider = SearchContext.GetInternalDeprecatedAssetsToConsider();
+	if (InternalDeprecatedAssetsToConsider.IsEmpty())
 	PRAGMA_ENABLE_DEPRECATION_WARNINGS
 	{
-		CurrentAssetToConsider = SearchContext.GetAssetsToConsiderSet();
+		CurrentAssetsToConsider = SearchContext.GetAssetsToConsiderSet();
 	}
+	CurrentAssetIndexesToConsiderPerDatabase = SearchContext.GetAssetIndexesToConsiderPerDatabase();
+
 	SearchContext.SetContinuingInteraction(ContinuingProperties.bIsContinuingInteraction);
 	SearchContext.SetContinuingContextInteraction(ContinuingProperties.bIsContinuingContextInteraction);
 
 	// collecting all the databases searches in AssetsToSearchPerDatabase
 	// and all the continuing pose searches in ContinuingPoseAssetsToSearchPerDatabase
 	const bool bUsePoseSearchBranchIn = ShouldUsePoseSearchBranchIn(AssetsToSearch);
-	FAssetsToSearchPerDatabase AssetsToSearchPerDatabase;
-	FAssetsToSearchPerDatabase ContinuingPoseAssetsToSearchPerDatabase;
-	PopulateSearches(AssetsToSearch, SearchContext, AssetsToSearchPerDatabase, bUsePoseSearchBranchIn);
-	PopulateContinuingPoseSearches(ContinuingProperties, AssetsToSearch, SearchContext, ContinuingPoseAssetsToSearchPerDatabase, bUsePoseSearchBranchIn);
+	FStackDatabaseToAssetIndexes AssetIndexesToSearchPerDatabase;
+	FStackDatabaseToAssetIndexes ContinuingPoseAssetIndexesToSearchPerDatabase;
+
+	PopulateSearches(AssetsToSearch, SearchContext, AssetIndexesToSearchPerDatabase, bUsePoseSearchBranchIn);
+	PopulateContinuingPoseSearches(ContinuingProperties, AssetsToSearch, SearchContext, ContinuingPoseAssetIndexesToSearchPerDatabase, bUsePoseSearchBranchIn);
 
 	FCachedContinuingPoseSearchResults CachedContinuingPoseSearchResults;
-	for (const FAssetsToSearchPerDatabasePair& AssetsToSearchPerDatabasePair : ContinuingPoseAssetsToSearchPerDatabase)
+	for (const TPair<const UPoseSearchDatabase*, FStackAssetIndexes>& AssetIndexesToSearchPerDatabasePair : ContinuingPoseAssetIndexesToSearchPerDatabase)
 	{
-		const UPoseSearchDatabase* Database = AssetsToSearchPerDatabasePair.Key;
+		const UPoseSearchDatabase* Database = AssetIndexesToSearchPerDatabasePair.Key;
 		check(Database);
 
-		const bool bInvalidatingContinuingPose = IsInvalidatingContinuingPose(ContinuingProperties.InterruptMode, Database, AssetsToSearchPerDatabase);
+		const bool bInvalidatingContinuingPose = IsInvalidatingContinuingPose(ContinuingProperties.InterruptMode, Database, AssetIndexesToSearchPerDatabase);
 		if (!bInvalidatingContinuingPose)
 		{
-			const int32 ContinuingPoseIdx = Database->GetPoseIndex(ContinuingProperties.PlayingAsset.Get(), ContinuingProperties.PlayingAssetAccumulatedTime, ContinuingProperties.bIsPlayingAssetMirrored, ContinuingProperties.PlayingAssetBlendParameters);
+			const int32 ContinuingPoseIdx = Database->GetPoseIndex(ContinuingProperties.PlayingAsset.Get(), ContinuingProperties.PlayingAssetAccumulatedTime, ContinuingProperties.bIsPlayingAssetMirrored, ContinuingProperties.PlayingAssetBlendParameters, &AssetIndexesToSearchPerDatabasePair.Value);
 
 			if (ContinuingPoseIdx != INDEX_NONE)
 			{
@@ -1353,7 +1304,7 @@ void UPoseSearchLibrary::MotionMatch(
 				// (used to gather the continuing pose search values adopted to create the MM query - relative to the schema, NOT the database)
 				const FSearchResult& SchemaContinuingPoseSearchResult = CachedContinuingPoseSearchResults.FindOrAdd(Database->Schema, DatabaseContinuingPoseSearchResult);
 
-				const bool bForceInterrupt = IsForceInterrupt(ContinuingProperties.InterruptMode, Database, AssetsToSearchPerDatabase);
+				const bool bForceInterrupt = IsForceInterrupt(ContinuingProperties.InterruptMode, Database, AssetIndexesToSearchPerDatabase);
 				if (!bForceInterrupt)
 				{
 					SearchContext.UpdateContinuingPoseSearchResult(DatabaseContinuingPoseSearchResult, SchemaContinuingPoseSearchResult);
@@ -1364,29 +1315,26 @@ void UPoseSearchLibrary::MotionMatch(
 	}
 
 	// performing all the other databases searches
-	for (const FAssetsToSearchPerDatabasePair& AssetsToSearchPerDatabasePair : AssetsToSearchPerDatabase)
+	PRAGMA_DISABLE_DEPRECATION_WARNINGS
+	TArray<const UObject*> EmptyAssetsToConsider;
+	SearchContext.SetInternalDeprecatedAssetsToConsider(EmptyAssetsToConsider);
+	PRAGMA_ENABLE_DEPRECATION_WARNINGS
+
+	SearchContext.SetAssetsToConsiderSet(nullptr);
+	SearchContext.SetAssetIndexesToConsiderPerDatabase(&AssetIndexesToSearchPerDatabase);
+
+	for (const TPair<const UPoseSearchDatabase*, FStackAssetIndexes>& AssetIndexesToSearchPerDatabasePair : AssetIndexesToSearchPerDatabase)
 	{
-		const UPoseSearchDatabase* Database = AssetsToSearchPerDatabasePair.Key;
+		const UPoseSearchDatabase* Database = AssetIndexesToSearchPerDatabasePair.Key;
 		check(Database);
 
-		PRAGMA_DISABLE_DEPRECATION_WARNINGS
-		TArray<const UObject*> InternalDeprecatedAssetsToConsider = SearchContext.GetInternalDeprecatedAssetsToConsider();
-		TArray<const UObject*> EmptyAssetsToConsider;
-		SearchContext.SetInternalDeprecatedAssetsToConsider(EmptyAssetsToConsider);
-		PRAGMA_ENABLE_DEPRECATION_WARNINGS
-
 		// in case we haven't searched the continuing pose for this Database, we haven't created and cached the query yet,
 		// but if we didn't invalidated the continuing pose (when IsInvalidatingContinuingPose is true), we still can reuse
 		// the updated FirstInstanceOfReconstructedContinuingPoseSearchResult data, and by calling UpdateContinuingPoseSearchResult we set the 
 		// SearchContext to be able to create a query for Database using the continuing pose data.
 		SearchContext.UpdateContinuingPoseSearchResult(CachedContinuingPoseSearchResults.FindOrDefault(Database), CachedContinuingPoseSearchResults.FindOrDefault(Database->Schema));
-		SearchContext.SetAssetsToConsiderSet(&AssetsToSearchPerDatabasePair.Value);
 		
 		Database->Search(SearchContext, SearchResults);
-		
-		PRAGMA_DISABLE_DEPRECATION_WARNINGS
-		SearchContext.SetInternalDeprecatedAssetsToConsider(InternalDeprecatedAssetsToConsider);
-		PRAGMA_ENABLE_DEPRECATION_WARNINGS
 	}
 
 #if ENABLE_DRAW_DEBUG && ENABLE_ANIM_DEBUG
@@ -1412,13 +1360,15 @@ void UPoseSearchLibrary::MotionMatch(
 	TraceMotionMatching(SearchContext, SearchResults, 0.f, 1.f, ContinuingProperties.InterruptMode);
 #endif // UE_POSE_SEARCH_TRACE_ENABLED
 
-	// restoring the asset to consider array view
+	// restoring the assets to consider
 	PRAGMA_DISABLE_DEPRECATION_WARNINGS
-	if (SearchContext.GetInternalDeprecatedAssetsToConsider().IsEmpty())
+	SearchContext.SetInternalDeprecatedAssetsToConsider(InternalDeprecatedAssetsToConsider);
+	if (InternalDeprecatedAssetsToConsider.IsEmpty())
 	PRAGMA_ENABLE_DEPRECATION_WARNINGS
 	{
-		SearchContext.SetAssetsToConsiderSet(CurrentAssetToConsider);
+		SearchContext.SetAssetsToConsiderSet(CurrentAssetsToConsider);
 	}
+	SearchContext.SetAssetIndexesToConsiderPerDatabase(CurrentAssetIndexesToConsiderPerDatabase);
 }
 
 const FAnimNode_PoseSearchHistoryCollector_Base* UPoseSearchLibrary::FindPoseHistoryNode(
diff --git a/Engine/Plugins/Animation/PoseSearch/Source/Runtime/Public/PoseSearch/PoseSearchContext.h b/Engine/Plugins/Animation/PoseSearch/Source/Runtime/Public/PoseSearch/PoseSearchContext.h
--- a/Engine/Plugins/Animation/PoseSearch/Source/Runtime/Public/PoseSearch/PoseSearchContext.h
+++ b/Engine/Plugins/Animation/PoseSearch/Source/Runtime/Public/PoseSearch/PoseSearchContext.h
@@ -19,8 +19,15 @@ class UPoseSearchFeatureChannel_Position;
 
 namespace UE::PoseSearch
 {
-	
+
+// usually used as set of asset to consider during a search of a database 
 typedef TSet<const UObject*, DefaultKeyFuncs<const UObject*>, TInlineSetAllocator<8, TMemStackSetAllocator<>>> FStackAssetSet;
+// Experimental, this feature might be removed without warning, not for production use
+// usually used as set of asset indexes to consider (per database)
+// Experimental, this feature might be removed without warning, not for production use
+typedef TSet<int32, DefaultKeyFuncs<int32>, TInlineSetAllocator<8, TMemStackSetAllocator<>>> FStackAssetIndexes;
+// usually used as container to map database to what asset indexes to consider
+typedef TMap<const UPoseSearchDatabase*, FStackAssetIndexes, TInlineSetAllocator<8, TMemStackSetAllocator<>>> FStackDatabaseToAssetIndexes;
 
 // Utils functions
 const UObject* GetAnimContext(const FChooserEvaluationContext* Context);
@@ -307,6 +314,16 @@ struct FSearchContext
 		return AssetsToConsider;
 	}
 
+	void SetAssetIndexesToConsiderPerDatabase(const FStackDatabaseToAssetIndexes* InAssetIndexesToConsiderPerDatabase)
+	{
+		AssetIndexesToConsiderPerDatabase = InAssetIndexesToConsiderPerDatabase;
+	}
+
+	const FStackDatabaseToAssetIndexes* GetAssetIndexesToConsiderPerDatabase() const
+	{
+		return AssetIndexesToConsiderPerDatabase;
+	}
+	
 	// Experimental, this feature might be removed without warning, not for production use
 	const FPoseSearchEvent& GetEventToSearch() const { return EventToSearch; }
 	
@@ -351,6 +368,8 @@ private:
 	// if AssetsToConsider is not empty, we'll search only for poses from UObject(s) that are in the AssetsToConsider
 	const FStackAssetSet* AssetsToConsider = nullptr;
 	
+	const FStackDatabaseToAssetIndexes* AssetIndexesToConsiderPerDatabase = nullptr;
+
 	UE_DEPRECATED(5.7, "Unnecessary once deprecated GetAssetsToConsiderSet / SetAssetsToConsiderSet have been removed")
 	TArray<const UObject*> AssetsToConsider_DEPRECATED;
 	
diff --git a/Engine/Plugins/Animation/PoseSearch/Source/Runtime/Public/PoseSearch/PoseSearchDatabase.h b/Engine/Plugins/Animation/PoseSearch/Source/Runtime/Public/PoseSearch/PoseSearchDatabase.h
--- a/Engine/Plugins/Animation/PoseSearch/Source/Runtime/Public/PoseSearch/PoseSearchDatabase.h
+++ b/Engine/Plugins/Animation/PoseSearch/Source/Runtime/Public/PoseSearch/PoseSearchDatabase.h
@@ -641,7 +641,7 @@ public:
 	// given an AnimationAsset, a normalized time AnimationAssetTime (in [0-1] range for blend spaces), and mirror state bMirrored
 	// will retrieve the associated pose index from this database closest in distance to BlendParameters.
 	// It'll return INDEX_NONE if it cannot find a pose
-	UE_API int32 GetPoseIndex(const UObject* AnimationAsset, float AnimationAssetTime, bool bMirrored, const FVector& BlendParameters) const;
+	UE_API int32 GetPoseIndex(const UObject* AnimationAsset, float AnimationAssetTime, bool bMirrored, const FVector& BlendParameters, const UE::PoseSearch::FStackAssetIndexes* AssetIndexesToConsider = nullptr) const;
 
 	UE_DEPRECATED(5.7, "Use AddAnimationAsset(const FPoseSearchDatabaseAnimationAsset& AnimationAsset) instead")
 	UE_API void AddAnimationAsset(FInstancedStruct AnimationAsset);
@@ -731,6 +731,7 @@ private:
 	UE_DEPRECATED(5.7, "Use PopulateSelectableAssetIdx taking FStackAssetSet as input instead")
 	void PopulateSelectableAssetIdx(FSelectableAssetIdx& SelectableAssetIdx, TConstArrayView<const UObject*> AssetsToConsider) const;
 	void PopulateSelectableAssetIdx(FSelectableAssetIdx& SelectableAssetIdx, const UE::PoseSearch::FStackAssetSet* AssetsToConsider) const;
+	void PopulateSelectableAssetIdx(FSelectableAssetIdx& SelectableAssetIdx, const UE::PoseSearch::FStackDatabaseToAssetIndexes* AssetIndexesToConsiderPerDatabase) const;
 
 	typedef TArray<int32, TInlineAllocator<256, TMemStackAllocator<>>> FNonSelectableIdx;
 	void PopulateNonSelectableIdx(FNonSelectableIdx& NonSelectableIdx, UE::PoseSearch::FSearchContext& SearchContext
