PoseSearch - avoiding multiple mem stack reallocations in AddToSearch, AddToSearchForDatabase

#Synced-CL 49608181

diff --git a/Engine/Plugins/Animation/PoseSearch/Source/Runtime/Private/PoseSearchLibrary.cpp b/Engine/Plugins/Animation/PoseSearch/Source/Runtime/Private/PoseSearchLibrary.cpp
--- a/Engine/Plugins/Animation/PoseSearch/Source/Runtime/Private/PoseSearchLibrary.cpp
+++ b/Engine/Plugins/Animation/PoseSearch/Source/Runtime/Private/PoseSearchLibrary.cpp
@@ -216,6 +216,7 @@ namespace UE::PoseSearch
 			{
 				const FSearchIndex& SearchIndex = Database->GetSearchIndex();
 				FStackAssetIndexes& AssetIndexesToSearch = AssetIndexesToSearchPerDatabase.FindOrAdd(Database);
+				AssetIndexesToSearch.Reserve(AssetIndexesToSearch.Num() + AssetIndexesForSourceAssetToSearch.Num());
 				for (int32 AssetIndexToSearch : AssetIndexesForSourceAssetToSearch)
 				{
 					const FSearchIndexAsset& SearchIndexAsset = SearchIndex.Assets[AssetIndexToSearch];
@@ -263,6 +264,7 @@ namespace UE::PoseSearch
 			{
 				FStackAssetIndexes& AssetIndexesToSearch = AssetIndexesToSearchPerDatabase.FindOrAdd(Database);
 				const int32 NumAnimationAssets = Database->GetNumAnimationAssets();
+				AssetIndexesToSearch.Reserve(NumAnimationAssets);
 				for (int32 AnimationAssetIndex = 0; AnimationAssetIndex < NumAnimationAssets; ++AnimationAssetIndex)
 				{
 					AssetIndexesToSearch.Add(AnimationAssetIndex);
