PoseSearch - fix for FAssetIndexer::Process crashing on schemas with roles improperly setup

#Synced-CL 49689419

diff --git a/Engine/Plugins/Animation/PoseSearch/Source/Runtime/Private/PoseSearchAssetIndexer.cpp b/Engine/Plugins/Animation/PoseSearch/Source/Runtime/Private/PoseSearchAssetIndexer.cpp
--- a/Engine/Plugins/Animation/PoseSearch/Source/Runtime/Private/PoseSearchAssetIndexer.cpp
+++ b/Engine/Plugins/Animation/PoseSearch/Source/Runtime/Private/PoseSearchAssetIndexer.cpp
@@ -180,85 +180,91 @@ void FAssetIndexer::AssignWorkingData(int32 InStartPoseIdx, TArrayView<float> In
 
 void FAssetIndexer::Process(int32 AssetIdx)
 {
-	bProcessFailed = false;
+	if (AssetSamplers.Num() == 0 || PoseMetadata.IsEmpty())
+	{
+		bProcessFailed = true;
+	}
+	else
+	{
+		bProcessFailed = false;
 
-	// Generate pose metadata
-	const float PlayLength = GetPlayLength();
-	FAnimNotifyContext PreAllocatedNotifyContext;
+		// Generate pose metadata
+		const float PlayLength = GetPlayLength();
+		FAnimNotifyContext PreAllocatedNotifyContext;
 
-	// @todo: optimize this code, by extracting ALL the notify states, and then perform time overlap with the poses
-	for (int32 SampleIdx = GetBeginSampleIdx(); SampleIdx != GetEndSampleIdx(); ++SampleIdx)
-	{
-		const float SampleTime = FMath::Min(CalculateSampleTime(SampleIdx), PlayLength);
-		float CostAddend = SamplingContext.BaseCostBias;
-		bool bBlockTransition = false;
+		// @todo: optimize this code, by extracting ALL the notify states, and then perform time overlap with the poses
+		for (int32 SampleIdx = GetBeginSampleIdx(); SampleIdx != GetEndSampleIdx(); ++SampleIdx)
+		{
+			const float SampleTime = FMath::Min(CalculateSampleTime(SampleIdx), PlayLength);
+			float CostAddend = SamplingContext.BaseCostBias;
+			bool bBlockTransition = false;
 
-		AssetSamplers.ExtractAnimNotifyStates(SampleTime, PreAllocatedNotifyContext, [&bBlockTransition, &CostAddend](const UAnimNotifyState* AnimNotifyState)
-			{
-				if (AnimNotifyState->GetClass()->IsChildOf<UAnimNotifyState_PoseSearchBlockTransition>())
-				{
-					bBlockTransition = true;
-				}
-				else if (AnimNotifyState->GetClass()->IsChildOf<UAnimNotifyState_PoseSearchModifyCost>())
+			AssetSamplers.ExtractAnimNotifyStates(SampleTime, PreAllocatedNotifyContext, [&bBlockTransition, &CostAddend](const UAnimNotifyState* AnimNotifyState)
 				{
-					const UAnimNotifyState_PoseSearchModifyCost* ModifyCostNotifyState = Cast<const UAnimNotifyState_PoseSearchModifyCost>(AnimNotifyState);
-					check(ModifyCostNotifyState);
-					CostAddend = ModifyCostNotifyState->CostAddend;
-				}
-				return true;
-			});
+					if (AnimNotifyState->GetClass()->IsChildOf<UAnimNotifyState_PoseSearchBlockTransition>())
+					{
+						bBlockTransition = true;
+					}
+					else if (AnimNotifyState->GetClass()->IsChildOf<UAnimNotifyState_PoseSearchModifyCost>())
+					{
+						const UAnimNotifyState_PoseSearchModifyCost* ModifyCostNotifyState = Cast<const UAnimNotifyState_PoseSearchModifyCost>(AnimNotifyState);
+						check(ModifyCostNotifyState);
+						CostAddend = ModifyCostNotifyState->CostAddend;
+					}
+					return true;
+				});
 
-		if (AssetSamplers.IsLoopable())
-		{
-			CostAddend += SamplingContext.LoopingCostBias;
-		}
+			if (AssetSamplers.IsLoopable())
+			{
+				CostAddend += SamplingContext.LoopingCostBias;
+			}
 
-		const int32 VectorIdx = GetVectorIdx(SampleIdx);
-		const int32 PoseIdx = StartPoseIdx + VectorIdx;
-		const int32 ValueOffset = PoseIdx * Schema.SchemaCardinality;
-		check(ValueOffset >= 0 && AssetIdx >= 0);
-		PoseMetadata[VectorIdx] = FPoseMetadata(ValueOffset, AssetIdx, bBlockTransition, CostAddend);
-	}
+			const int32 VectorIdx = GetVectorIdx(SampleIdx);
+			const int32 PoseIdx = StartPoseIdx + VectorIdx;
+			const int32 ValueOffset = PoseIdx * Schema.SchemaCardinality;
+			check(ValueOffset >= 0 && AssetIdx >= 0);
+			PoseMetadata[VectorIdx] = FPoseMetadata(ValueOffset, AssetIdx, bBlockTransition, CostAddend);
+		}
 
-	
-	AssetSamplers.ProcessAllAnimNotifyEvents([this](const TConstArrayView<FAnimNotifyEvent> AnimNotifyEvents)
-		{
-			for (const FAnimNotifyEvent& AnimNotifyEvent : AnimNotifyEvents)
+		AssetSamplers.ProcessAllAnimNotifyEvents([this](const TConstArrayView<FAnimNotifyEvent> AnimNotifyEvents)
 			{
-				if (AnimNotifyEvent.Notify && AnimNotifyEvent.Notify->GetClass()->IsChildOf<UAnimNotify_PoseSearchEvent>())
+				for (const FAnimNotifyEvent& AnimNotifyEvent : AnimNotifyEvents)
 				{
-					const UAnimNotify_PoseSearchEvent* EventNotify = Cast<const UAnimNotify_PoseSearchEvent>(AnimNotifyEvent.Notify);
-					check(EventNotify);
-					if (EventNotify->EventTag.IsValid())
+					if (AnimNotifyEvent.Notify && AnimNotifyEvent.Notify->GetClass()->IsChildOf<UAnimNotify_PoseSearchEvent>())
 					{
-						const int32 PoseIdx = SearchIndexAsset.GetPoseIndexFromTime(AnimNotifyEvent.GetTime(), Schema.SampleRate);
-						if (PoseIdx != INDEX_NONE)
+						const UAnimNotify_PoseSearchEvent* EventNotify = Cast<const UAnimNotify_PoseSearchEvent>(AnimNotifyEvent.Notify);
+						check(EventNotify);
+						if (EventNotify->EventTag.IsValid())
 						{
-							EventDataCollector.Emplace(EventNotify->EventTag, PoseIdx);
+							const int32 PoseIdx = SearchIndexAsset.GetPoseIndexFromTime(AnimNotifyEvent.GetTime(), Schema.SampleRate);
+							if (PoseIdx != INDEX_NONE)
+							{
+								EventDataCollector.Emplace(EventNotify->EventTag, PoseIdx);
+							}
 						}
 					}
 				}
-			}
-			return true;
-		});
+				return true;
+			});
 
-	// Generate pose features data
-	if (Schema.SchemaCardinality > 0)
-	{
-		for (const TObjectPtr<UPoseSearchFeatureChannel>& ChannelPtr : Schema.GetChannels())
+		// Generate pose features data
+		if (Schema.SchemaCardinality > 0)
 		{
-			if (!ChannelPtr->IndexAsset(*this))
+			for (const TObjectPtr<UPoseSearchFeatureChannel>& ChannelPtr : Schema.GetChannels())
 			{
-				bProcessFailed = true;
-				break;
+				if (!ChannelPtr->IndexAsset(*this))
+				{
+					bProcessFailed = true;
+					break;
+				}
 			}
 		}
-	}
 
-	// Computing stats
-	if (!bProcessFailed)
-	{
-		ComputeStats();
+		// Computing stats
+		if (!bProcessFailed)
+		{
+			ComputeStats();
+		}
 	}
 }
 
