PoseSearch - Created UPoseSearchSettings class with AvailabilitiesBufferSize property to control initial allocation of AnimContextsAvailabilitiesBuffer in UPoseSearchInteractionSubsystem
Contribution from daniel.ruizmanero

#Synced-CL 49689419

diff --git a/Engine/Plugins/Animation/PoseSearch/Source/Runtime/PoseSearch.Build.cs b/Engine/Plugins/Animation/PoseSearch/Source/Runtime/PoseSearch.Build.cs
--- a/Engine/Plugins/Animation/PoseSearch/Source/Runtime/PoseSearch.Build.cs
+++ b/Engine/Plugins/Animation/PoseSearch/Source/Runtime/PoseSearch.Build.cs
@@ -30,6 +30,7 @@ public class PoseSearch : ModuleRules
 		PrivateDependencyModuleNames.AddRange(
 		new string[]{
 			"AnimationWarpingRuntime",
+			"DeveloperSettings",
 			"GameplayTags",
 			"RewindDebuggerRuntimeInterface"
 		}
diff --git a/Engine/Plugins/Animation/PoseSearch/Source/Runtime/Private/PoseSearchInteractionSubsystem.cpp b/Engine/Plugins/Animation/PoseSearch/Source/Runtime/Private/PoseSearchInteractionSubsystem.cpp
--- a/Engine/Plugins/Animation/PoseSearch/Source/Runtime/Private/PoseSearchInteractionSubsystem.cpp
+++ b/Engine/Plugins/Animation/PoseSearch/Source/Runtime/Private/PoseSearchInteractionSubsystem.cpp
@@ -9,6 +9,7 @@
 #include "PoseSearch/PoseSearchInteractionUtils.h"
 #include "PoseSearch/PoseSearchRole.h"
 #include "PoseSearch/PoseSearchSchema.h"
+#include "PoseSearch/PoseSearchSettings.h"
 #include "PoseSearch/Trace/PoseSearchTraceLogger.h"
 #include "Util/GridIndexing2.h"
 #include "VisualLogger/VisualLogger.h"
@@ -600,6 +601,9 @@ void UPoseSearchInteractionSubsystem::Initialize(FSubsystemCollectionBase& Colle
 {
 	check(IsInGameThread());
 	Super::Initialize(Collection);
+
+	const UPoseSearchSettings& Settings = UPoseSearchSettings::Get();
+	AnimContextsAvailabilitiesBuffer.SetNum(Settings.AvailabilitiesBufferSize);
 }
 
 void UPoseSearchInteractionSubsystem::Deinitialize()
@@ -1356,10 +1360,9 @@ bool UPoseSearchInteractionSubsystem::ShouldRegenerateAllIslands()
 	const int32 AnimContextsAvailabilitiesIndexValue = AnimContextsAvailabilitiesIndex.GetValue();
 	AnimContextsAvailabilitiesNum = AnimContextsAvailabilitiesBuffer.Num();
 
-	// @todo: add some setttings to initialize AnimContextsAvailabilitiesBuffer to a big enough value
 	if (AnimContextsAvailabilitiesIndexValue > AnimContextsAvailabilitiesNum)
 	{
-		UE_LOG(LogPoseSearch, Warning, TEXT("UPoseSearchInteractionSubsystem::ShouldRegenerateAllIslands not enough space to add more availabilities locklessly. It'll be adjusted automatically, but some availability requests has been lost this frame [capacity %d / requests %d]"), AnimContextsAvailabilitiesNum, AnimContextsAvailabilitiesIndexValue);
+		UE_LOG(LogPoseSearch, Warning, TEXT("UPoseSearchInteractionSubsystem::ShouldRegenerateAllIslands not enough space to add more availabilities locklessly. It'll be adjusted automatically, but some availability requests has been lost this frame [capacity %d / requests %d]. To void this warning set the initial budget for 'Edit/ProjectSettings/Plugins/Pose Search/Availabilities Buffer Size' to %d. "), AnimContextsAvailabilitiesNum, AnimContextsAvailabilitiesIndexValue, AnimContextsAvailabilitiesIndexValue);
 		AnimContextsAvailabilitiesBuffer.SetNum(AnimContextsAvailabilitiesIndexValue);
 	}
 	else
diff --git a/Engine/Plugins/Animation/PoseSearch/Source/Runtime/Private/PoseSearchSettings.cpp b/Engine/Plugins/Animation/PoseSearch/Source/Runtime/Private/PoseSearchSettings.cpp
new file mode 100644
--- /dev/null
+++ b/Engine/Plugins/Animation/PoseSearch/Source/Runtime/Private/PoseSearchSettings.cpp
@@ -0,0 +1,16 @@
+// Copyright Epic Games, Inc. All Rights Reserved.
+
+#include "PoseSearch/PoseSearchSettings.h"
+
+UPoseSearchSettings::UPoseSearchSettings()
+{
+	CategoryName = TEXT("Plugins");
+}
+
+const UPoseSearchSettings& UPoseSearchSettings::Get()
+{
+	UPoseSearchSettings* MutableCDO = GetMutableDefault<UPoseSearchSettings>();
+	check(MutableCDO != nullptr)
+
+	return *MutableCDO;
+}
diff --git a/Engine/Plugins/Animation/PoseSearch/Source/Runtime/Public/PoseSearch/PoseSearchSettings.h b/Engine/Plugins/Animation/PoseSearch/Source/Runtime/Public/PoseSearch/PoseSearchSettings.h
new file mode 100644
--- /dev/null
+++ b/Engine/Plugins/Animation/PoseSearch/Source/Runtime/Public/PoseSearch/PoseSearchSettings.h
@@ -0,0 +1,22 @@
+// Copyright Epic Games, Inc. All Rights Reserved.
+
+#pragma once
+
+#include "Engine/DeveloperSettings.h"
+#include "PoseSearchSettings.generated.h"
+
+/** Default PoseSearch settings. */
+UCLASS(Config=Engine, defaultconfig, meta=(DisplayName="Pose Search"), MinimalAPI)
+class UPoseSearchSettings : public UDeveloperSettings
+{
+	GENERATED_BODY()
+
+	UPoseSearchSettings();
+
+public:
+	static const UPoseSearchSettings& Get();
+
+	/** Initial size for the interaction availabilities buffer */
+	UPROPERTY(EditAnywhere, config, Category = "Interaction", meta = (ClampMin = "0"))
+	int32 AvailabilitiesBufferSize = 8;
+};
\ No newline at end of file
