PoseSearch - PSIA editor fix crash on undo

#Synced-CL 49689419

diff --git a/Engine/Plugins/Animation/PoseSearch/Source/Editor/Private/PoseSearchInteractionAssetEditor.cpp b/Engine/Plugins/Animation/PoseSearch/Source/Editor/Private/PoseSearchInteractionAssetEditor.cpp
--- a/Engine/Plugins/Animation/PoseSearch/Source/Editor/Private/PoseSearchInteractionAssetEditor.cpp
+++ b/Engine/Plugins/Animation/PoseSearch/Source/Editor/Private/PoseSearchInteractionAssetEditor.cpp
@@ -341,7 +341,7 @@ void FInteractionAssetViewModel::Tick(float DeltaSeconds)
 
 	// testing CalculateWarpTransforms
 #if WITH_EDITORONLY_DATA
-	if (InteractionAsset->bEnableDebugWarp && PreviewActors.Num() == NumRoles)
+	if (InteractionAsset->bEnableDebugWarp && PreviewActors.Num() == NumRoles && PreviewActorRoleToIndex.Num() == NumRoles)
 	{
 		TArray<FTransform, TInlineAllocator<PreallocatedRolesNum>> ActorTransforms;
 		ActorTransforms.SetNum(NumRoles);
