PoseSearch - fix for Navigation issues in the Viewport of Pose Search Interaction Asset

#Synced-CL 49747430

diff --git a/Engine/Plugins/Animation/PoseSearch/Source/Editor/Private/PoseSearchInteractionAssetEditor.cpp b/Engine/Plugins/Animation/PoseSearch/Source/Editor/Private/PoseSearchInteractionAssetEditor.cpp
--- a/Engine/Plugins/Animation/PoseSearch/Source/Editor/Private/PoseSearchInteractionAssetEditor.cpp
+++ b/Engine/Plugins/Animation/PoseSearch/Source/Editor/Private/PoseSearchInteractionAssetEditor.cpp
@@ -272,169 +272,174 @@ void FInteractionAssetViewModel::Tick(float DeltaSeconds)
 	const UPoseSearchInteractionAsset* InteractionAsset = GetInteractionAsset();
 	if (!InteractionAsset)
 	{
-		RemovePreviewActors();
-		return;
-	}
+		PreviewPause();
+		SetPlayTime(0.f, false);
 
-	PlayTime += DeltaSeconds * DeltaTimeMultiplier;
+		for (FInteractionAssetPreviewActor& PreviewActor : PreviewActors)
+		{
+			PreviewActor.Destroy();
+		}
 
-	const int32 NumRoles = InteractionAsset->GetNumRoles();
-	FRoleToIndex InteractionAssetRoleToIndex;
-	for (int32 RoleIndex = 0; RoleIndex < NumRoles; ++RoleIndex)
-	{
-		InteractionAssetRoleToIndex.Add(InteractionAsset->GetRole(RoleIndex)) = RoleIndex;
+		PreviewActors.Reset();
 	}
-
-	// iterating backwards because of the possible RemoveAtSwap 
-	FRoleToIndex PreviewActorsRoleToIndex;
-	for (int32 ActorIndex = PreviewActors.Num() - 1; ActorIndex >= 0; --ActorIndex)
+	else
 	{
-		FInteractionAssetPreviewActor& PreviewActor = PreviewActors[ActorIndex];
-		if (!InteractionAssetRoleToIndex.Find(PreviewActor.GetRole()))
+		PlayTime += DeltaSeconds * DeltaTimeMultiplier;
+
+		const int32 NumRoles = InteractionAsset->GetNumRoles();
+		FRoleToIndex InteractionAssetRoleToIndex;
+		for (int32 RoleIndex = 0; RoleIndex < NumRoles; ++RoleIndex)
 		{
-			PreviewActor.Destroy();
-			PreviewActors.RemoveAtSwap(ActorIndex, EAllowShrinking::No);
+			InteractionAssetRoleToIndex.Add(InteractionAsset->GetRole(RoleIndex)) = RoleIndex;
 		}
-		else
+
+		// iterating backwards because of the possible RemoveAtSwap 
+		bool bAnyAddedOrRemovedActors = false;
+		FRoleToIndex PreviewActorsRoleToIndex;
+		for (int32 ActorIndex = PreviewActors.Num() - 1; ActorIndex >= 0; --ActorIndex)
 		{
-			PreviewActorsRoleToIndex.Add(PreviewActor.GetRole());		
+			FInteractionAssetPreviewActor& PreviewActor = PreviewActors[ActorIndex];
+			if (!InteractionAssetRoleToIndex.Find(PreviewActor.GetRole()))
+			{
+				PreviewActor.Destroy();
+				PreviewActors.RemoveAtSwap(ActorIndex, EAllowShrinking::No);
+				bAnyAddedOrRemovedActors = true;
+			}
+			else
+			{
+				PreviewActorsRoleToIndex.Add(PreviewActor.GetRole());
+			}
 		}
-	}
 
-	if (PreviewActors.Num() != InteractionAssetRoleToIndex.Num())
-	{
-		for (int32 RoleIndex = 0; RoleIndex < NumRoles; ++RoleIndex)
+		if (PreviewActors.Num() != InteractionAssetRoleToIndex.Num())
 		{
-			if (!PreviewActorsRoleToIndex.Find(InteractionAsset->GetRole(RoleIndex)))
+			for (int32 RoleIndex = 0; RoleIndex < NumRoles; ++RoleIndex)
 			{
-				FInteractionAssetPreviewActor PreviewActor;
-				if (PreviewActor.SpawnPreviewActor(GetWorld(), InteractionAsset, InteractionAsset->GetRole(RoleIndex)))
+				if (!PreviewActorsRoleToIndex.Find(InteractionAsset->GetRole(RoleIndex)))
 				{
-					PreviewActors.Add(MoveTemp(PreviewActor));
+					FInteractionAssetPreviewActor PreviewActor;
+					if (PreviewActor.SpawnPreviewActor(GetWorld(), InteractionAsset, InteractionAsset->GetRole(RoleIndex)))
+					{
+						PreviewActors.Add(MoveTemp(PreviewActor));
+						bAnyAddedOrRemovedActors = true;
+					}
 				}
 			}
 		}
-	}
 
-	for (FInteractionAssetPreviewActor& PreviewActor : PreviewActors)
-	{
-		PreviewActor.UpdatePreviewActor(InteractionAsset, PlayTime);
-	}
-	
-	FRoleToIndex PreviewActorRoleToIndex;
-	PreviewActorRoleToIndex.Reserve(PreviewActors.Num());
+		if (bAnyAddedOrRemovedActors)
+		{
+			PreviewPause();
+			SetPlayTime(0.f, false);
+		}
 
-	TArray<TObjectPtr<const UObject>, TInlineAllocator<PreallocatedRolesNum>> AnimContexts;
-	AnimContexts.SetNum(NumRoles);
+		for (FInteractionAssetPreviewActor& PreviewActor : PreviewActors)
+		{
+			PreviewActor.UpdatePreviewActor(InteractionAsset, PlayTime);
+		}
 
-	TArray<FRole, TInlineAllocator<PreallocatedRolesNum>> Roles;
-	Roles.SetNum(NumRoles);
+		FRoleToIndex PreviewActorRoleToIndex;
+		PreviewActorRoleToIndex.Reserve(PreviewActors.Num());
 
-	for (int32 PreviewActorIndex = 0; PreviewActorIndex < PreviewActors.Num(); ++PreviewActorIndex)
-	{
-		PreviewActorRoleToIndex.Add(PreviewActors[PreviewActorIndex].GetRole()) = PreviewActorIndex;
+		TArray<TObjectPtr<const UObject>, TInlineAllocator<PreallocatedRolesNum>> AnimContexts;
+		AnimContexts.SetNum(NumRoles);
 
-		FInteractionAssetPreviewActor& PreviewActor = PreviewActors[PreviewActorIndex];
-		AnimContexts[PreviewActorIndex] = PreviewActor.GetAnimPreviewInstance();
-		Roles[PreviewActorIndex] = PreviewActor.GetRole();
-	}
+		TArray<FRole, TInlineAllocator<PreallocatedRolesNum>> Roles;
+		Roles.SetNum(NumRoles);
 
-	// testing CalculateWarpTransforms
-#if WITH_EDITORONLY_DATA
-	if (InteractionAsset->bEnableDebugWarp && PreviewActors.Num() == NumRoles && PreviewActorRoleToIndex.Num() == NumRoles)
-	{
-		TArray<FTransform, TInlineAllocator<PreallocatedRolesNum>> ActorTransforms;
-		ActorTransforms.SetNum(NumRoles);
-		for (int32 InteractionAssetRoleIndex = 0; InteractionAssetRoleIndex < NumRoles; ++InteractionAssetRoleIndex)
+		for (int32 PreviewActorIndex = 0; PreviewActorIndex < PreviewActors.Num(); ++PreviewActorIndex)
 		{
-			const FRole& InteractionAssetRole = InteractionAsset->GetRole(InteractionAssetRoleIndex);
-			const int32 PreviewActorIndex = PreviewActorRoleToIndex[InteractionAssetRole];
+			PreviewActorRoleToIndex.Add(PreviewActors[PreviewActorIndex].GetRole()) = PreviewActorIndex;
 
 			FInteractionAssetPreviewActor& PreviewActor = PreviewActors[PreviewActorIndex];
-			check(PreviewActor.GetRole() == InteractionAssetRole);
-
-			ActorTransforms[InteractionAssetRoleIndex] = PreviewActor.GetDebugActorTransformFromSampler();
+			AnimContexts[PreviewActorIndex] = PreviewActor.GetAnimPreviewInstance();
+			Roles[PreviewActorIndex] = PreviewActor.GetRole();
 		}
 
-		TArray<FTransform, TInlineAllocator<PreallocatedRolesNum>> FullAlignedActorTransforms;
-		FullAlignedActorTransforms.SetNum(NumRoles);
-		const TConstArrayView<const UMirrorDataTable*> MirrorDataTables;
-		const TConstArrayView<bool> RelevantRoleIndexes;
+		// testing CalculateWarpTransforms
+#if WITH_EDITORONLY_DATA
+		if (InteractionAsset->bEnableDebugWarp && PreviewActors.Num() == NumRoles && PreviewActorRoleToIndex.Num() == NumRoles)
+		{
+			TArray<FTransform, TInlineAllocator<PreallocatedRolesNum>> ActorTransforms;
+			ActorTransforms.SetNum(NumRoles);
+			for (int32 InteractionAssetRoleIndex = 0; InteractionAssetRoleIndex < NumRoles; ++InteractionAssetRoleIndex)
+			{
+				const FRole& InteractionAssetRole = InteractionAsset->GetRole(InteractionAssetRoleIndex);
+				const int32 PreviewActorIndex = PreviewActorRoleToIndex[InteractionAssetRole];
 
-		InteractionAsset->CalculateWarpTransforms(PlayTime, 0.f, ActorTransforms, FullAlignedActorTransforms, MirrorDataTables, RelevantRoleIndexes, AnimContexts);
+				FInteractionAssetPreviewActor& PreviewActor = PreviewActors[PreviewActorIndex];
+				check(PreviewActor.GetRole() == InteractionAssetRole);
 
-#if ENABLE_ANIM_DEBUG
-		const UWorld* DebugWorld = GetWorld();
-#endif // ENABLE_ANIM_DEBUG
+				ActorTransforms[InteractionAssetRoleIndex] = PreviewActor.GetDebugActorTransformFromSampler();
+			}
 
-		for (int32 InteractionAssetRoleIndex = 0; InteractionAssetRoleIndex < NumRoles; ++InteractionAssetRoleIndex)
-		{
-			const FRole& InteractionAssetRole = InteractionAsset->GetRole(InteractionAssetRoleIndex);
-			const int32 PreviewActorIndex = PreviewActorRoleToIndex[InteractionAssetRole];
+			TArray<FTransform, TInlineAllocator<PreallocatedRolesNum>> FullAlignedActorTransforms;
+			FullAlignedActorTransforms.SetNum(NumRoles);
+			const TConstArrayView<const UMirrorDataTable*> MirrorDataTables;
+			const TConstArrayView<bool> RelevantRoleIndexes;
 
-			FInteractionAssetPreviewActor& PreviewActor = PreviewActors[PreviewActorIndex];
-			check(PreviewActor.GetRole() == InteractionAssetRole);
-			FTransform DebugActorTransform;
-			DebugActorTransform.Blend(ActorTransforms[InteractionAssetRoleIndex], FullAlignedActorTransforms[InteractionAssetRoleIndex], InteractionAsset->DebugWarpAmount);
+			InteractionAsset->CalculateWarpTransforms(PlayTime, 0.f, ActorTransforms, FullAlignedActorTransforms, MirrorDataTables, RelevantRoleIndexes, AnimContexts);
 
 #if ENABLE_ANIM_DEBUG
-			// FullAlignedActorTransforms drawing 
-			const FVector FullAlignedActorPosition = FullAlignedActorTransforms[InteractionAssetRoleIndex].GetTranslation();
-			const FQuat FullAlignedActorRotation = FullAlignedActorTransforms[InteractionAssetRoleIndex].GetRotation();
-			DrawDebugCircle(DebugWorld, FMatrix(FVector::ZAxisVector, FVector::XAxisVector, FVector::YAxisVector, FullAlignedActorPosition), 5.f, 16, FColorList::Yellow, false, 0.f, SDPG_Foreground, 0.f, false);
-			DrawDebugLine(DebugWorld, FullAlignedActorPosition, FullAlignedActorPosition + FullAlignedActorRotation.RotateVector(FVector::XAxisVector) * 10, FColorList::Yellow, false, 0.f, SDPG_Foreground, 0.f);
+			const UWorld* DebugWorld = GetWorld();
 #endif // ENABLE_ANIM_DEBUG
 
-			PreviewActor.ForceDebugActorTransform(DebugActorTransform);
-		}
-
-#if ENABLE_ANIM_DEBUG
-		if (InteractionAsset->DebugTimeOffset > UE_KINDA_SMALL_NUMBER)
-		{
-			InteractionAsset->CalculateWarpTransforms(PlayTime, InteractionAsset->DebugTimeOffset, ActorTransforms, FullAlignedActorTransforms, MirrorDataTables, RelevantRoleIndexes, AnimContexts);
-
 			for (int32 InteractionAssetRoleIndex = 0; InteractionAssetRoleIndex < NumRoles; ++InteractionAssetRoleIndex)
 			{
+				const FRole& InteractionAssetRole = InteractionAsset->GetRole(InteractionAssetRoleIndex);
+				const int32 PreviewActorIndex = PreviewActorRoleToIndex[InteractionAssetRole];
+
+				FInteractionAssetPreviewActor& PreviewActor = PreviewActors[PreviewActorIndex];
+				check(PreviewActor.GetRole() == InteractionAssetRole);
+				FTransform DebugActorTransform;
+				DebugActorTransform.Blend(ActorTransforms[InteractionAssetRoleIndex], FullAlignedActorTransforms[InteractionAssetRoleIndex], InteractionAsset->DebugWarpAmount);
+
+#if ENABLE_ANIM_DEBUG
 				// FullAlignedActorTransforms drawing 
 				const FVector FullAlignedActorPosition = FullAlignedActorTransforms[InteractionAssetRoleIndex].GetTranslation();
 				const FQuat FullAlignedActorRotation = FullAlignedActorTransforms[InteractionAssetRoleIndex].GetRotation();
-				DrawDebugCircle(DebugWorld, FMatrix(FVector::ZAxisVector, FVector::XAxisVector, FVector::YAxisVector, FullAlignedActorPosition), 5.f, 16, FColorList::Pink, false, 0.f, SDPG_Foreground, 0.f, false);
-				DrawDebugLine(DebugWorld, FullAlignedActorPosition, FullAlignedActorPosition + FullAlignedActorRotation.RotateVector(FVector::XAxisVector) * 10, FColorList::Pink, false, 0.f, SDPG_Foreground, 0.f);
+				DrawDebugCircle(DebugWorld, FMatrix(FVector::ZAxisVector, FVector::XAxisVector, FVector::YAxisVector, FullAlignedActorPosition), 5.f, 16, FColorList::Yellow, false, 0.f, SDPG_Foreground, 0.f, false);
+				DrawDebugLine(DebugWorld, FullAlignedActorPosition, FullAlignedActorPosition + FullAlignedActorRotation.RotateVector(FVector::XAxisVector) * 10, FColorList::Yellow, false, 0.f, SDPG_Foreground, 0.f);
+#endif // ENABLE_ANIM_DEBUG
+
+				PreviewActor.ForceDebugActorTransform(DebugActorTransform);
+			}
+
+#if ENABLE_ANIM_DEBUG
+			if (InteractionAsset->DebugTimeOffset > UE_KINDA_SMALL_NUMBER)
+			{
+				InteractionAsset->CalculateWarpTransforms(PlayTime, InteractionAsset->DebugTimeOffset, ActorTransforms, FullAlignedActorTransforms, MirrorDataTables, RelevantRoleIndexes, AnimContexts);
+
+				for (int32 InteractionAssetRoleIndex = 0; InteractionAssetRoleIndex < NumRoles; ++InteractionAssetRoleIndex)
+				{
+					// FullAlignedActorTransforms drawing 
+					const FVector FullAlignedActorPosition = FullAlignedActorTransforms[InteractionAssetRoleIndex].GetTranslation();
+					const FQuat FullAlignedActorRotation = FullAlignedActorTransforms[InteractionAssetRoleIndex].GetRotation();
+					DrawDebugCircle(DebugWorld, FMatrix(FVector::ZAxisVector, FVector::XAxisVector, FVector::YAxisVector, FullAlignedActorPosition), 5.f, 16, FColorList::Pink, false, 0.f, SDPG_Foreground, 0.f, false);
+					DrawDebugLine(DebugWorld, FullAlignedActorPosition, FullAlignedActorPosition + FullAlignedActorRotation.RotateVector(FVector::XAxisVector) * 10, FColorList::Pink, false, 0.f, SDPG_Foreground, 0.f);
+				}
 			}
-		}
 #endif // ENABLE_ANIM_DEBUG
-	}
+		}
 #endif // WITH_EDITORONLY_DATA
 
-	// extracting and drawing Constraints
-	
-	// @todo: fill up those parameters
-	const FVector BlendParameters = FVector::ZeroVector;
-	const bool bIsMirrored = false;
+		// extracting and drawing Constraints
 
-	FConstraints Constraints;
-	Constraints.ProcessAll(InteractionAsset, StepDeltaTime, PlayTime, bIsMirrored, BlendParameters, AnimContexts, Roles);
-	
-	if (const UWorld* World = GetWorld())
-	{
-		for (const FConstraint& Constraint : Constraints.GetAll())
-		{
-			DrawDebugLine(World, Constraint.FromSocketTransform.GetTranslation(), Constraint.ToSocketTransform.GetTranslation(), FColor::Red, false, -1.f, SDPG_Foreground);
-		}
-	}
-}
+		// @todo: fill up those parameters
+		const FVector BlendParameters = FVector::ZeroVector;
+		const bool bIsMirrored = false;
 
-void FInteractionAssetViewModel::RemovePreviewActors()
-{
-	PlayTime = 0.f;
-	DeltaTimeMultiplier = 1.f;
+		FConstraints Constraints;
+		Constraints.ProcessAll(InteractionAsset, StepDeltaTime, PlayTime, bIsMirrored, BlendParameters, AnimContexts, Roles);
 
-	for (FInteractionAssetPreviewActor& PreviewActor : PreviewActors)
-	{
-		PreviewActor.Destroy();
+		if (const UWorld* World = GetWorld())
+		{
+			for (const FConstraint& Constraint : Constraints.GetAll())
+			{
+				DrawDebugLine(World, Constraint.FromSocketTransform.GetTranslation(), Constraint.ToSocketTransform.GetTranslation(), FColor::Red, false, -1.f, SDPG_Foreground);
+			}
+		}
 	}
-
-	PreviewActors.Reset();
 }
 
 TRange<double> FInteractionAssetViewModel::GetPreviewPlayRange() const
diff --git a/Engine/Plugins/Animation/PoseSearch/Source/Editor/Public/PoseSearchInteractionAssetEditor.h b/Engine/Plugins/Animation/PoseSearch/Source/Editor/Public/PoseSearchInteractionAssetEditor.h
--- a/Engine/Plugins/Animation/PoseSearch/Source/Editor/Public/PoseSearchInteractionAssetEditor.h
+++ b/Engine/Plugins/Animation/PoseSearch/Source/Editor/Public/PoseSearchInteractionAssetEditor.h
@@ -56,7 +56,6 @@ public:
 	virtual void AddReferencedObjects(FReferenceCollector& Collector) override;
 	virtual FString GetReferencerName() const override { return TEXT("FInteractionAssetViewModel"); }
 	void Initialize(UPoseSearchInteractionAsset* InteractionAsset, const TSharedRef<FInteractionAssetPreviewScene>& PreviewScene);
-	void RemovePreviewActors();
 	void PreviewBackwardEnd();
 	void PreviewBackwardStep();
 	void PreviewBackward();
