Add support for posesearch generation of warped montage trajectories
- Remove all direct actor transform references from root motion modifiers.
- Add 'UMotionWarpingTrajectoryAdapter' which may optionally exist on a 'UMotionWarpingComponent' to predict trajectory.
- Add `FRootMotionModifierSwapState` to store warping state pre-prediction & restore state post prediction
- Add "Prediction Mode" to 'UMotionWarpingComponent' where it uses the swap state / new adapter to predict warped trajectories
- Add Experimental `UPoseSearchTrajectoryLibrary::PoseSearchGenerateWarpedTransformTrajectory` to generate warped trajectories
- Add `GVarPoseSearchTrajectoryUseWarpedMontage` to enable warped montage trajectories by default

#tests PIE
#Synced-CL 49775968

diff --git a/Engine/Plugins/Animation/MotionWarping/Source/MotionWarping/Private/MotionWarpingComponent.cpp b/Engine/Plugins/Animation/MotionWarping/Source/MotionWarping/Private/MotionWarpingComponent.cpp
--- a/Engine/Plugins/Animation/MotionWarping/Source/MotionWarping/Private/MotionWarpingComponent.cpp
+++ b/Engine/Plugins/Animation/MotionWarping/Source/MotionWarping/Private/MotionWarpingComponent.cpp
@@ -14,6 +14,7 @@
 #include "GameFramework/CharacterMovementComponent.h"
 #include "AnimNotifyState_MotionWarping.h"
 #include "MotionWarpingCharacterAdapter.h"
+#include "MotionWarpingMontageTrajectoryAdapter.h"
 #include "MotionWarpingSwitchOffCondition.h"
 #include "Net/Core/PushModel/PushModel.h"
 #include "Net/UnrealNetwork.h"
@@ -434,7 +435,9 @@ void UMotionWarpingComponent::UpdateWithContext(const FMotionWarpingUpdateContex
 {
 	UpdateSwitchOffConditions();
 
-	if (Context.Animation.IsValid())
+	// Currently just uses existing modifiers and swaps state, no need to create
+	// Also We can predict outside of gamethread which means creating Modifiers / UObjects.
+	if (!bPredictingTrajectory && Context.Animation.IsValid())
 	{
 		const UAnimSequenceBase* Animation = Context.Animation.Get();
 		const float PreviousPosition = Context.PreviousPosition;
@@ -542,7 +545,7 @@ void UMotionWarpingComponent::UpdateWithContext(const FMotionWarpingUpdateContex
 FTransform UMotionWarpingComponent::ProcessRootMotionPreConvertToWorld(const FTransform& InRootMotion, float DeltaSeconds, const FMotionWarpingUpdateContext* InContext)
 {
 #if !(UE_BUILD_SHIPPING || UE_BUILD_TEST)
-	if (FMotionWarpingCVars::CVarMotionWarpingDisable.GetValueOnGameThread() > 0)
+	if (FMotionWarpingCVars::CVarMotionWarpingDisable.GetValueOnAnyThread() > 0)
 	{
 		return InRootMotion;
 	}
@@ -567,7 +570,7 @@ FTransform UMotionWarpingComponent::ProcessRootMotionPreConvertToWorld(const FTr
 	}
 
 #if !(UE_BUILD_SHIPPING || UE_BUILD_TEST)
-	const int32 DebugLevel = FMotionWarpingCVars::CVarMotionWarpingDebug.GetValueOnGameThread();
+	const int32 DebugLevel = FMotionWarpingCVars::CVarMotionWarpingDebug.GetValueOnAnyThread();
 	if (DebugLevel >= 2 && OwnerAdapter)
 	{
 		const float DrawDebugDuration = FMotionWarpingCVars::CVarMotionWarpingDrawDebugDuration.GetValueOnGameThread();
@@ -704,7 +707,10 @@ void UMotionWarpingComponent::AddOrUpdateWarpTarget(const FMotionWarpingTarget&
 			}
 		}
 
-		MARK_PROPERTY_DIRTY_FROM_NAME(UMotionWarpingComponent, WarpTargets, this);
+		if (!bPredictingTrajectory)
+		{
+			MARK_PROPERTY_DIRTY_FROM_NAME(UMotionWarpingComponent, WarpTargets, this);
+		}
 	}
 }
 
@@ -712,7 +718,7 @@ int32 UMotionWarpingComponent::RemoveWarpTarget(FName WarpTargetName)
 {
 	const int32 NumRemoved = WarpTargets.RemoveAll([&WarpTargetName](const FMotionWarpingTarget& WarpTarget) { return WarpTarget.Name == WarpTargetName; });
 	
-	if(NumRemoved > 0)
+	if (NumRemoved > 0 && !bPredictingTrajectory)
 	{
 		MARK_PROPERTY_DIRTY_FROM_NAME(UMotionWarpingComponent, WarpTargets, this);
 	}
@@ -730,7 +736,7 @@ int32 UMotionWarpingComponent::RemoveWarpTargets(const TArray<FName>& WarpTarget
 		NumRemoved += RemoveWarpTarget(WarpTargetName);		
 	}
 
-	if (NumRemoved > 0)
+	if (NumRemoved > 0 && !bPredictingTrajectory)
 	{
 		MARK_PROPERTY_DIRTY_FROM_NAME(UMotionWarpingComponent, WarpTargets, this);
 	}
@@ -783,6 +789,54 @@ FSwitchOffConditionData* UMotionWarpingComponent::FindSwitchOffConditionData(FNa
 	});
 }
 
+void UMotionWarpingComponent::BeginTrajectoryPredictions(const FAnimMontageInstance* MontageInstance)
+{
+	check(!bPredictingTrajectory);
+
+	// @TODO: We don't support SwitchOffConditions yet, TBD if we will
+	ensure(SwitchOffConditions.IsEmpty());
+
+	bPredictingTrajectory = true;
+
+	CachedOwnerAdapter = OwnerAdapter;
+
+	UMotionWarpingMontageTrajectoryAdapter* MontageTrajectoryAdapter = nullptr;
+	if (!PredictionAdapter)
+	{
+		// Create prediction adapter and assign as owner adapter
+		MontageTrajectoryAdapter = CreateOwnerAdapter<UMotionWarpingMontageTrajectoryAdapter>();
+		PredictionAdapter = MontageTrajectoryAdapter;
+	}
+	else
+	{
+		OwnerAdapter = PredictionAdapter;
+		MontageTrajectoryAdapter = CastChecked<UMotionWarpingMontageTrajectoryAdapter>(OwnerAdapter);
+	}
+
+	MontageTrajectoryAdapter->InitializeFrom(CachedOwnerAdapter, MontageInstance);
+}
+
+void UMotionWarpingComponent::EndTrajectoryPredictions()
+{
+	check(bPredictingTrajectory);
+
+	// Restore Pre-prediction owner
+	OwnerAdapter = CachedOwnerAdapter;
+
+	// Restore Warp State to pre-prediction
+	if (UMotionWarpingMontageTrajectoryAdapter* MontageTrajectoryAdapter = Cast<UMotionWarpingMontageTrajectoryAdapter>(PredictionAdapter))
+	{
+		MontageTrajectoryAdapter->RestoreState();
+	}
+
+	bPredictingTrajectory = false;
+}
+
+bool UMotionWarpingComponent::IsPredictingTrajectory() const
+{
+	return bPredictingTrajectory;
+}
+
 int32 UMotionWarpingComponent::RemoveAllWarpTargets()
 {
 	const int32 NumRemoved = WarpTargets.Num();
@@ -794,7 +848,7 @@ int32 UMotionWarpingComponent::RemoveAllWarpTargets()
 	
 	WarpTargets.Reset();
 
-	if (NumRemoved > 0)
+	if (NumRemoved > 0 && !bPredictingTrajectory)
 	{
 		MARK_PROPERTY_DIRTY_FROM_NAME(UMotionWarpingComponent, WarpTargets, this);
 	}
diff --git a/Engine/Plugins/Animation/MotionWarping/Source/MotionWarping/Private/MotionWarpingMontageTrajectoryAdapter.cpp b/Engine/Plugins/Animation/MotionWarping/Source/MotionWarping/Private/MotionWarpingMontageTrajectoryAdapter.cpp
new file mode 100644
--- /dev/null
+++ b/Engine/Plugins/Animation/MotionWarping/Source/MotionWarping/Private/MotionWarpingMontageTrajectoryAdapter.cpp
@@ -0,0 +1,118 @@
+// Copyright Epic Games, Inc. All Rights Reserved.
+
+#include "MotionWarpingMontageTrajectoryAdapter.h"
+
+#include "Animation/AnimMontage.h"
+#include "Components/CapsuleComponent.h"
+#include "MotionWarpingComponent.h"
+
+#include UE_INLINE_GENERATED_CPP_BY_NAME(MotionWarpingMontageTrajectoryAdapter)
+
+void UMotionWarpingMontageTrajectoryAdapter::InitializeFrom(const UMotionWarpingBaseAdapter* InAdapter, const FAnimMontageInstance* InMontageInstance)
+{
+	ensure(InAdapter);
+	ensure(InMontageInstance);
+
+	if (!InAdapter || !InMontageInstance)
+	{
+		return;
+	}
+
+	CurrentTransform = InAdapter->GetCurrentTransform();
+	VisualRootLocation = InAdapter->GetVisualRootLocation();
+	BaseVisualTranslationOffset = InAdapter->GetBaseVisualTranslationOffset();
+	BaseVisualRotationOffset = InAdapter->GetBaseVisualRotationOffset();
+
+	// @TODO: Weight can vary, technically we need to predict that. But so far nothing uses this but adjustment blend warp
+	Montage = InMontageInstance->Montage;
+	MontageWeight = InMontageInstance->GetWeight();
+	MontagePlayRate = InMontageInstance->Montage->RateScale * InMontageInstance->GetPlayRate();
+
+	// Cache Warp State & Create a copy of warp state for the warping component to predict with
+	{
+		UMotionWarpingComponent* WarpingComponent = CastChecked<UMotionWarpingComponent>(GetOuter());
+
+		// We don't support Switch Off Conditions yet (The class itself needs to be updated to be actorless & predict both src / target
+		ensure(WarpingComponent->SwitchOffConditions.IsEmpty());
+
+		CachedWarpState.Modifiers = WarpingComponent->Modifiers;
+		CachedWarpState.WarpTargets = WarpingComponent->WarpTargets;
+		CachedWarpState.SwitchOffConditions = WarpingComponent->SwitchOffConditions;
+
+		CachedWarpState.ModifierSwapStates.Reset();
+		for (URootMotionModifier* Modifer : WarpingComponent->Modifiers)
+		{
+			// Can't duplicate object on worker thread, cache state instead
+			CachedWarpState.ModifierSwapStates.Add(Modifer->GetSwapState());
+		}
+	}
+}
+
+void UMotionWarpingMontageTrajectoryAdapter::RestoreState()
+{
+	UMotionWarpingComponent* WarpingComponent = CastChecked<UMotionWarpingComponent>(GetOuter());
+
+	WarpingComponent->Modifiers = MoveTemp(CachedWarpState.Modifiers);
+	WarpingComponent->WarpTargets = MoveTemp(CachedWarpState.WarpTargets);
+	WarpingComponent->SwitchOffConditions = MoveTemp(CachedWarpState.SwitchOffConditions);
+
+	ensure(WarpingComponent->Modifiers.Num() == CachedWarpState.ModifierSwapStates.Num());
+	for (int32 Index = 0; Index < CachedWarpState.ModifierSwapStates.Num(); Index++)
+	{
+		WarpingComponent->Modifiers[Index]->RestoreFromSwapState(CachedWarpState.ModifierSwapStates[Index]);
+	}
+
+	CachedWarpState.Modifiers.Reset();
+	CachedWarpState.WarpTargets.Reset();
+	CachedWarpState.SwitchOffConditions.Reset();
+	CachedWarpState.ModifierSwapStates.Reset();
+}
+
+FTransform UMotionWarpingMontageTrajectoryAdapter::GetCurrentTransform() const
+{
+	return CurrentTransform;
+}
+
+FVector UMotionWarpingMontageTrajectoryAdapter::GetVisualRootLocation() const
+{
+	return VisualRootLocation;
+}
+
+void UMotionWarpingMontageTrajectoryAdapter::SetCurrentTransform(const FTransform InCurrentTransform)
+{
+	CurrentTransform = InCurrentTransform;
+}
+
+void UMotionWarpingMontageTrajectoryAdapter::SetVisualRootLocation(const FVector InVisualRootLocation)
+{
+	VisualRootLocation = InVisualRootLocation;
+}
+
+FVector UMotionWarpingMontageTrajectoryAdapter::GetBaseVisualTranslationOffset() const
+{
+	return BaseVisualTranslationOffset;
+}
+
+FQuat UMotionWarpingMontageTrajectoryAdapter::GetBaseVisualRotationOffset() const
+{ 
+	return BaseVisualRotationOffset;
+}
+
+FTransform UMotionWarpingMontageTrajectoryAdapter::WarpLocalRootMotionOnMontageTrajectory(const FTransform& LocalRootMotionTransform, float DeltaSeconds, float InMontagePreviousPosition, float InMontageCurrentPosition)
+{
+	if (WarpLocalRootMotionDelegate.IsBound() && Montage)
+	{
+		FMotionWarpingUpdateContext WarpingContext;
+		
+		WarpingContext.DeltaSeconds = DeltaSeconds;
+		WarpingContext.Animation = Montage;
+		WarpingContext.CurrentPosition = InMontageCurrentPosition;
+		WarpingContext.PreviousPosition = InMontagePreviousPosition;
+		WarpingContext.Weight = MontageWeight;
+		WarpingContext.PlayRate = MontagePlayRate;
+
+		return WarpLocalRootMotionDelegate.Execute(LocalRootMotionTransform, DeltaSeconds, &WarpingContext);
+	}
+
+	return LocalRootMotionTransform;
+}
diff --git a/Engine/Plugins/Animation/MotionWarping/Source/MotionWarping/Private/RootMotionModifier.cpp b/Engine/Plugins/Animation/MotionWarping/Source/MotionWarping/Private/RootMotionModifier.cpp
--- a/Engine/Plugins/Animation/MotionWarping/Source/MotionWarping/Private/RootMotionModifier.cpp
+++ b/Engine/Plugins/Animation/MotionWarping/Source/MotionWarping/Private/RootMotionModifier.cpp
@@ -238,15 +238,45 @@ ACharacter* URootMotionModifier::GetCharacterOwner() const
 	return Cast<ACharacter>(GetActorOwner());
 }
 
+TSharedPtr<FRootMotionModifierSwapState> URootMotionModifier::GetSwapState() const
+{
+	TSharedPtr<FRootMotionModifierSwapState> Result = MakeShared<FRootMotionModifierSwapState>();
+	PopulateSwapState(Result);
+	return Result;
+}
+
+void URootMotionModifier::PopulateSwapState(TSharedPtr<FRootMotionModifierSwapState> InSwapState) const
+{
+	InSwapState->StartTime = StartTime;
+	InSwapState->EndTime = EndTime;
+	InSwapState->PreviousPosition = PreviousPosition;
+	InSwapState->CurrentPosition = CurrentPosition;
+	InSwapState->Weight = Weight;
+	InSwapState->PlayRate = PlayRate;
+	InSwapState->StartTransform = StartTransform;
+	InSwapState->ActualStartTime = ActualStartTime;
+	InSwapState->TotalRootMotionWithinWindow = TotalRootMotionWithinWindow;
+	InSwapState->State = State;
+}
+
+void URootMotionModifier::RestoreFromSwapState(TSharedPtr<FRootMotionModifierSwapState> InSwapState)
+{
+	StartTime = InSwapState->StartTime;
+	EndTime = InSwapState->EndTime;
+	PreviousPosition = InSwapState->PreviousPosition;
+	CurrentPosition = InSwapState->CurrentPosition;
+	Weight = InSwapState->Weight;
+	PlayRate = InSwapState->PlayRate;
+	StartTransform = InSwapState->StartTransform;
+	ActualStartTime = InSwapState->ActualStartTime;
+	TotalRootMotionWithinWindow = InSwapState->TotalRootMotionWithinWindow;
+	State = InSwapState->State;
+}
+
 void URootMotionModifier::Update(const FMotionWarpingUpdateContext& Context)
 {
 	const AActor* ActorOwner = GetActorOwner();
 
-	if (ActorOwner == nullptr)
-	{
-		return;
-	}
-
 	// Mark for removal if our animation is not relevant anymore
 	if (!Context.Animation.IsValid() || Context.Animation.Get() != Animation)
 	{
@@ -330,7 +360,7 @@ void URootMotionModifier::OnStateChanged(ERootMotionModifierState LastState)
 			checkf(OwnerAdapter, TEXT("Root motion modifiers expect an owner and adapter"));
 
 			const FVector CurrentLocation = OwnerAdapter->GetVisualRootLocation();
-			const FQuat CurrentRotation = OwnerAdapter->GetActor()->GetActorQuat();
+			const FQuat CurrentRotation = OwnerAdapter->GetCurrentRotation();
 			
 			ActualStartTime = PreviousPosition;
 
@@ -425,7 +455,7 @@ void URootMotionModifier_Warp::OnTargetTransformChanged()
 	{
 		ActualStartTime = PreviousPosition;
 
-		const FQuat CurrentRotation = WarpingAdapter->GetActor()->GetActorQuat();
+		const FQuat CurrentRotation = WarpingAdapter->GetCurrentRotation();
 		const FVector CurrentLocation = WarpingAdapter->GetVisualRootLocation();
 		StartTransform = FTransform(CurrentRotation, CurrentLocation);
 	}
@@ -441,6 +471,44 @@ void URootMotionModifier_Warp::OnStateChanged(ERootMotionModifierState LastState
 	}
 }
 
+
+TSharedPtr<FRootMotionModifierSwapState> URootMotionModifier_Warp::GetSwapState() const
+{
+	TSharedPtr<FRootMotionModifierSwapState> Result = MakeShared<FRootMotionModifier_WarpSwapState>();
+	PopulateSwapState(Result);
+	return Result;
+}
+
+void URootMotionModifier_Warp::PopulateSwapState(TSharedPtr<FRootMotionModifierSwapState> InSwapState) const
+{
+	Super::PopulateSwapState(InSwapState);
+
+	TSharedPtr<FRootMotionModifier_WarpSwapState> WarpSwapState = StaticCastSharedPtr<FRootMotionModifier_WarpSwapState>(InSwapState);
+
+	WarpSwapState->AdditionalRotationOffset = AdditionalRotationOffset;
+	WarpSwapState->CachedTargetTransform = CachedTargetTransform;
+	WarpSwapState->RootMotionRemainingAfterNotify = RootMotionRemainingAfterNotify;
+	WarpSwapState->CachedOffsetFromWarpPoint = CachedOffsetFromWarpPoint;
+	WarpSwapState->RotationOffset = RotationOffset;
+	WarpSwapState->bWarpingPaused = bWarpingPaused;
+	WarpSwapState->bRootMotionPaused = bRootMotionPaused;
+}
+
+void URootMotionModifier_Warp::RestoreFromSwapState(TSharedPtr<FRootMotionModifierSwapState> InSwapState)
+{
+	Super::RestoreFromSwapState(InSwapState);
+
+	TSharedPtr<FRootMotionModifier_WarpSwapState> WarpSwapState = StaticCastSharedPtr<FRootMotionModifier_WarpSwapState>(InSwapState);
+
+	AdditionalRotationOffset = WarpSwapState->AdditionalRotationOffset;
+	CachedTargetTransform = WarpSwapState->CachedTargetTransform;
+	RootMotionRemainingAfterNotify = WarpSwapState->RootMotionRemainingAfterNotify;
+	CachedOffsetFromWarpPoint = WarpSwapState->CachedOffsetFromWarpPoint;
+	RotationOffset = WarpSwapState->RotationOffset;
+	bWarpingPaused = WarpSwapState->bWarpingPaused;
+	bRootMotionPaused = WarpSwapState->bRootMotionPaused;
+}
+
 FQuat URootMotionModifier_Warp::GetTargetRotation() const
 {
 	switch (RotationType)
@@ -481,7 +549,7 @@ FQuat URootMotionModifier_Warp::WarpRotation(const FTransform& RootMotionDelta,
 
 	if (const UMotionWarpingBaseAdapter* WarpingAdapter = GetOwnerAdapter())
 	{
-		CurrentRotation = WarpingAdapter->GetActor()->GetActorQuat() * WarpingAdapter->GetBaseVisualRotationOffset();
+		CurrentRotation = WarpingAdapter->GetCurrentRotation() * WarpingAdapter->GetBaseVisualRotationOffset();
 		TargetRotation = CurrentRotation.Inverse() * (GetTargetRotation() * WarpingAdapter->GetBaseVisualRotationOffset() * RootMotionRemainingAfterNotify.GetRotation().Inverse());
 	}
 	else
@@ -532,7 +600,8 @@ FQuat URootMotionModifier_Warp::WarpRotation(const FTransform& RootMotionDelta,
 		}
 	}
 
-	const FQuat DeltaOut = TargetRotThisFrame * TotalRootMotionRotation.Inverse();
+	FQuat DeltaOut = TargetRotThisFrame * TotalRootMotionRotation.Inverse();
+	DeltaOut = DeltaOut.Equals(FQuat::Identity) ? FQuat::Identity : DeltaOut;
 	
 	return (DeltaOut * RootMotionDelta.GetRotation());
 }
@@ -697,4 +766,4 @@ FTransform URootMotionModifier_Blueprintable::ProcessRootMotion(const FTransform
 	}
 
 	return InRootMotion;
-}
\ No newline at end of file
+}
diff --git a/Engine/Plugins/Animation/MotionWarping/Source/MotionWarping/Private/RootMotionModifier_PrecomputedWarp.cpp b/Engine/Plugins/Animation/MotionWarping/Source/MotionWarping/Private/RootMotionModifier_PrecomputedWarp.cpp
--- a/Engine/Plugins/Animation/MotionWarping/Source/MotionWarping/Private/RootMotionModifier_PrecomputedWarp.cpp
+++ b/Engine/Plugins/Animation/MotionWarping/Source/MotionWarping/Private/RootMotionModifier_PrecomputedWarp.cpp
@@ -406,4 +406,47 @@ FTransform URootMotionModifier_PrecomputedWarp::ProcessRootMotion(const FTransfo
 	
 	
 	return InRootMotion;
+}
+
+TSharedPtr<FRootMotionModifierSwapState> URootMotionModifier_PrecomputedWarp::GetSwapState() const
+{
+	TSharedPtr<FRootMotionModifierSwapState> Result = MakeShared<FRootMotionModifier_PrecomputedWarpSwapState>();
+	PopulateSwapState(Result);
+	return Result;
+}
+
+void URootMotionModifier_PrecomputedWarp::PopulateSwapState(TSharedPtr<FRootMotionModifierSwapState> InSwapState) const
+{
+	Super::PopulateSwapState(InSwapState);
+
+	TSharedPtr<FRootMotionModifier_PrecomputedWarpSwapState> PrecomputedWarpSwapState = StaticCastSharedPtr<FRootMotionModifier_PrecomputedWarpSwapState>(InSwapState);
+
+	PrecomputedWarpSwapState->bFirstFrame = bFirstFrame;
+	PrecomputedWarpSwapState->ActualStartTime = ActualStartTime;
+	PrecomputedWarpSwapState->RoundedEndTime = RoundedEndTime;
+	PrecomputedWarpSwapState->PreviousFrame = PreviousFrame;
+	PrecomputedWarpSwapState->StartingRootTransform = StartingRootTransform;
+	PrecomputedWarpSwapState->TargetTransform = TargetTransform;
+	PrecomputedWarpSwapState->FilteredSteeringTarget = FilteredSteeringTarget;
+	PrecomputedWarpSwapState->TargetSmoothingState = TargetSmoothingState;
+	PrecomputedWarpSwapState->WarpedTrajectory = WarpedTrajectory;
+	PrecomputedWarpSwapState->AnimTrajectoryData = AnimTrajectoryData;
+}
+
+void URootMotionModifier_PrecomputedWarp::RestoreFromSwapState(TSharedPtr<FRootMotionModifierSwapState> InSwapState)
+{
+	Super::RestoreFromSwapState(InSwapState);
+
+	TSharedPtr<FRootMotionModifier_PrecomputedWarpSwapState> PrecomputedWarpSwapState = StaticCastSharedPtr<FRootMotionModifier_PrecomputedWarpSwapState>(InSwapState);
+
+	bFirstFrame = PrecomputedWarpSwapState->bFirstFrame;
+	ActualStartTime = PrecomputedWarpSwapState->ActualStartTime;
+	RoundedEndTime = PrecomputedWarpSwapState->RoundedEndTime;
+	PreviousFrame = PrecomputedWarpSwapState->PreviousFrame;
+	StartingRootTransform = PrecomputedWarpSwapState->StartingRootTransform;
+	TargetTransform = PrecomputedWarpSwapState->TargetTransform;
+	FilteredSteeringTarget = PrecomputedWarpSwapState->FilteredSteeringTarget;
+	TargetSmoothingState = PrecomputedWarpSwapState->TargetSmoothingState;
+	WarpedTrajectory = PrecomputedWarpSwapState->WarpedTrajectory;
+	AnimTrajectoryData = PrecomputedWarpSwapState->AnimTrajectoryData;
 }
\ No newline at end of file
diff --git a/Engine/Plugins/Animation/MotionWarping/Source/MotionWarping/Private/RootMotionModifier_SkewWarp.cpp b/Engine/Plugins/Animation/MotionWarping/Source/MotionWarping/Private/RootMotionModifier_SkewWarp.cpp
--- a/Engine/Plugins/Animation/MotionWarping/Source/MotionWarping/Private/RootMotionModifier_SkewWarp.cpp
+++ b/Engine/Plugins/Animation/MotionWarping/Source/MotionWarping/Private/RootMotionModifier_SkewWarp.cpp
@@ -6,6 +6,7 @@
 #include "Components/CapsuleComponent.h"
 #include "MotionWarpingComponent.h"
 #include "MotionWarpingAdapter.h"
+#include "MotionWarpingMontageTrajectoryAdapter.h"
 #include "SceneView.h"
 
 #if WITH_EDITOR
@@ -136,14 +137,8 @@ FVector URootMotionModifier_SkewWarp::WarpTranslation(const FTransform& CurrentT
 FTransform URootMotionModifier_SkewWarp::ProcessRootMotion(const FTransform& InRootMotion, float DeltaSeconds)
 {
 	const UMotionWarpingBaseAdapter* OwnerAdapter = GetOwnerAdapter();
-	const AActor* OwnerAsActor = nullptr;
 
-	if (OwnerAdapter)
-	{
-		OwnerAsActor = OwnerAdapter->GetActor();
-	}
-
-	if (!OwnerAdapter || !OwnerAsActor)
+	if (!OwnerAdapter)
 	{
 		return InRootMotion;
 	}
@@ -168,7 +163,7 @@ FTransform URootMotionModifier_SkewWarp::ProcessRootMotion(const FTransform& InR
 		else if (!bWarpingPaused)
 		{
 		const FVector CurrentLocation       = OwnerAdapter->GetVisualRootLocation();
-		const FQuat CurrentRotation         = OwnerAdapter->GetActor()->GetActorQuat();
+		const FQuat CurrentRotation         = OwnerAdapter->GetCurrentRotation();
 		FVector MeshTranslationOffset       = OwnerAdapter->GetBaseVisualTranslationOffset();
 		if (!bWarpToFeetLocation)
 		{
@@ -191,7 +186,7 @@ FTransform URootMotionModifier_SkewWarp::ProcessRootMotion(const FTransform& InR
 		{
 			if (!DeltaTranslation.IsNearlyZero())
 			{
-				const FTransform MeshTransform = FTransform(MeshRotationOffset, MeshTranslationOffset) * OwnerAsActor->GetActorTransform();
+				const FTransform MeshTransform = FTransform(MeshRotationOffset, MeshTranslationOffset) * OwnerAdapter->GetCurrentTransform();
 				TargetLocation = MeshTransform.InverseTransformPositionNoScale(TargetLocation) - RootMotionRemainingAfterNotify.GetTranslation();
 
 				FVector WarpedTranslation = WarpTranslation(FTransform::Identity, DeltaTranslation, TotalTranslation, TargetLocation) + ExtraRootMotion.GetLocation();
@@ -248,7 +243,7 @@ FTransform URootMotionModifier_SkewWarp::ProcessRootMotion(const FTransform& InR
 
 	// Debug
 #if !(UE_BUILD_SHIPPING || UE_BUILD_TEST)
-	const int32 DebugLevel = FMotionWarpingCVars::CVarMotionWarpingDebug.GetValueOnGameThread();
+	const int32 DebugLevel = FMotionWarpingCVars::CVarMotionWarpingDebug.GetValueOnAnyThread();
 	if (DebugLevel == 1 || DebugLevel == 3)
 	{
 		PrintLog(TEXT("SkewWarp"), InRootMotion, FinalRootMotion);
@@ -257,7 +252,11 @@ FTransform URootMotionModifier_SkewWarp::ProcessRootMotion(const FTransform& InR
 	if (DebugLevel == 2 || DebugLevel == 3)
 	{
 		const float DrawDebugDuration = FMotionWarpingCVars::CVarMotionWarpingDrawDebugDuration.GetValueOnGameThread();
-		DrawDebugCoordinateSystem(OwnerAsActor->GetWorld(), GetTargetLocation(), GetTargetRotator(), 50.f, false, DrawDebugDuration, 0, 1.f);
+
+		if (const AActor* OwnerAsActor = OwnerAdapter->GetActor())
+		{
+			DrawDebugCoordinateSystem(OwnerAsActor->GetWorld(), GetTargetLocation(), GetTargetRotator(), 50.f, false, DrawDebugDuration, 0, 1.f);
+		}
 	}
 #endif
 
diff --git a/Engine/Plugins/Animation/MotionWarping/Source/MotionWarping/Public/MotionWarpingAdapter.h b/Engine/Plugins/Animation/MotionWarping/Source/MotionWarping/Public/MotionWarpingAdapter.h
--- a/Engine/Plugins/Animation/MotionWarping/Source/MotionWarping/Public/MotionWarpingAdapter.h
+++ b/Engine/Plugins/Animation/MotionWarping/Source/MotionWarping/Public/MotionWarpingAdapter.h
@@ -2,10 +2,10 @@
 
 #pragma once
 
+#include "GameFramework/Actor.h"
 #include "RootMotionModifier.h"
 #include "MotionWarpingAdapter.generated.h"
 
-
 DECLARE_DELEGATE_RetVal_ThreeParams(FTransform, FOnWarpLocalspaceRootMotionWithContext, const FTransform&, float, const FMotionWarpingUpdateContext*)
 DECLARE_DELEGATE_RetVal_ThreeParams(FTransform, FOnWarpWorldspaceRootMotionWithContext, const FTransform&, float, const FMotionWarpingUpdateContext*)
 
@@ -19,9 +19,33 @@ class UMotionWarpingBaseAdapter : public UObject
 
 public:
 	virtual ~UMotionWarpingBaseAdapter() = default;
+
 	virtual AActor* GetActor() const { return nullptr; }
 	virtual USkeletalMeshComponent* GetMesh() const { return nullptr; }
-	virtual FVector GetVisualRootLocation() const { return FVector::ZeroVector; }
+
+	// Current Transform. Warping should use this as child adapters may not have an actor.
+	virtual FTransform GetCurrentTransform() const 
+	{ 
+		if (AActor* OwningActor = GetActor())
+		{
+			return OwningActor->GetActorTransform();
+		}
+
+		return FTransform::Identity; 
+	}
+
+	// Current visual root (IE: Bottom of actor capsule).
+	virtual FVector GetVisualRootLocation() const 
+	{ 
+		return FVector::ZeroVector; 
+	}
+
+	// Current Transform Rotation. Utility method to make API consistent.
+	FQuat GetCurrentRotation() const 
+	{
+		return GetCurrentTransform().GetRotation();
+	}
+
 	virtual FVector GetBaseVisualTranslationOffset() const { return FVector::ZeroVector; }
 	virtual FQuat GetBaseVisualRotationOffset() const { return FQuat::Identity; }
 
diff --git a/Engine/Plugins/Animation/MotionWarping/Source/MotionWarping/Public/MotionWarpingComponent.h b/Engine/Plugins/Animation/MotionWarping/Source/MotionWarping/Public/MotionWarpingComponent.h
--- a/Engine/Plugins/Animation/MotionWarping/Source/MotionWarping/Public/MotionWarpingComponent.h
+++ b/Engine/Plugins/Animation/MotionWarping/Source/MotionWarping/Public/MotionWarpingComponent.h
@@ -19,6 +19,7 @@ class UCharacterMovementComponent;
 class UMotionWarpingBaseAdapter;
 class UMotionWarpingComponent;
 class UAnimNotifyState_MotionWarping;
+struct FAnimMontageInstance;
 struct FCompactPose;
 template<class PoseType> struct FCSPose;
 
@@ -100,6 +101,10 @@ class UMotionWarpingComponent : public UActorComponent
 {
 	GENERATED_BODY()
 
+public:
+
+	friend class UMotionWarpingMontageTrajectoryAdapter;
+
 public:
 
 	/** Whether to look inside animations within montage when looking for warping windows */
@@ -223,22 +228,46 @@ public:
 	
 	UE_API FSwitchOffConditionData* FindSwitchOffConditionData(FName WarpTargetName);
 
+	/** Creates Prediction adapter if needed and swaps state to prediction state */
+	UE_API void BeginTrajectoryPredictions(const FAnimMontageInstance* MontageInstance);
+
+	/** Reverts state back to before any predictions */
+	UE_API void EndTrajectoryPredictions();
+
+	/** True if we are currentl prediciting a trajectory */
+	UE_API bool IsPredictingTrajectory() const;
+
 protected:
 
 	/** Adapter that connects motion warping to an owner */
-	UPROPERTY(Transient)
+	UPROPERTY(Transient, VisibleInstanceOnly, Category = "Warp State")
 	TObjectPtr<UMotionWarpingBaseAdapter> OwnerAdapter;
+
+	/** Experimental. Adapter used to predict trajectories. */
+	UE_EXPERIMENTAL(5.8, "Prediction Warping API subject to change.")
+	UPROPERTY(Transient, VisibleInstanceOnly, Category = "Warp State")
+	TObjectPtr<UMotionWarpingBaseAdapter> PredictionAdapter;
+
+	/** Experimental. Cache of owner adapter before swapped out to do prediction, swapped back post prediction */
+	UE_EXPERIMENTAL(5.8, "Prediction Warping API subject to change.")
+	UPROPERTY(Transient, VisibleInstanceOnly, Category = "Warp State")
+	TObjectPtr<UMotionWarpingBaseAdapter> CachedOwnerAdapter;
 	
 	/** List of root motion modifiers */
-	UPROPERTY(Transient)
+	UPROPERTY(Transient, VisibleInstanceOnly, Category = "Warp State")
 	TArray<TObjectPtr<URootMotionModifier>> Modifiers;
 
-	UPROPERTY(Transient, Replicated)
+	UPROPERTY(Transient, Replicated, VisibleInstanceOnly, Category = "Warp State")
 	TArray<FMotionWarpingTarget> WarpTargets;
 		
-	UPROPERTY(Transient, Experimental)
+	UPROPERTY(Transient, Experimental, VisibleInstanceOnly, Category = "Warp State")
     TArray<FSwitchOffConditionData> SwitchOffConditions;
 
+	/** True if we are predicting trajectory via a temp swap state. */
+	UE_EXPERIMENTAL(5.8, "Prediction Warping API subject to change.")
+	UPROPERTY(Transient, Experimental, VisibleInstanceOnly, Category = "Warp State")
+	bool bPredictingTrajectory = false;
+
 #if !(UE_BUILD_SHIPPING || UE_BUILD_TEST)
 	TOptional<FVector> OriginalRootMotionAccum;
 	TOptional<FVector> WarpedRootMotionAccum;
diff --git a/Engine/Plugins/Animation/MotionWarping/Source/MotionWarping/Public/MotionWarpingMontageTrajectoryAdapter.h b/Engine/Plugins/Animation/MotionWarping/Source/MotionWarping/Public/MotionWarpingMontageTrajectoryAdapter.h
new file mode 100644
--- /dev/null
+++ b/Engine/Plugins/Animation/MotionWarping/Source/MotionWarping/Public/MotionWarpingMontageTrajectoryAdapter.h
@@ -0,0 +1,91 @@
+// Copyright Epic Games, Inc. All Rights Reserved.
+
+#pragma once
+
+#include "MotionWarpingAdapter.h"
+#include "MotionWarpingMontageTrajectoryAdapter.generated.h"
+
+#define UE_API MOTIONWARPING_API
+
+class UMotionWarpingSwitchOffCondition;
+struct FAnimMontageInstance;
+struct FSwitchOffConditionData;
+
+/** Temp state to iterate on for properties normally owned by the warping component. Used so we can predict without mutating real actor state. */
+USTRUCT(Experimental)
+struct FMotionWarpingSwapState
+{
+	GENERATED_BODY()
+
+	/** List of root motion modifiers, stored as modifiers can be removed during prediction */
+	UPROPERTY(Transient, VisibleInstanceOnly, Category = "Warp State")
+	TArray<TObjectPtr<URootMotionModifier>> Modifiers;
+
+	/** List of warp targets */
+	UPROPERTY(Transient, VisibleInstanceOnly, Category = "Warp State")
+	TArray<FMotionWarpingTarget> WarpTargets;
+
+	/** 
+	 * List of switch off conditions
+	 * 
+	 * @TODO: Currently unused. Switch off conditions aren't actorless yet & to be accurate we may want to also predict against target actor trajectories
+	 */
+	UPROPERTY(Transient, Experimental, VisibleInstanceOnly, Category = "Warp State")
+	TArray<FSwitchOffConditionData> SwitchOffConditions;
+
+	/** State of modifiers prior to prediction */
+	TArray<TSharedPtr<FRootMotionModifierSwapState>> ModifierSwapStates;
+};
+
+/** Adapter for montage trajectories to participate in motion warping */
+UCLASS(Experimental, MinimalAPI)
+class UMotionWarpingMontageTrajectoryAdapter : public UMotionWarpingBaseAdapter
+{
+	GENERATED_BODY()
+
+public:
+
+	/** Initialize current transforms from some other adapter & swaps component state to prediction state. */
+	UE_API void InitializeFrom(const UMotionWarpingBaseAdapter* InAdapter, const FAnimMontageInstance* InMontageInstance);
+
+	/** Restores component state to before any prediction has been done */
+	UE_API void RestoreState();
+
+	UE_API virtual FTransform GetCurrentTransform() const override;
+	UE_API virtual FVector GetVisualRootLocation() const override;
+
+	UE_API void SetCurrentTransform(const FTransform InCurrentTransform);
+	UE_API void SetVisualRootLocation(const FVector InVisualRootLocation);
+
+	UE_API virtual FVector GetBaseVisualTranslationOffset() const override;
+	UE_API virtual FQuat GetBaseVisualRotationOffset() const override;
+
+	/** Triggered iteratively as we sample the trajectory */
+	UE_API FTransform WarpLocalRootMotionOnMontageTrajectory(const FTransform& LocalRootMotionTransform, float DeltaSeconds, float InMontagePreviousPosition, float InMontageCurrentPosition);
+
+public:
+
+	/** Cached state before any prediction, we restore this after predicting. */
+	UPROPERTY(Transient)
+	FMotionWarpingSwapState CachedWarpState;
+
+	/** Montage to warp trajectory for */
+	UPROPERTY(Transient)
+	TObjectPtr<const UAnimMontage> Montage;
+
+private:
+
+	/** Warping properties captured in adapter initialization that update as we iterate*/
+	FTransform CurrentTransform = FTransform::Identity;
+	FVector VisualRootLocation = FVector::ZeroVector;
+
+	/** Warping properties that do not change after adapter initialization */ 
+	FVector BaseVisualTranslationOffset = FVector::ZeroVector;
+	FQuat BaseVisualRotationOffset = FQuat::Identity;
+
+	/** Properties for warp context state */
+	float MontageWeight = 0.0f;
+	float MontagePlayRate = 0.0f;
+};
+
+#undef UE_API
diff --git a/Engine/Plugins/Animation/MotionWarping/Source/MotionWarping/Public/RootMotionModifier.h b/Engine/Plugins/Animation/MotionWarping/Source/MotionWarping/Public/RootMotionModifier.h
--- a/Engine/Plugins/Animation/MotionWarping/Source/MotionWarping/Public/RootMotionModifier.h
+++ b/Engine/Plugins/Animation/MotionWarping/Source/MotionWarping/Public/RootMotionModifier.h
@@ -15,6 +15,8 @@ class USceneComponent;
 class ACharacter;
 class UMotionWarpingBaseAdapter;
 
+struct FRootMotionModifierSwapState;
+
 /** 
  * Context passed to any active root motion modifier during the update phase. 
  * Contains relevant data from the animation that contributed to root motion this frame (or in the past when replaying saved moves)
@@ -148,6 +150,15 @@ public:
 	UE_DEPRECATED(5.5, "Motion Warping is no longer limited to Character actors. Use GetActorOwner instead.")
 	UE_API ACharacter* GetCharacterOwner() const;
 
+	/** Returns a snapshot of the current state of this modifier */
+	UE_API virtual TSharedPtr<FRootMotionModifierSwapState> GetSwapState() const;
+
+	/** Restores this modifer to some previous state snapshot */
+	UE_API virtual void RestoreFromSwapState(TSharedPtr<FRootMotionModifierSwapState> InSwapState);
+
+	/** Populates given swap state */
+	UE_API virtual void PopulateSwapState(TSharedPtr<FRootMotionModifierSwapState> InSwapState) const;
+
 	UE_API virtual void Update(const FMotionWarpingUpdateContext& Context);
 	virtual FTransform ProcessRootMotion(const FTransform& InRootMotion, float DeltaSeconds) { return FTransform::Identity; }
 
@@ -167,6 +178,23 @@ private:
 	ERootMotionModifierState State = ERootMotionModifierState::Waiting;
 };
 
+/** Stateful properties for URootMotionModifier */
+struct FRootMotionModifierSwapState
+{
+	virtual ~FRootMotionModifierSwapState() = default;
+
+	float StartTime = 0.f;
+	float EndTime = 0.f;
+	float PreviousPosition = 0.f;
+	float CurrentPosition = 0.f;
+	float Weight = 0.f;
+	float PlayRate = 1.f;
+	FTransform StartTransform;
+	float ActualStartTime = 0.f;
+	FTransform TotalRootMotionWithinWindow;
+	ERootMotionModifierState State = ERootMotionModifierState::Waiting;
+};
+
 UENUM(BlueprintType, meta=(Experimental))
 enum class EWarpTargetLocationOffsetDirection : uint8
 {
@@ -404,6 +432,10 @@ public:
 
 	UE_API virtual void OnStateChanged(ERootMotionModifierState LastState) override;
 
+	UE_API virtual TSharedPtr<FRootMotionModifierSwapState> GetSwapState() const override;
+	UE_API virtual void PopulateSwapState(TSharedPtr<FRootMotionModifierSwapState> InSwapState) const override;
+	UE_API virtual void RestoreFromSwapState(TSharedPtr<FRootMotionModifierSwapState> InSwapState) override;
+
 #if !(UE_BUILD_SHIPPING || UE_BUILD_TEST)
 	UE_API void PrintLog(const FString& Name, const FTransform& OriginalRootMotion, const FTransform& WarpedRootMotion) const;
 #endif
@@ -432,6 +464,17 @@ protected:
 	bool bRootMotionPaused = false;
 };
 
+struct FRootMotionModifier_WarpSwapState : public FRootMotionModifierSwapState
+{
+	FRotator AdditionalRotationOffset = FRotator::ZeroRotator;
+	FTransform CachedTargetTransform = FTransform::Identity;
+	FTransform RootMotionRemainingAfterNotify = FTransform::Identity;
+	TOptional<FTransform> CachedOffsetFromWarpPoint;
+	FQuat RotationOffset = FQuat::Identity;
+	bool bWarpingPaused = false;
+	bool bRootMotionPaused = false;
+};
+
 // URootMotionModifier_SimpleWarp. 
 // DEPRECATED in favor of URootMotionModifier_SkewWarp (kept for reference)
 ///////////////////////////////////////////////////////////////
diff --git a/Engine/Plugins/Animation/MotionWarping/Source/MotionWarping/Public/RootMotionModifier_PrecomputedWarp.h b/Engine/Plugins/Animation/MotionWarping/Source/MotionWarping/Public/RootMotionModifier_PrecomputedWarp.h
--- a/Engine/Plugins/Animation/MotionWarping/Source/MotionWarping/Public/RootMotionModifier_PrecomputedWarp.h
+++ b/Engine/Plugins/Animation/MotionWarping/Source/MotionWarping/Public/RootMotionModifier_PrecomputedWarp.h
@@ -86,6 +86,10 @@ public:
 
 	MOTIONWARPING_API virtual FTransform ProcessRootMotion(const FTransform& InRootMotion, float DeltaSeconds) override;
 
+	MOTIONWARPING_API virtual TSharedPtr<FRootMotionModifierSwapState> GetSwapState() const override;
+	MOTIONWARPING_API virtual void PopulateSwapState(TSharedPtr<FRootMotionModifierSwapState> InSwapState) const override;
+	MOTIONWARPING_API virtual void RestoreFromSwapState(TSharedPtr<FRootMotionModifierSwapState> InSwapState) override;
+
 	UPROPERTY(EditAnywhere, Category = Blending, Meta = (EditCondition="!bSeparateTranslationCurves"))
 	FPrecomputedWarpCurve TranslationWarpingCurve;
 
@@ -153,3 +157,21 @@ public:
 
 	FAnimTrajectoryData AnimTrajectoryData;
 };
+
+struct FRootMotionModifier_PrecomputedWarpSwapState : public FRootMotionModifier_WarpSwapState
+{
+	bool bFirstFrame = true;
+	float ActualStartTime = 0.f;
+	float RoundedEndTime = 0.f;
+	float PreviousFrame = 0.f;
+
+	FTransform StartingRootTransform;
+	FTransform TargetTransform;
+
+	FQuat FilteredSteeringTarget = FQuat::Identity;
+	FQuaternionSpringState TargetSmoothingState;
+
+	TArray<FTransform> WarpedTrajectory;
+
+	URootMotionModifier_PrecomputedWarp::FAnimTrajectoryData AnimTrajectoryData;
+};
diff --git a/Engine/Plugins/Animation/PoseSearch/PoseSearch.uplugin b/Engine/Plugins/Animation/PoseSearch/PoseSearch.uplugin
--- a/Engine/Plugins/Animation/PoseSearch/PoseSearch.uplugin
+++ b/Engine/Plugins/Animation/PoseSearch/PoseSearch.uplugin
@@ -44,6 +44,10 @@
 			"TargetAllowList": [
 				"Editor"
 			]
+		},
+		{
+			"Name": "MotionWarping",
+			"Enabled": true
 		}
 	]
 }
\ No newline at end of file
diff --git a/Engine/Plugins/Animation/PoseSearch/Source/Runtime/PoseSearch.Build.cs b/Engine/Plugins/Animation/PoseSearch/Source/Runtime/PoseSearch.Build.cs
--- a/Engine/Plugins/Animation/PoseSearch/Source/Runtime/PoseSearch.Build.cs
+++ b/Engine/Plugins/Animation/PoseSearch/Source/Runtime/PoseSearch.Build.cs
@@ -32,6 +32,7 @@ public class PoseSearch : ModuleRules
 			"AnimationWarpingRuntime",
 			"DeveloperSettings",
 			"GameplayTags",
+			"MotionWarping",
 			"RewindDebuggerRuntimeInterface"
 		}
 		);
diff --git a/Engine/Plugins/Animation/PoseSearch/Source/Runtime/Private/PoseSearchTrajectoryLibrary.cpp b/Engine/Plugins/Animation/PoseSearch/Source/Runtime/Private/PoseSearchTrajectoryLibrary.cpp
--- a/Engine/Plugins/Animation/PoseSearch/Source/Runtime/Private/PoseSearchTrajectoryLibrary.cpp
+++ b/Engine/Plugins/Animation/PoseSearch/Source/Runtime/Private/PoseSearchTrajectoryLibrary.cpp
@@ -4,6 +4,7 @@
 #include "PoseSearch/PoseSearchTrajectoryPredictor.h"
 #include "Animation/AnimInstanceProxy.h"
 #include "Animation/AnimInstance.h"
+#include "Animation/AnimMontage.h"
 #include "Components/SkeletalMeshComponent.h"
 #include "Engine/Engine.h"
 #include "GameFramework/Character.h"
@@ -11,15 +12,24 @@
 #include "PoseSearch/PoseSearchDefines.h"
 #include "Kismet/KismetMathLibrary.h"
 #include "Animation/TrajectoryTypes.h"
+#include "MotionWarpingComponent.h"
+#include "MotionWarpingMontageTrajectoryAdapter.h"
 
 #include UE_INLINE_GENERATED_CPP_BY_NAME(PoseSearchTrajectoryLibrary)
 
 // @todo: remove GVarPoseSearchTrajectoryFacingFromMesh once the behavior is enabled by defaulted
 static bool GVarPoseSearchTrajectoryFacingFromMesh = true;
 #if ENABLE_ANIM_DEBUG
-static FAutoConsoleVariableRef CVarPoseSearchTrajectoryFacingFromMesh(TEXT("a.PoseSearchTrajectory.FacingFromMesh"), GVarPoseSearchTrajectoryFacingFromMesh, TEXT("Always generating the trajectory in mesh space rather than using the Character.View.Rotation.Yaw when the character is set to orienting the rotation to movement"));
+static FAutoConsoleVariableRef CVarPoseSearchTrajectoryFacingFromMesh(TEXT("a.PoseSearchTrajectory.FacingFromMesh")
+	, GVarPoseSearchTrajectoryFacingFromMesh
+	, TEXT("Always generating the trajectory in mesh space rather than using the Character.View.Rotation.Yaw when the character is set to orienting the rotation to movement"));
 #endif // ENABLE_ANIM_DEBUG
 
+static bool GVarPoseSearchTrajectoryUseWarpedMontage = false;
+static FAutoConsoleVariableRef CVarPoseSearchTrajectoryUseWarpedMontage(TEXT("a.PoseSearchTrajectory.UseWarpedMontage")
+	, GVarPoseSearchTrajectoryUseWarpedMontage
+	, TEXT("Experimental. If a montage is playing, use that montage's warped root motion for generated trajectories."));
+
 bool FPoseSearchTrajectoryData::UpdateData(
 	float DeltaTime,
 	const FAnimInstanceProxy& AnimInstanceProxy,
@@ -493,6 +503,104 @@ void UPoseSearchTrajectoryLibrary::UpdatePrediction_SimulateCharacterMovement(FT
 	}
 }
 
+void UPoseSearchTrajectoryLibrary::UpdatePrediction_SimulateMontage(const UObject* InContext,
+	const FAnimMontageInstance* MontageInstance,
+	FTransformTrajectory& Trajectory,
+	const FPoseSearchTrajectoryData& TrajectoryData, const FPoseSearchTrajectoryData::FDerived& TrajectoryDataDerived,
+	const FPoseSearchTrajectoryData::FSampling& TrajectoryDataSampling, float DeltaTime, bool bSimulateWarping, UMotionWarpingComponent* InWarpingComponent)
+{
+	// @TODO: Perhaps we should still simulate gravity from the CMC in addition to the warped montage
+	// See manipulations in `UCharacterMovementComponent::PerformMovement` on `AnimRootMotionVelocity` to see where this falls short
+
+	FTransform CurrentTransform = FTransform(TrajectoryDataDerived.Facing, TrajectoryDataDerived.Position);
+	
+	const int32 NumHistorySamples = TrajectoryDataSampling.NumHistorySamples;
+	const float SecondsPerPredictionSample = TrajectoryDataSampling.SecondsPerPredictionSample;
+
+	float AccumulatedSeconds = DeltaTime;
+
+	const ACharacter* Character = Cast<ACharacter>(InContext);
+	if (!Character)
+	{
+		if (const UAnimInstance* AnimInstance = Cast<UAnimInstance>(InContext))
+		{
+			Character = Cast<ACharacter>(AnimInstance->GetOwningActor());
+		}
+		else if (const UActorComponent* AnimNextComponent = Cast<UActorComponent>(InContext))
+		{
+			Character = Cast<ACharacter>(AnimNextComponent->GetOwner());
+		}
+	}
+
+	UMotionWarpingComponent* WarpingComp = InWarpingComponent;
+	if (!WarpingComp && bSimulateWarping)
+	{
+		WarpingComp = Character ? Character->FindComponentByClass<UMotionWarpingComponent>() : nullptr;
+	}
+	UMotionWarpingMontageTrajectoryAdapter* PredictionAdapter = nullptr;
+
+	FTransform VisualRootTransform = CurrentTransform;
+	if (bSimulateWarping && WarpingComp && !WarpingComp->GetModifiers().IsEmpty())
+	{
+		WarpingComp->BeginTrajectoryPredictions(MontageInstance);
+
+		// Visual roots should be base visual * transform, but in practice they are adapter implementation specifc & deviate (Individually consider scaled capsule height / etc).
+		// Preserve this deviation by transforming the visual root in addition to to current transform
+		PredictionAdapter = CastChecked<UMotionWarpingMontageTrajectoryAdapter>(WarpingComp->GetOwnerAdapter());
+		if (PredictionAdapter)
+		{
+			VisualRootTransform.SetLocation(PredictionAdapter->GetVisualRootLocation());
+		}
+	}
+
+	const int32 LastIndex = Trajectory.Samples.Num() - 1;
+	if (NumHistorySamples <= LastIndex)
+	{
+		for (int32 Index = NumHistorySamples; ; ++Index)
+		{
+			// Store Values
+			Trajectory.Samples[Index].Position = CurrentTransform.GetLocation();
+			Trajectory.Samples[Index].Facing = CurrentTransform.GetRotation();
+			Trajectory.Samples[Index].TimeInSeconds = AccumulatedSeconds;
+
+			if (Index == LastIndex)
+			{
+				break;
+			}
+
+			// Forward step
+			float CurrentMontagePosition = MontageInstance->GetPosition() + AccumulatedSeconds;
+			FAnimExtractContext ExtractionCtx(0.0, true, FDeltaTimeRecord(), false);
+			FTransform RootMotion = MontageInstance->Montage->ExtractRootMotionFromTrackRange(CurrentMontagePosition, CurrentMontagePosition + SecondsPerPredictionSample, ExtractionCtx);
+			
+			if (bSimulateWarping && PredictionAdapter && WarpingComp && WarpingComp->IsPredictingTrajectory())
+			{
+				const FTransform WarpedRootMotion = PredictionAdapter->WarpLocalRootMotionOnMontageTrajectory(RootMotion, SecondsPerPredictionSample, CurrentMontagePosition, CurrentMontagePosition + SecondsPerPredictionSample);
+
+				CurrentTransform = WarpedRootMotion * CurrentTransform;
+				VisualRootTransform = WarpedRootMotion * VisualRootTransform;
+
+				// Root modifiers expect actor component transforms, not root transforms
+				PredictionAdapter->SetCurrentTransform(FTransform(PredictionAdapter->GetBaseVisualRotationOffset(), PredictionAdapter->GetBaseVisualTranslationOffset()).Inverse() * CurrentTransform);
+
+				// Adapter separates base visual rotation / offset capsule concerns, opt to preserve via iterative changes
+				PredictionAdapter->SetVisualRootLocation(VisualRootTransform.GetLocation());
+			}
+			else
+			{
+				CurrentTransform = RootMotion * CurrentTransform;
+			}
+
+			AccumulatedSeconds += SecondsPerPredictionSample;
+		}
+	}
+
+	if (bSimulateWarping && WarpingComp && WarpingComp->IsPredictingTrajectory())
+	{
+		WarpingComp->EndTrajectoryPredictions();
+	}
+}
+
 PRAGMA_DISABLE_DEPRECATION_WARNINGS
 
 // Deprecated
@@ -523,6 +631,25 @@ void UPoseSearchTrajectoryLibrary::PoseSearchGenerateTransformTrajectory(const U
 	float InDeltaTime, FTransformTrajectory& InOutTrajectory, float& InOutDesiredControllerYawLastUpdate, FTransformTrajectory& OutTrajectory,
 	float InHistorySamplingInterval, int32 InTrajectoryHistoryCount, float InPredictionSamplingInterval, int32 InTrajectoryPredictionCount)
 {
+	if (GVarPoseSearchTrajectoryUseWarpedMontage)
+	{
+		UPoseSearchTrajectoryLibrary::PoseSearchGenerateWarpedTransformTrajectory(
+			InContext,
+			InTrajectoryData,
+			InDeltaTime,
+			InOutTrajectory,
+			InOutDesiredControllerYawLastUpdate,
+			OutTrajectory,
+			InHistorySamplingInterval,
+			InTrajectoryHistoryCount,
+			InPredictionSamplingInterval,
+			InTrajectoryPredictionCount,
+			true,
+			true);
+
+		return;
+	}
+
 	FPoseSearchTrajectoryData::FSampling TrajectoryDataSampling;
 	TrajectoryDataSampling.NumHistorySamples = InTrajectoryHistoryCount;
 	TrajectoryDataSampling.SecondsPerHistorySample = InHistorySamplingInterval;
@@ -543,6 +670,53 @@ void UPoseSearchTrajectoryLibrary::PoseSearchGenerateTransformTrajectory(const U
 	OutTrajectory = InOutTrajectory;
 }
 
+void UPoseSearchTrajectoryLibrary::PoseSearchGenerateWarpedTransformTrajectory(const UObject* InContext, const FPoseSearchTrajectoryData& InTrajectoryData,
+	float InDeltaTime, FTransformTrajectory& InOutTrajectory, float& InOutDesiredControllerYawLastUpdate, FTransformTrajectory& OutTrajectory,
+	float InHistorySamplingInterval, int32 InTrajectoryHistoryCount, float InPredictionSamplingInterval, int32 InTrajectoryPredictionCount,
+	bool bSimulateMontages, bool bSimulateWarping)
+{
+	FPoseSearchTrajectoryData::FSampling TrajectoryDataSampling;
+	TrajectoryDataSampling.NumHistorySamples = InTrajectoryHistoryCount;
+	TrajectoryDataSampling.SecondsPerHistorySample = InHistorySamplingInterval;
+	TrajectoryDataSampling.NumPredictionSamples = InTrajectoryPredictionCount;
+	TrajectoryDataSampling.SecondsPerPredictionSample = InPredictionSamplingInterval;
+
+	FPoseSearchTrajectoryData::FState TrajectoryDataState;
+	TrajectoryDataState.DesiredControllerYawLastUpdate = InOutDesiredControllerYawLastUpdate;
+
+	FPoseSearchTrajectoryData::FDerived TrajectoryDataDerived;
+	InTrajectoryData.UpdateData(InDeltaTime, InContext, TrajectoryDataDerived, TrajectoryDataState);
+	InitTrajectorySamples(InOutTrajectory, TrajectoryDataDerived.Position, TrajectoryDataDerived.Facing, TrajectoryDataSampling, InDeltaTime);
+	UpdateHistory_TransformHistory(InOutTrajectory, TrajectoryDataDerived.Position, TrajectoryDataDerived.Velocity, TrajectoryDataSampling, InDeltaTime);
+
+	const ACharacter* Character = Cast<ACharacter>(InContext);
+	if (!Character)
+	{
+		if (const UAnimInstance* AnimInstance = Cast<UAnimInstance>(InContext))
+		{
+			Character = Cast<ACharacter>(AnimInstance->GetOwningActor());
+		}
+		else if (const UActorComponent* AnimNextComponent = Cast<UActorComponent>(InContext))
+		{
+			Character = Cast<ACharacter>(AnimNextComponent->GetOwner());
+		}
+	}
+
+	const FAnimMontageInstance* AnimMontageInstance = Character ? Character->GetRootMotionAnimMontageInstance() : nullptr;
+	if (AnimMontageInstance && bSimulateMontages)
+	{
+		UpdatePrediction_SimulateMontage(InContext, AnimMontageInstance, InOutTrajectory, InTrajectoryData, TrajectoryDataDerived, TrajectoryDataSampling, InDeltaTime, bSimulateWarping);
+	}
+	else
+	{
+		UpdatePrediction_SimulateCharacterMovement(InOutTrajectory, InTrajectoryData, TrajectoryDataDerived, TrajectoryDataSampling, InDeltaTime);
+	}
+
+	InOutDesiredControllerYawLastUpdate = TrajectoryDataState.DesiredControllerYawLastUpdate;
+
+	OutTrajectory = InOutTrajectory;
+}
+
 void UPoseSearchTrajectoryLibrary::PoseSearchGeneratePredictorTransformTrajectory(UObject* InPredictor,
 	const FPoseSearchTrajectoryData& InTrajectoryData, float InDeltaTime, FTransformTrajectory& InOutTrajectory,
 	float& InOutDesiredControllerYawLastUpdate, FTransformTrajectory& OutTrajectory, float InHistorySamplingInterval,
diff --git a/Engine/Plugins/Animation/PoseSearch/Source/Runtime/Public/PoseSearch/PoseSearchTrajectoryLibrary.h b/Engine/Plugins/Animation/PoseSearch/Source/Runtime/Public/PoseSearch/PoseSearchTrajectoryLibrary.h
--- a/Engine/Plugins/Animation/PoseSearch/Source/Runtime/Public/PoseSearch/PoseSearchTrajectoryLibrary.h
+++ b/Engine/Plugins/Animation/PoseSearch/Source/Runtime/Public/PoseSearch/PoseSearchTrajectoryLibrary.h
@@ -10,7 +10,10 @@
 #define UE_API POSESEARCH_API
 
 class UAnimInstance;
+class UMotionWarpingComponent;
 class IPoseSearchTrajectoryPredictorInterface;
+
+struct FAnimMontageInstance;
 struct FTransformTrajectory;
 struct FTransformTrajectorySample;
 
@@ -295,6 +298,18 @@ public:
 	// Update prediction by simulating the movement math for ground locomotion from UCharacterMovementComponent.
 	static UE_API void UpdatePrediction_SimulateCharacterMovement(FTransformTrajectory& Trajectory, const FPoseSearchTrajectoryData& TrajectoryData, const FPoseSearchTrajectoryData::FDerived& TrajectoryDataDerived, const FPoseSearchTrajectoryData::FSampling& TrajectoryDataSampling, float DeltaTime);
 
+	// Update prediction by simulating an active montage's root motion.
+	UE_EXPERIMENTAL(5.8, "Experimental pose search function, API may change in the future")
+	static UE_API void UpdatePrediction_SimulateMontage(const UObject* InContext,
+															  const FAnimMontageInstance* MontageInstance,
+															  FTransformTrajectory& Trajectory,
+															  const FPoseSearchTrajectoryData& TrajectoryData,
+															  const FPoseSearchTrajectoryData::FDerived& TrajectoryDataDerived,
+															  const FPoseSearchTrajectoryData::FSampling& TrajectoryDataSampling,
+															  float DeltaTime,
+															  bool bSimulateWarping = true,
+															  UMotionWarpingComponent* InWarpingComponent = nullptr);
+
 	// Experimental: Update the history purely based on current position, without taking into account ground velocities
 	UE_EXPERIMENTAL(5.7, "Experimental pose search function, API may change in the future")
 	static UE_API void UpdateHistory_WorldSpace(FTransformTrajectory& Trajectory,
@@ -308,6 +323,14 @@ public:
 		UPARAM(ref) FTransformTrajectory& InOutTrajectory, UPARAM(ref) float& InOutDesiredControllerYawLastUpdate, FTransformTrajectory& OutTrajectory,
 		float InHistorySamplingInterval = 0.04f, int32 InTrajectoryHistoryCount = 10, float InPredictionSamplingInterval = 0.2f, int32 InTrajectoryPredictionCount = 8);
 
+	UE_EXPERIMENTAL(5.8, "Experimental pose search function, API may change in the future")
+	UFUNCTION(BlueprintCallable, Category = "Animation|PoseSearch", meta = (BlueprintThreadSafe, DisplayName = "Pose Search Generate Trajectory (for Character)"))
+	static UE_API void PoseSearchGenerateWarpedTransformTrajectory(const UObject* InAnimInstance,
+		UPARAM(ref) const FPoseSearchTrajectoryData& InTrajectoryData, float InDeltaTime,
+		UPARAM(ref) FTransformTrajectory& InOutTrajectory, UPARAM(ref) float& InOutDesiredControllerYawLastUpdate, FTransformTrajectory& OutTrajectory,
+		float InHistorySamplingInterval = 0.04f, int32 InTrajectoryHistoryCount = 10, float InPredictionSamplingInterval = 0.2f, int32 InTrajectoryPredictionCount = 8,
+		bool bSimulateMontages = false, bool bSimulateWarping = false);
+
 	// Generates a prediction trajectory based of the current movement intent. For use with predictors. InPredictor must implement IPoseSearchTrajectoryPredictorInterface
 	UFUNCTION(BlueprintCallable, Category = "Animation|PoseSearch|Experimental", meta = (BlueprintThreadSafe, DisplayName = "Pose Search Generate Trajectory (using Predictor)", DeprecatedFunction, DeprecationMessage="Use PoseSearchGeneratePredictorTrajectory that does not take FPoseSearchTrajectoryData."))
 	static UE_API void PoseSearchGeneratePredictorTransformTrajectory(UObject* InPredictor,
