PoseSearch Trajectory: Allow for CMC prediction after warped montages.

#tests PIE

#Synced-CL 49775968

diff --git a/Engine/Plugins/Animation/PoseSearch/Source/Runtime/Private/PoseSearchTrajectoryLibrary.cpp b/Engine/Plugins/Animation/PoseSearch/Source/Runtime/Private/PoseSearchTrajectoryLibrary.cpp
--- a/Engine/Plugins/Animation/PoseSearch/Source/Runtime/Private/PoseSearchTrajectoryLibrary.cpp
+++ b/Engine/Plugins/Animation/PoseSearch/Source/Runtime/Private/PoseSearchTrajectoryLibrary.cpp
@@ -77,7 +77,8 @@ bool FPoseSearchTrajectoryData::UpdateData(
 
 	TrajectoryDataDerived.Velocity = MoveComp->Velocity;
 	TrajectoryDataDerived.Acceleration = MoveComp->GetCurrentAcceleration();
-		
+	
+	// @TODO: Consider if our trajectory implies we will become falling mid future trajectory, rather than rely only on initial state
 	TrajectoryDataDerived.bStepGroundPrediction = !MoveComp->IsFalling() && !MoveComp->IsFlying();
 
 	if (TrajectoryDataDerived.Acceleration.IsZero())
@@ -425,27 +426,31 @@ FVector UPoseSearchTrajectoryLibrary::RemapVectorMagnitudeWithCurve(
 	return Vector;
 }
 
-void UPoseSearchTrajectoryLibrary::UpdatePrediction_SimulateCharacterMovement(FTransformTrajectory& Trajectory,
-	const FPoseSearchTrajectoryData& TrajectoryData, const FPoseSearchTrajectoryData::FDerived& TrajectoryDataDerived,
-	const FPoseSearchTrajectoryData::FSampling& TrajectoryDataSampling, float DeltaTime)
+void UPoseSearchTrajectoryLibrary::UpdatePrediction_SimulateCharacterMovementInit(FTransformTrajectory& Trajectory, 
+	FVector& OutCurrentPositionWS, 
+	FVector& OutCurrentVelocityWS, 
+	FVector& OutCurrentAccelerationWS,
+	FQuat& OutCurrentFacingWS,
+	const FPoseSearchTrajectoryData& TrajectoryData, 
+	const FPoseSearchTrajectoryData::FDerived& TrajectoryDataDerived)
 {
-	FVector CurrentPositionWS = TrajectoryDataDerived.Position;
-	FVector CurrentVelocityWS = RemapVectorMagnitudeWithCurve(TrajectoryDataDerived.Velocity, TrajectoryData.bUseSpeedRemappingCurve, TrajectoryData.SpeedRemappingCurve);
-	FVector CurrentAccelerationWS = RemapVectorMagnitudeWithCurve(TrajectoryDataDerived.Acceleration, TrajectoryData.bUseAccelerationRemappingCurve, TrajectoryData.AccelerationRemappingCurve);
+	OutCurrentPositionWS = TrajectoryDataDerived.Position;
+	OutCurrentVelocityWS = RemapVectorMagnitudeWithCurve(TrajectoryDataDerived.Velocity, TrajectoryData.bUseSpeedRemappingCurve, TrajectoryData.SpeedRemappingCurve);
+	OutCurrentAccelerationWS = RemapVectorMagnitudeWithCurve(TrajectoryDataDerived.Acceleration, TrajectoryData.bUseAccelerationRemappingCurve, TrajectoryData.AccelerationRemappingCurve);
 
 	// Bending CurrentVelocityWS towards CurrentAccelerationWS
-	if (TrajectoryData.BendVelocityTowardsAcceleration > UE_KINDA_SMALL_NUMBER && !CurrentAccelerationWS.IsNearlyZero())
+	if (TrajectoryData.BendVelocityTowardsAcceleration > UE_KINDA_SMALL_NUMBER && !OutCurrentAccelerationWS.IsNearlyZero())
 	{
-		const float CurrentSpeed = CurrentVelocityWS.Length();
-		const FVector VelocityWSAlongAcceleration = CurrentAccelerationWS.GetUnsafeNormal() * CurrentSpeed;
+		const float CurrentSpeed = OutCurrentVelocityWS.Length();
+		const FVector VelocityWSAlongAcceleration = OutCurrentAccelerationWS.GetUnsafeNormal() * CurrentSpeed;
 		if (TrajectoryData.BendVelocityTowardsAcceleration < 1.f - UE_KINDA_SMALL_NUMBER)
 		{
-			CurrentVelocityWS = FMath::Lerp(CurrentVelocityWS, VelocityWSAlongAcceleration, TrajectoryData.BendVelocityTowardsAcceleration);
+			OutCurrentVelocityWS = FMath::Lerp(OutCurrentVelocityWS, VelocityWSAlongAcceleration, TrajectoryData.BendVelocityTowardsAcceleration);
 
-			const float NewLength = CurrentVelocityWS.Length();
+			const float NewLength = OutCurrentVelocityWS.Length();
 			if (NewLength > UE_KINDA_SMALL_NUMBER)
 			{
-				CurrentVelocityWS *= CurrentSpeed / NewLength;
+				OutCurrentVelocityWS *= CurrentSpeed / NewLength;
 			}
 			else
 			{
@@ -454,16 +459,66 @@ void UPoseSearchTrajectoryLibrary::UpdatePrediction_SimulateCharacterMovement(FT
 		}
 		else
 		{
-			CurrentVelocityWS = VelocityWSAlongAcceleration;
+			OutCurrentVelocityWS = VelocityWSAlongAcceleration;
 		}
 	}
 
-	FQuat CurrentFacingWS = TrajectoryDataDerived.Facing;
-	
-	const int32 NumHistorySamples = TrajectoryDataSampling.NumHistorySamples;
+	OutCurrentFacingWS = TrajectoryDataDerived.Facing;
+}
+
+void UPoseSearchTrajectoryLibrary::UpdatePrediction_SimulateCharacterMovementStep(FTransformTrajectory& Trajectory,
+	FVector& InOutCurrentPositionWS,
+	FVector& InOutCurrentVelocityWS,
+	FVector& InOutCurrentAccelerationWS,
+	FQuat& InOutCurrentFacingWS,
+	const FPoseSearchTrajectoryData& TrajectoryData,
+	const FPoseSearchTrajectoryData::FDerived& TrajectoryDataDerived,
+	const FPoseSearchTrajectoryData::FSampling& TrajectoryDataSampling,
+	float& InOutAccumulatedSeconds)
+{
 	const float SecondsPerPredictionSample = TrajectoryDataSampling.SecondsPerPredictionSample;
 	const FQuat ControllerRotationPerStep = FQuat::MakeFromEuler(FVector(0.f, 0.f, TrajectoryDataDerived.ControllerYawRate * SecondsPerPredictionSample));
 
+	InOutCurrentPositionWS += InOutCurrentVelocityWS * SecondsPerPredictionSample;
+	InOutAccumulatedSeconds += SecondsPerPredictionSample;
+
+	if (TrajectoryDataDerived.bStepGroundPrediction)
+	{
+		InOutCurrentAccelerationWS = RemapVectorMagnitudeWithCurve(ControllerRotationPerStep * InOutCurrentAccelerationWS,
+			TrajectoryData.bUseAccelerationRemappingCurve, TrajectoryData.AccelerationRemappingCurve);
+		const FVector NewVelocityWS = TrajectoryData.StepCharacterMovementGroundPrediction(SecondsPerPredictionSample, InOutCurrentVelocityWS, InOutCurrentAccelerationWS, TrajectoryDataDerived);
+		InOutCurrentVelocityWS = RemapVectorMagnitudeWithCurve(NewVelocityWS, TrajectoryData.bUseSpeedRemappingCurve, TrajectoryData.SpeedRemappingCurve);
+
+		// Account for the controller (e.g. the camera) rotating.
+		InOutCurrentFacingWS = ControllerRotationPerStep * InOutCurrentFacingWS;
+		if (TrajectoryDataDerived.bOrientRotationToMovement && !InOutCurrentAccelerationWS.IsNearlyZero())
+		{
+			// Rotate towards acceleration.
+			const FVector CurrentAccelerationCS = TrajectoryDataDerived.MeshCompRelativeRotation.RotateVector(InOutCurrentAccelerationWS);
+			InOutCurrentFacingWS = FMath::QInterpConstantTo(InOutCurrentFacingWS, CurrentAccelerationCS.ToOrientationQuat(), SecondsPerPredictionSample, TrajectoryData.RotateTowardsMovementSpeed);
+		}
+	}
+}
+
+void UPoseSearchTrajectoryLibrary::UpdatePrediction_SimulateCharacterMovement(FTransformTrajectory& Trajectory,
+	const FPoseSearchTrajectoryData& TrajectoryData, const FPoseSearchTrajectoryData::FDerived& TrajectoryDataDerived,
+	const FPoseSearchTrajectoryData::FSampling& TrajectoryDataSampling, float DeltaTime)
+{
+	FVector CurrentPositionWS = FVector::ZeroVector;
+	FVector CurrentVelocityWS = FVector::ZeroVector;
+	FVector CurrentAccelerationWS = FVector::ZeroVector;
+	FQuat CurrentFacingWS = FQuat::Identity;
+
+	UPoseSearchTrajectoryLibrary::UpdatePrediction_SimulateCharacterMovementInit(Trajectory, 
+		CurrentPositionWS, 
+		CurrentVelocityWS, 
+		CurrentAccelerationWS, 
+		CurrentFacingWS,
+		TrajectoryData, 
+		TrajectoryDataDerived);
+	
+	const int32 NumHistorySamples = TrajectoryDataSampling.NumHistorySamples;
+
 	float AccumulatedSeconds = DeltaTime;
 
 	const int32 LastIndex = Trajectory.Samples.Num() - 1;
@@ -480,25 +535,15 @@ void UPoseSearchTrajectoryLibrary::UpdatePrediction_SimulateCharacterMovement(FT
 				break;
 			}
 
-			CurrentPositionWS += CurrentVelocityWS * SecondsPerPredictionSample;
-			AccumulatedSeconds += SecondsPerPredictionSample;
-
-			if (TrajectoryDataDerived.bStepGroundPrediction)
-			{
-				CurrentAccelerationWS = RemapVectorMagnitudeWithCurve(ControllerRotationPerStep * CurrentAccelerationWS,
-					TrajectoryData.bUseAccelerationRemappingCurve, TrajectoryData.AccelerationRemappingCurve);
-				const FVector NewVelocityWS = TrajectoryData.StepCharacterMovementGroundPrediction(SecondsPerPredictionSample, CurrentVelocityWS, CurrentAccelerationWS, TrajectoryDataDerived);
-				CurrentVelocityWS = RemapVectorMagnitudeWithCurve(NewVelocityWS, TrajectoryData.bUseSpeedRemappingCurve, TrajectoryData.SpeedRemappingCurve);
-
-				// Account for the controller (e.g. the camera) rotating.
-				CurrentFacingWS = ControllerRotationPerStep * CurrentFacingWS;
-				if (TrajectoryDataDerived.bOrientRotationToMovement && !CurrentAccelerationWS.IsNearlyZero())
-				{
-					// Rotate towards acceleration.
-					const FVector CurrentAccelerationCS = TrajectoryDataDerived.MeshCompRelativeRotation.RotateVector(CurrentAccelerationWS);
-					CurrentFacingWS = FMath::QInterpConstantTo(CurrentFacingWS, CurrentAccelerationCS.ToOrientationQuat(), SecondsPerPredictionSample, TrajectoryData.RotateTowardsMovementSpeed);
-				}
-			}
+			UPoseSearchTrajectoryLibrary::UpdatePrediction_SimulateCharacterMovementStep(Trajectory, 
+				CurrentPositionWS, 
+				CurrentVelocityWS, 
+				CurrentAccelerationWS, 
+				CurrentFacingWS, 
+				TrajectoryData, 
+				TrajectoryDataDerived, 
+				TrajectoryDataSampling,
+				AccumulatedSeconds);
 		}
 	}
 }
@@ -507,10 +552,28 @@ void UPoseSearchTrajectoryLibrary::UpdatePrediction_SimulateMontage(const UObjec
 	const FAnimMontageInstance* MontageInstance,
 	FTransformTrajectory& Trajectory,
 	const FPoseSearchTrajectoryData& TrajectoryData, const FPoseSearchTrajectoryData::FDerived& TrajectoryDataDerived,
-	const FPoseSearchTrajectoryData::FSampling& TrajectoryDataSampling, float DeltaTime, bool bSimulateWarping, UMotionWarpingComponent* InWarpingComponent)
+	const FPoseSearchTrajectoryData::FSampling& TrajectoryDataSampling, float DeltaTime, bool bSimulateWarping, bool bSimulateCharacterMovement, UMotionWarpingComponent* InWarpingComponent)
 {
-	// @TODO: Perhaps we should still simulate gravity from the CMC in addition to the warped montage
-	// See manipulations in `UCharacterMovementComponent::PerformMovement` on `AnimRootMotionVelocity` to see where this falls short
+	FVector CurrentPositionWS = FVector::ZeroVector;
+	FVector CurrentVelocityWS = FVector::ZeroVector;
+	FVector CurrentAccelerationWS = FVector::ZeroVector;
+	FQuat CurrentFacingWS = FQuat::Identity;
+
+	if (bSimulateCharacterMovement)
+	{
+		UPoseSearchTrajectoryLibrary::UpdatePrediction_SimulateCharacterMovementInit(Trajectory,
+			CurrentPositionWS,
+			CurrentVelocityWS,
+			CurrentAccelerationWS,
+			CurrentFacingWS,
+			TrajectoryData,
+			TrajectoryDataDerived);
+	}
+	else
+	{
+		CurrentPositionWS = TrajectoryDataDerived.Position;
+		CurrentFacingWS = TrajectoryDataDerived.Facing;
+	}
 
 	FTransform CurrentTransform = FTransform(TrajectoryDataDerived.Facing, TrajectoryDataDerived.Position);
 	
@@ -559,8 +622,8 @@ void UPoseSearchTrajectoryLibrary::UpdatePrediction_SimulateMontage(const UObjec
 		for (int32 Index = NumHistorySamples; ; ++Index)
 		{
 			// Store Values
-			Trajectory.Samples[Index].Position = CurrentTransform.GetLocation();
-			Trajectory.Samples[Index].Facing = CurrentTransform.GetRotation();
+			Trajectory.Samples[Index].Position = CurrentPositionWS;
+			Trajectory.Samples[Index].Facing = CurrentFacingWS;
 			Trajectory.Samples[Index].TimeInSeconds = AccumulatedSeconds;
 
 			if (Index == LastIndex)
@@ -591,7 +654,35 @@ void UPoseSearchTrajectoryLibrary::UpdatePrediction_SimulateMontage(const UObjec
 				CurrentTransform = RootMotion * CurrentTransform;
 			}
 
-			AccumulatedSeconds += SecondsPerPredictionSample;
+			CurrentPositionWS = CurrentTransform.GetLocation();
+			CurrentFacingWS = CurrentTransform.GetRotation();
+
+			if (bSimulateCharacterMovement && PredictionAdapter)
+			{
+				// @TODO: Ideally we should feedback our prediction, but there are issues in trajectory simulation not considering accel that make this a bad idea
+				// For now, don't feedback, which is equivalent to only using initial velocity as an offset. As that gives better results / closer mimics CMC.
+				
+				// Need to manually update visual root before velocity changes
+				// PredictionAdapter->SetVisualRootLocation(VisualRootTransform.GetLocation() += CurrentVelocityWS * SecondsPerPredictionSample);
+
+				UPoseSearchTrajectoryLibrary::UpdatePrediction_SimulateCharacterMovementStep(Trajectory,
+					CurrentPositionWS,
+					CurrentVelocityWS,
+					CurrentAccelerationWS,
+					CurrentFacingWS,
+					TrajectoryData,
+					TrajectoryDataDerived,
+					TrajectoryDataSampling,
+					AccumulatedSeconds);
+
+				// CurrentTransform.SetLocation(CurrentPositionWS);
+				// CurrentTransform.SetRotation(CurrentFacingWS);
+				// PredictionAdapter->SetCurrentTransform(FTransform(PredictionAdapter->GetBaseVisualRotationOffset(), PredictionAdapter->GetBaseVisualTranslationOffset()).Inverse() * CurrentTransform);
+			}
+			else
+			{
+				AccumulatedSeconds += SecondsPerPredictionSample;
+			}
 		}
 	}
 
diff --git a/Engine/Plugins/Animation/PoseSearch/Source/Runtime/Public/PoseSearch/PoseSearchTrajectoryLibrary.h b/Engine/Plugins/Animation/PoseSearch/Source/Runtime/Public/PoseSearch/PoseSearchTrajectoryLibrary.h
--- a/Engine/Plugins/Animation/PoseSearch/Source/Runtime/Public/PoseSearch/PoseSearchTrajectoryLibrary.h
+++ b/Engine/Plugins/Animation/PoseSearch/Source/Runtime/Public/PoseSearch/PoseSearchTrajectoryLibrary.h
@@ -308,6 +308,7 @@ public:
 															  const FPoseSearchTrajectoryData::FSampling& TrajectoryDataSampling,
 															  float DeltaTime,
 															  bool bSimulateWarping = true,
+															  bool bSimulateCharacterMovement = true,
 															  UMotionWarpingComponent* InWarpingComponent = nullptr);
 
 	// Experimental: Update the history purely based on current position, without taking into account ground velocities
@@ -380,6 +381,24 @@ public:
 private:
 	
 	static UE_API FVector RemapVectorMagnitudeWithCurve(const FVector& Vector, bool bUseCurve, const FRuntimeFloatCurve& Curve);
+
+	static UE_API void UpdatePrediction_SimulateCharacterMovementInit(FTransformTrajectory& Trajectory, 
+																	  FVector& OutCurrentPositionWS,
+																	  FVector& OutCurrentVelocityWS,
+																	  FVector& OutCurrentAccelerationWS,
+																	  FQuat& OutCurrentFacingWS,
+																	  const FPoseSearchTrajectoryData& TrajectoryData, 
+																	  const FPoseSearchTrajectoryData::FDerived& TrajectoryDataDerived);
+
+	static UE_API void UpdatePrediction_SimulateCharacterMovementStep(FTransformTrajectory& Trajectory, 
+																	  FVector& InOutCurrentPositionWS,
+																	  FVector& InOutCurrentVelocityWS,
+																	  FVector& InOutCurrentAccelerationWS,
+																	  FQuat& InOutCurrentFacingWS,
+																	  const FPoseSearchTrajectoryData& TrajectoryData, 
+																	  const FPoseSearchTrajectoryData::FDerived& TrajectoryDataDerived, 
+																	  const FPoseSearchTrajectoryData::FSampling& TrajectoryDataSampling, 
+																	  float& InOutAccumulatedSeconds);
 };
 
 #undef UE_API
