PoseSearch
- fix for UPoseSearchInteractionAsset::CalculateWarpTransforms edge case when characters are too close to the reference system frame
- removed no longer necessary UPoseSearchInteractionAsset::CalculateWarpTransforms input RelevantRoleIndexes

#Synced-CL 49766586

diff --git a/Engine/Plugins/Animation/PoseSearch/Source/Editor/Private/PoseSearchInteractionAssetEditor.cpp b/Engine/Plugins/Animation/PoseSearch/Source/Editor/Private/PoseSearchInteractionAssetEditor.cpp
--- a/Engine/Plugins/Animation/PoseSearch/Source/Editor/Private/PoseSearchInteractionAssetEditor.cpp
+++ b/Engine/Plugins/Animation/PoseSearch/Source/Editor/Private/PoseSearchInteractionAssetEditor.cpp
@@ -376,9 +376,8 @@ void FInteractionAssetViewModel::Tick(float DeltaSeconds)
 			TArray<FTransform, TInlineAllocator<PreallocatedRolesNum>> FullAlignedActorTransforms;
 			FullAlignedActorTransforms.SetNum(NumRoles);
 			const TConstArrayView<const UMirrorDataTable*> MirrorDataTables;
-			const TConstArrayView<bool> RelevantRoleIndexes;
 
-			InteractionAsset->CalculateWarpTransforms(PlayTime, 0.f, ActorTransforms, FullAlignedActorTransforms, MirrorDataTables, RelevantRoleIndexes, AnimContexts);
+			InteractionAsset->CalculateWarpTransforms(PlayTime, 0.f, ActorTransforms, FullAlignedActorTransforms, MirrorDataTables, AnimContexts);
 
 #if ENABLE_ANIM_DEBUG
 			const UWorld* DebugWorld = GetWorld();
@@ -408,7 +407,7 @@ void FInteractionAssetViewModel::Tick(float DeltaSeconds)
 #if ENABLE_ANIM_DEBUG
 			if (InteractionAsset->DebugTimeOffset > UE_KINDA_SMALL_NUMBER)
 			{
-				InteractionAsset->CalculateWarpTransforms(PlayTime, InteractionAsset->DebugTimeOffset, ActorTransforms, FullAlignedActorTransforms, MirrorDataTables, RelevantRoleIndexes, AnimContexts);
+				InteractionAsset->CalculateWarpTransforms(PlayTime, InteractionAsset->DebugTimeOffset, ActorTransforms, FullAlignedActorTransforms, MirrorDataTables, AnimContexts);
 
 				for (int32 InteractionAssetRoleIndex = 0; InteractionAssetRoleIndex < NumRoles; ++InteractionAssetRoleIndex)
 				{
diff --git a/Engine/Plugins/Animation/PoseSearch/Source/Runtime/Private/PoseSearchInteractionAsset.cpp b/Engine/Plugins/Animation/PoseSearch/Source/Runtime/Private/PoseSearchInteractionAsset.cpp
--- a/Engine/Plugins/Animation/PoseSearch/Source/Runtime/Private/PoseSearchInteractionAsset.cpp
+++ b/Engine/Plugins/Animation/PoseSearch/Source/Runtime/Private/PoseSearchInteractionAsset.cpp
@@ -12,6 +12,170 @@
 
 #include UE_INLINE_GENERATED_CPP_BY_NAME(PoseSearchInteractionAsset)
 
+namespace UE::PoseSearch
+{
+
+static FVector FindReferencePosition(const TArrayView<const FTransform> Transforms, const TArrayView<float> NormalizedWarpingWeightTranslation)
+{
+	const int32 Num = Transforms.Num();
+	check(Num == NormalizedWarpingWeightTranslation.Num());
+
+	FVector PositionsSum = FVector::ZeroVector;
+	for (int32 ItemIndex = 0; ItemIndex < Num; ++ItemIndex)
+	{
+		PositionsSum += Transforms[ItemIndex].GetTranslation() * NormalizedWarpingWeightTranslation[ItemIndex];
+	}
+
+	return PositionsSum;
+}
+
+static FQuat FindReferenceOrientationNoBanking(const TArrayView<const FTransform> Transforms, const TArrayView<int32> SortedByWarpingWeightRotationItemIndex, const TArrayView<float> NormalizedWarpingWeightRotation)
+{
+	const int32 Num = Transforms.Num();
+	check(Num > 0 && Num == SortedByWarpingWeightRotationItemIndex.Num() && Num == NormalizedWarpingWeightRotation.Num());
+
+	// @todo: use a Slerp or a proper FastLerp
+	// for now we don't account for the shortest path while FastLerping those queternion together
+	FQuat WeightedQuaternion = FQuat::Identity * 0.f;
+	for (int32 ItemIndex : SortedByWarpingWeightRotationItemIndex)
+	{
+		WeightedQuaternion += Transforms[ItemIndex].GetRotation() * NormalizedWarpingWeightRotation[ItemIndex];
+	}
+	WeightedQuaternion.Normalize();
+	return WeightedQuaternion;
+}
+
+static FQuat FindReferenceOrientationFullBanking(const TArrayView<const FTransform> Transforms, const TArrayView<int32> SortedByWarpingWeightRotationItemIndex, const TArrayView<float> NormalizedWarpingWeightRotation)
+{
+	const int32 Num = Transforms.Num();
+
+	check(Num > 0 && Num == SortedByWarpingWeightRotationItemIndex.Num() && Num == NormalizedWarpingWeightRotation.Num());
+
+	if (Num > 1)
+	{
+		const int32 LastItemIndex = Num - 1;
+		FVector OtherItemsPositionsSum = Transforms[SortedByWarpingWeightRotationItemIndex[0]].GetTranslation();
+		for (int32 ItemIndex = 1; ItemIndex < LastItemIndex; ++ItemIndex)
+		{
+			OtherItemsPositionsSum += Transforms[SortedByWarpingWeightRotationItemIndex[ItemIndex]].GetTranslation();
+		}
+
+		const FVector OtherItemsPositionAverage = OtherItemsPositionsSum / LastItemIndex;
+		const FVector DeltaPosition = OtherItemsPositionAverage - Transforms[SortedByWarpingWeightRotationItemIndex[LastItemIndex]].GetTranslation();
+
+		if (!DeltaPosition.IsNearlyZero())
+		{
+			return DeltaPosition.ToOrientationQuat();
+		}
+	}
+
+	return FindReferenceOrientationNoBanking(Transforms, SortedByWarpingWeightRotationItemIndex, NormalizedWarpingWeightRotation);
+}
+
+static FQuat FindReferenceOrientation(const TArrayView<const FTransform> Transforms, const TArrayView<int32> SortedByWarpingWeightRotationItemIndex, const TArrayView<float> NormalizedWarpingWeightRotation, float WarpingBankingWeight)
+{
+	if (WarpingBankingWeight < UE_KINDA_SMALL_NUMBER)
+	{
+		return FindReferenceOrientationNoBanking(Transforms, SortedByWarpingWeightRotationItemIndex, NormalizedWarpingWeightRotation);
+	}
+
+	if (WarpingBankingWeight > 1.f - UE_KINDA_SMALL_NUMBER)
+	{
+		return FindReferenceOrientationFullBanking(Transforms, SortedByWarpingWeightRotationItemIndex, NormalizedWarpingWeightRotation);
+	}
+
+	return FQuat::Slerp(
+		FindReferenceOrientationNoBanking(Transforms, SortedByWarpingWeightRotationItemIndex, NormalizedWarpingWeightRotation),
+		FindReferenceOrientationFullBanking(Transforms, SortedByWarpingWeightRotationItemIndex, NormalizedWarpingWeightRotation),
+		WarpingBankingWeight);
+}
+
+void Normalize(TArrayView<float> Array, int32 IndexToExclude = INDEX_NONE)
+{
+	const int32 Num = Array.Num();
+	check(Num >= 2);
+
+	float Sum = 0.f;
+	for (int32 Index = 0; Index < Num; ++Index)
+	{
+		if (Index != IndexToExclude)
+		{
+			Sum += Array[Index];
+		}
+	}
+
+	if (Sum > UE_KINDA_SMALL_NUMBER)
+	{
+		for (int32 Index = 0; Index < Num; ++Index)
+		{
+			if (Index != IndexToExclude)
+			{
+				Array[Index] /= Sum;
+			}
+		}
+	}
+	else
+	{
+		const float HomogeneousWeight = (IndexToExclude == INDEX_NONE) ? 1.f / Num : 1.f / (Num - 1);
+		for (int32 Index = 0; Index < Num; ++Index)
+		{
+			if (Index != IndexToExclude)
+			{
+				Array[Index] = HomogeneousWeight;
+			}
+		}
+	}
+}
+
+bool CalculateFullAlignedActorRootBoneTransform(const FTransform& ActorRootBoneTransform, const FTransform& AssetRootBoneTransform, 
+	const FVector& ActorsReferencePosition, const FVector& AssetReferencePosition,
+	float WarpingBankingWeight, FTransform& OutFullAlignedActorRootBoneTransform, const UWorld* DebugWorld = nullptr)
+{
+	const FVector ActorRootBonePosition = ActorRootBoneTransform.GetLocation();
+	const FVector AssetRootBonePosition = AssetRootBoneTransform.GetLocation();
+				
+	const FVector AssetReferencePositionLocal = AssetRootBoneTransform.InverseTransformPosition(AssetReferencePosition);
+	const FVector ActorsReferencePositionFromAsset = ActorRootBoneTransform.TransformPosition(AssetReferencePositionLocal);
+
+	const FVector ActorsReferenceDisplacementFromAsset = ActorsReferencePositionFromAsset - ActorRootBonePosition;
+	const FVector ActorsReferenceDisplacement = ActorsReferencePosition - ActorRootBonePosition;
+
+	const float ActorsReferenceDisplacementFromAssetLength = ActorsReferenceDisplacementFromAsset.Length();
+	const float ActorsReferenceDisplacementLength = ActorsReferenceDisplacement.Length();
+
+	if (ActorsReferenceDisplacementFromAssetLength > UE_SMALL_NUMBER && ActorsReferenceDisplacementLength > UE_SMALL_NUMBER)
+	{
+		const FVector ActorsReferenceDisplacementWithReducedBanking(ActorsReferenceDisplacement.X, ActorsReferenceDisplacement.Y,
+			FMath::Lerp(ActorsReferenceDisplacementFromAsset.Z, ActorsReferenceDisplacement.Z, WarpingBankingWeight));
+		const FQuat DeltaRotation = FQuat::FindBetweenVectors(ActorsReferenceDisplacementFromAsset, ActorsReferenceDisplacementWithReducedBanking);
+
+		OutFullAlignedActorRootBoneTransform.SetLocation(ActorsReferencePosition - ActorsReferenceDisplacement * (ActorsReferenceDisplacementFromAssetLength / ActorsReferenceDisplacementLength));
+		OutFullAlignedActorRootBoneTransform.SetRotation(ActorRootBoneTransform.GetRotation() * DeltaRotation);
+		OutFullAlignedActorRootBoneTransform.SetScale3D(ActorRootBoneTransform.GetScale3D());
+		
+#if ENABLE_ANIM_DEBUG && WITH_EDITOR
+		DrawDebugLine(DebugWorld, ActorsReferencePosition, ActorRootBonePosition, FColorList::Blue, false, 0.f, SDPG_Foreground, 0.f);
+		DrawDebugLine(DebugWorld, ActorsReferencePositionFromAsset, ActorRootBonePosition, FColorList::Red, false, 0.f, SDPG_Foreground, 0.f);
+		const FVector ReconstructedForDebug = DeltaRotation.RotateVector(ActorsReferenceDisplacementFromAsset) + ActorRootBonePosition;
+		DrawDebugCircle(DebugWorld, FMatrix(FVector::ZAxisVector, FVector::XAxisVector, FVector::YAxisVector, ReconstructedForDebug), 5.f, 16, FColorList::Red, false, 0.f, SDPG_Foreground, 0.f, false);
+		
+		// asset actor drawing 
+		DrawDebugCircle(DebugWorld, FMatrix(FVector::ZAxisVector, FVector::XAxisVector, FVector::YAxisVector, AssetRootBonePosition), 5.f, 16, FColorList::Green, false, 0.f, SDPG_Foreground, 0.f, false);
+		DrawDebugLine(DebugWorld, AssetRootBonePosition, AssetRootBonePosition + AssetRootBoneTransform.GetRotation().RotateVector(FVector::XAxisVector) * 10, FColorList::Green, false, 0.f, SDPG_Foreground, 0.f);
+		
+		DrawDebugCircle(DebugWorld, FMatrix(FVector::ZAxisVector, FVector::XAxisVector, FVector::YAxisVector, ActorRootBonePosition), 5.f, 16, FColorList::Green, false, 0.f, SDPG_Foreground, 0.f, false);
+		DrawDebugLine(DebugWorld, ActorRootBonePosition, ActorRootBonePosition + ActorRootBoneTransform.GetRotation().RotateVector(FVector::XAxisVector) * 10, FColorList::Green, false, 0.f, SDPG_Foreground, 0.f);
+#endif // ENABLE_ANIM_DEBUG && WITH_EDITOR
+
+		return true;
+	}
+
+	OutFullAlignedActorRootBoneTransform = ActorRootBoneTransform;
+	return false;
+}
+
+} // namespace UE::PoseSearch
+
 bool UPoseSearchInteractionAsset::IsLooping() const
 {
 	float CommonPlayLength = -1.f;
@@ -110,89 +274,6 @@ float UPoseSearchInteractionAsset::GetPlayLength(const FVector& BlendParameters)
 	return MaxPlayLength;
 }
 
-FQuat UPoseSearchInteractionAsset::FindReferenceOrientationNoBanking(const TArrayView<const FTransform> Transforms, const TArrayView<int32> SortedByWarpingWeightRotationItemIndex, const TArrayView<float> NormalizedWarpingWeightRotation) const
-{
-	check(Items.Num() > 0);
-	check(Items.Num() == Transforms.Num());
-	check(Items.Num() == SortedByWarpingWeightRotationItemIndex.Num());
-	check(Items.Num() == NormalizedWarpingWeightRotation.Num());
-
-	// @todo: use a Slerp or a proper FastLerp
-	// for now we don't account for the shortest path while FastLerping those queternion together
-	FQuat WeightedQuaternion = FQuat::Identity * 0.f;
-	for (int32 ItemIndex : SortedByWarpingWeightRotationItemIndex)
-	{
-		WeightedQuaternion += Transforms[ItemIndex].GetRotation() * NormalizedWarpingWeightRotation[ItemIndex];
-	}
-	WeightedQuaternion.Normalize();
-	return WeightedQuaternion;
-}
-
-FQuat UPoseSearchInteractionAsset::FindReferenceOrientationFullBanking(const TArrayView<const FTransform> Transforms, const TArrayView<int32> SortedByWarpingWeightRotationItemIndex, const TArrayView<float> NormalizedWarpingWeightRotation) const
-{
-	const int32 ItemsNum = Items.Num();
-
-	check(ItemsNum > 0);
-	check(ItemsNum == Transforms.Num());
-	check(ItemsNum == SortedByWarpingWeightRotationItemIndex.Num());
-	check(ItemsNum == NormalizedWarpingWeightRotation.Num());
-
-	if (ItemsNum > 1)
-	{
-		const int32 LastItemIndex = ItemsNum - 1;
-		FVector OtherItemsPositionsSum = Transforms[SortedByWarpingWeightRotationItemIndex[0]].GetTranslation();
-		for (int32 ItemIndex = 1; ItemIndex < LastItemIndex; ++ItemIndex)
-		{
-			OtherItemsPositionsSum += Transforms[SortedByWarpingWeightRotationItemIndex[ItemIndex]].GetTranslation();
-		}
-
-		const FVector OtherItemsPositionAverage = OtherItemsPositionsSum / LastItemIndex;
-		const FVector DeltaPosition = OtherItemsPositionAverage - Transforms[SortedByWarpingWeightRotationItemIndex[LastItemIndex]].GetTranslation();
-
-		if (!DeltaPosition.IsNearlyZero())
-		{
-			return DeltaPosition.ToOrientationQuat();
-		}
-	}
-
-	return FindReferenceOrientationNoBanking(Transforms, SortedByWarpingWeightRotationItemIndex, NormalizedWarpingWeightRotation);
-}
-
-FQuat UPoseSearchInteractionAsset::FindReferenceOrientation(const TArrayView<const FTransform> Transforms, const TArrayView<int32> SortedByWarpingWeightRotationItemIndex, const TArrayView<float> NormalizedWarpingWeightRotation) const
-{
-	if (WarpingBankingWeight < UE_KINDA_SMALL_NUMBER)
-	{
-		return FindReferenceOrientationNoBanking(Transforms, SortedByWarpingWeightRotationItemIndex, NormalizedWarpingWeightRotation);
-	}
-
-	if (WarpingBankingWeight > 1.f - UE_KINDA_SMALL_NUMBER)
-	{
-		return FindReferenceOrientationFullBanking(Transforms, SortedByWarpingWeightRotationItemIndex, NormalizedWarpingWeightRotation);
-	}
-
-	return FQuat::Slerp(
-		FindReferenceOrientationNoBanking(Transforms, SortedByWarpingWeightRotationItemIndex, NormalizedWarpingWeightRotation),
-		FindReferenceOrientationFullBanking(Transforms, SortedByWarpingWeightRotationItemIndex, NormalizedWarpingWeightRotation),
-		WarpingBankingWeight);
-}
-
-FVector UPoseSearchInteractionAsset::FindReferencePosition(const TArrayView<const FTransform> Transforms, const TArrayView<float> NormalizedWarpingWeightTranslation) const
-{
-	const int32 ItemsNum = Items.Num();
-	
-	check(ItemsNum > 0);
-	check(Transforms.Num() == ItemsNum);
-	check(Transforms.Num() == NormalizedWarpingWeightTranslation.Num());
-
-	FVector PositionsSum = FVector::ZeroVector;
-	for (int32 ItemIndex = 0; ItemIndex < ItemsNum; ++ItemIndex)
-	{
-		PositionsSum += Transforms[ItemIndex].GetTranslation() * NormalizedWarpingWeightTranslation[ItemIndex];
-	}
-
-	return PositionsSum;
-}
-
 UAnimationAsset* UPoseSearchInteractionAsset::GetAnimationAsset(const UE::PoseSearch::FRole& Role) const
 {
 	for (const FPoseSearchInteractionAssetItem& Item : Items)
@@ -254,34 +335,15 @@ USkeletalMesh* UPoseSearchInteractionAsset::GetPreviewMesh(const UE::PoseSearch:
 #endif // WITH_EDITOR
 
 void UPoseSearchInteractionAsset::CalculateWarpTransforms(float Time, float TimeOffset, const TConstArrayView<const FTransform> ActorRootBoneTransforms, TArrayView<FTransform> FullAlignedActorRootBoneTransforms,
-	const TConstArrayView<const UMirrorDataTable*> MirrorDataTables, const TConstArrayView<bool> RelevantRoleIndexes, const TConstArrayView<TObjectPtr<const UObject>> DebugAnimContexts) const
+	const TConstArrayView<const UMirrorDataTable*> MirrorDataTables, const TConstArrayView<TObjectPtr<const UObject>> DebugAnimContexts) const
 {
 	using namespace UE::PoseSearch;
 
 	check(ActorRootBoneTransforms.Num() == GetNumRoles());
 	check(FullAlignedActorRootBoneTransforms.Num() == GetNumRoles());
-	check(RelevantRoleIndexes.IsEmpty() || RelevantRoleIndexes.Num() == GetNumRoles());
 	
 	const int32 ItemsNum = Items.Num();
-
-	int32 RelevantItemsNum;
-	if (RelevantRoleIndexes.IsEmpty())
-	{
-		RelevantItemsNum = ItemsNum;
-	}
-	else
-	{
-		RelevantItemsNum = 0;
-		for (bool bIsRelevant : RelevantRoleIndexes)
-		{
-			if (bIsRelevant)
-			{
-				++RelevantItemsNum;
-			}
-		}
-	}
-
-	if (RelevantItemsNum < 2)
+	if (ItemsNum < 2)
 	{
 		for (int32 ItemIndex = 0; ItemIndex < ItemsNum; ++ItemIndex)
 		{
@@ -290,7 +352,6 @@ void UPoseSearchInteractionAsset::CalculateWarpTransforms(float Time, float Time
 	}
 	else
 	{
-		// we have at least two relevant items!
 		TArray<FTransform, TInlineAllocator<PreallocatedRolesNum>> AssetRootBoneTransforms;
 		AssetRootBoneTransforms.SetNum(ItemsNum);
 
@@ -348,107 +409,88 @@ void UPoseSearchInteractionAsset::CalculateWarpTransforms(float Time, float Time
 #endif // ENABLE_ANIM_DEBUG && WITH_EDITOR
 		}
 
-		TArray<int32, TInlineAllocator<PreallocatedRolesNum>> SortedByWarpingWeightRotationItemIndex;
-		TArray<float, TInlineAllocator<PreallocatedRolesNum>> NormalizedWarpingWeightRotation;
 		TArray<float, TInlineAllocator<PreallocatedRolesNum>> NormalizedWarpingWeightTranslation;
-		SortedByWarpingWeightRotationItemIndex.SetNum(ItemsNum);
-		NormalizedWarpingWeightRotation.SetNum(ItemsNum);
-		NormalizedWarpingWeightTranslation.SetNum(ItemsNum);
-
-		float WarpingWeightTranslationSum = 0.f;
-		float WarpingWeightRotationSum = 0.f;
-
-		if (RelevantRoleIndexes.IsEmpty())
+		NormalizedWarpingWeightTranslation.SetNumUninitialized(ItemsNum);
+		for (int32 ItemIndex = 0; ItemIndex < ItemsNum; ++ItemIndex)
 		{
-			for (int32 ItemIndex = 0; ItemIndex < ItemsNum; ++ItemIndex)
-			{
-				SortedByWarpingWeightRotationItemIndex[ItemIndex] = ItemIndex;
-				WarpingWeightTranslationSum += Items[ItemIndex].WarpingWeightTranslation;
-				WarpingWeightRotationSum += Items[ItemIndex].WarpingWeightRotation;
-			}
+			NormalizedWarpingWeightTranslation[ItemIndex] = Items[ItemIndex].WarpingWeightTranslation;
 		}
-		else
+		Normalize(NormalizedWarpingWeightTranslation);
+		
+		const FVector AssetReferencePosition = FindReferencePosition(AssetRootBoneTransforms, NormalizedWarpingWeightTranslation);
+		const FVector ActorsReferencePosition = FindReferencePosition(ActorRootBoneTransforms, NormalizedWarpingWeightTranslation);
+
+		const UWorld* DebugWorld = nullptr;
+#if ENABLE_ANIM_DEBUG && WITH_EDITOR
+		for (const TObjectPtr<const UObject>& DebugAnimContextPtr : DebugAnimContexts)
 		{
-			for (int32 ItemIndex = 0; ItemIndex < ItemsNum; ++ItemIndex)
+			if (const UObject* Outer = DebugAnimContextPtr->GetOuter())
 			{
-				SortedByWarpingWeightRotationItemIndex[ItemIndex] = ItemIndex;
-				if (RelevantRoleIndexes[ItemIndex])
+				DebugWorld = Outer->GetWorld();
+				if (DebugWorld)
 				{
-					WarpingWeightTranslationSum += Items[ItemIndex].WarpingWeightTranslation;
-					WarpingWeightRotationSum += Items[ItemIndex].WarpingWeightRotation;
+					break;
 				}
 			}
 		}
+#endif // ENABLE_ANIM_DEBUG && WITH_EDITOR
 
-		const float NormalizedHomogeneousWeight = 1.f / RelevantItemsNum;
-		if (RelevantRoleIndexes.IsEmpty())
+		if (MinimalTranslationWeight > 0.f)
 		{
-			if (WarpingWeightTranslationSum > UE_KINDA_SMALL_NUMBER)
-			{
-				for (int32 ItemIndex = 0; ItemIndex < ItemsNum; ++ItemIndex)
-				{
-					NormalizedWarpingWeightTranslation[ItemIndex] = Items[ItemIndex].WarpingWeightTranslation / WarpingWeightTranslationSum;
-				}
-			}
-			else
+			// finding FullAlignedActorRootBoneTransforms with the minimal possible translations
+			for (int32 ItemIndex = 0; ItemIndex < ItemsNum; ++ItemIndex)
 			{
-				for (int32 ItemIndex = 0; ItemIndex < ItemsNum; ++ItemIndex)
+				if (!CalculateFullAlignedActorRootBoneTransform(ActorRootBoneTransforms[ItemIndex], AssetRootBoneTransforms[ItemIndex], 
+					ActorsReferencePosition, AssetReferencePosition, WarpingBankingWeight, FullAlignedActorRootBoneTransforms[ItemIndex], DebugWorld))
 				{
-					NormalizedWarpingWeightTranslation[ItemIndex] = NormalizedHomogeneousWeight;
-				}
-			}
+					// finding the references without the ItemIndex, since the current actors and asset references are too close to ActorRootBoneTransforms[ItemIndex] and AssetRootBoneTransforms[ItemIndex]
+					TArray<float, TInlineAllocator<PreallocatedRolesNum>> NormalizedWarpingWeightTranslationWithoutItemIndex;
+					NormalizedWarpingWeightTranslationWithoutItemIndex = NormalizedWarpingWeightTranslation;
+					NormalizedWarpingWeightTranslationWithoutItemIndex[ItemIndex] = 0.f;
 
-			if (WarpingWeightRotationSum > UE_KINDA_SMALL_NUMBER)
-			{
-				SortedByWarpingWeightRotationItemIndex.Sort([this](const int32 A, const int32 B)
-					{
-						return Items[A].WarpingWeightRotation < Items[B].WarpingWeightRotation;
-					});
+					Normalize(NormalizedWarpingWeightTranslationWithoutItemIndex, ItemIndex);
 
-				for (int32 ItemIndex = 0; ItemIndex < ItemsNum; ++ItemIndex)
-				{
-					NormalizedWarpingWeightRotation[ItemIndex] = Items[ItemIndex].WarpingWeightRotation / WarpingWeightRotationSum;
-				}
-			}
-			else
-			{
-				for (int32 ItemIndex = 0; ItemIndex < ItemsNum; ++ItemIndex)
-				{
-					NormalizedWarpingWeightRotation[ItemIndex] = NormalizedHomogeneousWeight;
-				}
-			}
-		}
-		else
-		{
-			if (WarpingWeightTranslationSum > UE_KINDA_SMALL_NUMBER)
-			{
-				for (int32 ItemIndex = 0; ItemIndex < ItemsNum; ++ItemIndex)
-				{
-					if (RelevantRoleIndexes[ItemIndex])
+					const FVector AssetReferencePositionWithoutItemIndex = FindReferencePosition(AssetRootBoneTransforms, NormalizedWarpingWeightTranslationWithoutItemIndex);
+					const FVector ActorsReferencePositionWithoutItemIndex = FindReferencePosition(ActorRootBoneTransforms, NormalizedWarpingWeightTranslationWithoutItemIndex);
+
+#if ENABLE_ANIM_DEBUG && WITH_EDITOR
+					DrawDebugCircle(DebugWorld, FMatrix(FVector::ZAxisVector, FVector::XAxisVector, FVector::YAxisVector, AssetReferencePositionWithoutItemIndex), 8.f, 16, FColorList::Black, false, 0.f, SDPG_Foreground, 0.f, false);
+					DrawDebugCircle(DebugWorld, FMatrix(FVector::ZAxisVector, FVector::XAxisVector, FVector::YAxisVector, ActorsReferencePositionWithoutItemIndex), 8.f, 16, FColorList::Black, false, 0.f, SDPG_Foreground, 0.f, false);
+#endif // ENABLE_ANIM_DEBUG && WITH_EDITOR
+
+					if (CalculateFullAlignedActorRootBoneTransform(ActorRootBoneTransforms[ItemIndex], AssetRootBoneTransforms[ItemIndex],
+						ActorsReferencePositionWithoutItemIndex, AssetReferencePositionWithoutItemIndex, WarpingBankingWeight, FullAlignedActorRootBoneTransforms[ItemIndex], DebugWorld))
 					{
-						NormalizedWarpingWeightTranslation[ItemIndex] = Items[ItemIndex].WarpingWeightTranslation / WarpingWeightTranslationSum;
+						FullAlignedActorRootBoneTransforms[ItemIndex].SetLocation(ActorRootBoneTransforms[ItemIndex].GetLocation());
 					}
 					else
 					{
-						NormalizedWarpingWeightTranslation[ItemIndex] = 0.f;
+						UE_LOG(LogPoseSearch, Error, TEXT("UPoseSearchInteractionAsset::CalculateWarpTransforms actors are too close to each other to be able to calculate warp transforms ['%s' at time %f]"), *GetName(), Time);
 					}
 				}
 			}
-			else
+
+#if ENABLE_ANIM_DEBUG && WITH_EDITOR
+			DrawDebugCircle(DebugWorld, FMatrix(FVector::ZAxisVector, FVector::XAxisVector, FVector::YAxisVector, AssetReferencePosition), 5.f, 16, FColorList::Blue, false, 0.f, SDPG_Foreground, 0.f, false);
+			DrawDebugCircle(DebugWorld, FMatrix(FVector::ZAxisVector, FVector::XAxisVector, FVector::YAxisVector, ActorsReferencePosition), 5.f, 16, FColorList::Blue, false, 0.f, SDPG_Foreground, 0.f, false);
+#endif // ENABLE_ANIM_DEBUG && WITH_EDITOR
+		}
+		
+		if (MinimalTranslationWeight < 1.f)
+		{
+			TArray<int32, TInlineAllocator<PreallocatedRolesNum>> SortedByWarpingWeightRotationItemIndex;
+			TArray<float, TInlineAllocator<PreallocatedRolesNum>> NormalizedWarpingWeightRotation;
+			SortedByWarpingWeightRotationItemIndex.SetNum(ItemsNum);
+			NormalizedWarpingWeightRotation.SetNum(ItemsNum);
+		
+			float WarpingWeightRotationSum = 0.f;
+			for (int32 ItemIndex = 0; ItemIndex < ItemsNum; ++ItemIndex)
 			{
-				for (int32 ItemIndex = 0; ItemIndex < ItemsNum; ++ItemIndex)
-				{
-					if (RelevantRoleIndexes[ItemIndex])
-					{
-						NormalizedWarpingWeightTranslation[ItemIndex] = NormalizedHomogeneousWeight;
-					}
-					else
-					{
-						NormalizedWarpingWeightTranslation[ItemIndex] = 0.f;
-					}
-				}
+				SortedByWarpingWeightRotationItemIndex[ItemIndex] = ItemIndex;
+				WarpingWeightRotationSum += Items[ItemIndex].WarpingWeightRotation;
 			}
 
+			const float NormalizedHomogeneousWeight = 1.f / ItemsNum;
 			if (WarpingWeightRotationSum > UE_KINDA_SMALL_NUMBER)
 			{
 				SortedByWarpingWeightRotationItemIndex.Sort([this](const int32 A, const int32 B)
@@ -458,107 +500,19 @@ void UPoseSearchInteractionAsset::CalculateWarpTransforms(float Time, float Time
 
 				for (int32 ItemIndex = 0; ItemIndex < ItemsNum; ++ItemIndex)
 				{
-					if (RelevantRoleIndexes[ItemIndex])
-					{
-						NormalizedWarpingWeightRotation[ItemIndex] = Items[ItemIndex].WarpingWeightRotation / WarpingWeightRotationSum;
-					}
-					else
-					{
-						NormalizedWarpingWeightRotation[ItemIndex] = 0.f;
-					}
+					NormalizedWarpingWeightRotation[ItemIndex] = Items[ItemIndex].WarpingWeightRotation / WarpingWeightRotationSum;
 				}
 			}
 			else
 			{
 				for (int32 ItemIndex = 0; ItemIndex < ItemsNum; ++ItemIndex)
 				{
-					if (RelevantRoleIndexes[ItemIndex])
-					{
-						NormalizedWarpingWeightRotation[ItemIndex] = NormalizedHomogeneousWeight;
-					}
-					else
-					{
-						NormalizedWarpingWeightRotation[ItemIndex] = 0.f;
-					}
-				}
-			}
-		}
-		const FVector AssetReferencePosition = FindReferencePosition(AssetRootBoneTransforms, NormalizedWarpingWeightTranslation);
-		const FVector ActorsReferencePosition = FindReferencePosition(ActorRootBoneTransforms, NormalizedWarpingWeightTranslation);
-
-#if ENABLE_ANIM_DEBUG && WITH_EDITOR
-		const UWorld* DebugWorld = nullptr;
-		for (const TObjectPtr<const UObject>& DebugAnimContextPtr : DebugAnimContexts)
-		{
-			if (const UObject* Outer = DebugAnimContextPtr->GetOuter())
-			{
-				DebugWorld = Outer->GetWorld();
-				if (DebugWorld)
-				{
-					break;
-				}
-			}
-		}
-#endif // ENABLE_ANIM_DEBUG && WITH_EDITOR
-
-		if (MinimalTranslationWeight > 0.f)
-		{
-			// finding FullAlignedActorRootBoneTransforms with the minimal possible translations
-			for (int32 ItemIndex = 0; ItemIndex < ItemsNum; ++ItemIndex)
-			{
-				const FVector ActorRootBonePosition = ActorRootBoneTransforms[ItemIndex].GetLocation();
-				const FVector AssetRootBonePosition = AssetRootBoneTransforms[ItemIndex].GetLocation();
-				
-				const FVector AssetReferencePositionLocal = AssetRootBoneTransforms[ItemIndex].InverseTransformPosition(AssetReferencePosition);
-				const FVector ActorsReferencePositionFromAsset = ActorRootBoneTransforms[ItemIndex].TransformPosition(AssetReferencePositionLocal);
-
-				const FVector ActorsReferenceDisplacementFromAsset = ActorsReferencePositionFromAsset - ActorRootBonePosition;
-				const FVector ActorsReferenceDisplacement = ActorsReferencePosition - ActorRootBonePosition;
-
-				const float ActorsReferenceDisplacementFromAssetLength = ActorsReferenceDisplacementFromAsset.Length();
-				const float ActorsReferenceDisplacementLength = ActorsReferenceDisplacement.Length();
-
-				if (ActorsReferenceDisplacementFromAssetLength > UE_SMALL_NUMBER && ActorsReferenceDisplacementLength > UE_SMALL_NUMBER)
-				{
-					const FVector ActorsReferenceDisplacementWithReducedBanking(ActorsReferenceDisplacement.X, ActorsReferenceDisplacement.Y, 
-						FMath::Lerp(ActorsReferenceDisplacementFromAsset.Z, ActorsReferenceDisplacement.Z, WarpingBankingWeight));
-					const FQuat DeltaRotation = FQuat::FindBetweenVectors(ActorsReferenceDisplacementFromAsset, ActorsReferenceDisplacementWithReducedBanking);
-
-					FullAlignedActorRootBoneTransforms[ItemIndex].SetLocation(ActorsReferencePosition - ActorsReferenceDisplacement * (ActorsReferenceDisplacementFromAssetLength / ActorsReferenceDisplacementLength));
-					FullAlignedActorRootBoneTransforms[ItemIndex].SetRotation(ActorRootBoneTransforms[ItemIndex].GetRotation() * DeltaRotation);
-					FullAlignedActorRootBoneTransforms[ItemIndex].SetScale3D(ActorRootBoneTransforms[ItemIndex].GetScale3D());
-
-#if ENABLE_ANIM_DEBUG && WITH_EDITOR
-					DrawDebugLine(DebugWorld, ActorsReferencePosition, ActorRootBonePosition, FColorList::Blue, false, 0.f, SDPG_Foreground, 0.f);
-					DrawDebugLine(DebugWorld, ActorsReferencePositionFromAsset, ActorRootBonePosition, FColorList::Red, false, 0.f, SDPG_Foreground, 0.f);
-					const FVector ReconstructedForDebug = DeltaRotation.RotateVector(ActorsReferenceDisplacementFromAsset) + ActorRootBonePosition;
-					DrawDebugCircle(DebugWorld, FMatrix(FVector::ZAxisVector, FVector::XAxisVector, FVector::YAxisVector, ReconstructedForDebug), 5.f, 16, FColorList::Red, false, 0.f, SDPG_Foreground, 0.f, false);
-
-					// asset actor drawing 
-					DrawDebugCircle(DebugWorld, FMatrix(FVector::ZAxisVector, FVector::XAxisVector, FVector::YAxisVector, AssetRootBonePosition), 5.f, 16, FColorList::Green, false, 0.f, SDPG_Foreground, 0.f, false);
-					DrawDebugLine(DebugWorld, AssetRootBonePosition, AssetRootBonePosition + AssetRootBoneTransforms[ItemIndex].GetRotation().RotateVector(FVector::XAxisVector) * 10, FColorList::Green, false, 0.f, SDPG_Foreground, 0.f);
-
-					DrawDebugCircle(DebugWorld, FMatrix(FVector::ZAxisVector, FVector::XAxisVector, FVector::YAxisVector, ActorRootBonePosition), 5.f, 16, FColorList::Green, false, 0.f, SDPG_Foreground, 0.f, false);
-					DrawDebugLine(DebugWorld, ActorRootBonePosition, ActorRootBonePosition + ActorRootBoneTransforms[ItemIndex].GetRotation().RotateVector(FVector::XAxisVector) * 10, FColorList::Green, false, 0.f, SDPG_Foreground, 0.f);
-#endif // ENABLE_ANIM_DEBUG && WITH_EDITOR
-				}
-				else
-				{
-					UE_LOG(LogPoseSearch, Error, TEXT("UPoseSearchInteractionAsset::CalculateWarpTransforms actors are too close to each other to be able to calculate warp transforms: distance in the world %f, distance in the asset %f ['%s' at time %f]"), ActorsReferenceDisplacementLength, ActorsReferenceDisplacementFromAssetLength, *GetName(), Time);
-					FullAlignedActorRootBoneTransforms[ItemIndex] = ActorRootBoneTransforms[ItemIndex];
+					NormalizedWarpingWeightRotation[ItemIndex] = NormalizedHomogeneousWeight;
 				}
 			}
 
-#if ENABLE_ANIM_DEBUG && WITH_EDITOR
-			DrawDebugCircle(DebugWorld, FMatrix(FVector::ZAxisVector, FVector::XAxisVector, FVector::YAxisVector, AssetReferencePosition), 5.f, 16, FColorList::Blue, false, 0.f, SDPG_Foreground, 0.f, false);
-			DrawDebugCircle(DebugWorld, FMatrix(FVector::ZAxisVector, FVector::XAxisVector, FVector::YAxisVector, ActorsReferencePosition), 5.f, 16, FColorList::Blue, false, 0.f, SDPG_Foreground, 0.f, false);
-#endif // ENABLE_ANIM_DEBUG && WITH_EDITOR
-		}
-		
-		if (MinimalTranslationWeight < 1.f)
-		{
-			const FQuat AssetReferenceOrientation = FindReferenceOrientation(AssetRootBoneTransforms, SortedByWarpingWeightRotationItemIndex, NormalizedWarpingWeightRotation);
-			const FQuat ActorsReferenceOrientation = FindReferenceOrientation(ActorRootBoneTransforms, SortedByWarpingWeightRotationItemIndex, NormalizedWarpingWeightRotation);
+			const FQuat AssetReferenceOrientation = FindReferenceOrientation(AssetRootBoneTransforms, SortedByWarpingWeightRotationItemIndex, NormalizedWarpingWeightRotation, WarpingBankingWeight);
+			const FQuat ActorsReferenceOrientation = FindReferenceOrientation(ActorRootBoneTransforms, SortedByWarpingWeightRotationItemIndex, NormalizedWarpingWeightRotation, WarpingBankingWeight);
 
 			FQuat WeightedActorsReferenceOrientation = ActorsReferenceOrientation;
 			if (WarpingWeightRotationSum > UE_KINDA_SMALL_NUMBER)
diff --git a/Engine/Plugins/Animation/PoseSearch/Source/Runtime/Private/PoseSearchInteractionUtils.cpp b/Engine/Plugins/Animation/PoseSearch/Source/Runtime/Private/PoseSearchInteractionUtils.cpp
--- a/Engine/Plugins/Animation/PoseSearch/Source/Runtime/Private/PoseSearchInteractionUtils.cpp
+++ b/Engine/Plugins/Animation/PoseSearch/Source/Runtime/Private/PoseSearchInteractionUtils.cpp
@@ -75,13 +75,7 @@ void CalculateFullAlignedTransforms(const FPoseSearchBlueprintResult& Result, fl
 		}
 	}
 
-	TArray<bool, TInlineAllocator<PreallocatedRolesNum, TMemStackAllocator<>>> RelevantRoleIndexes;
-	RelevantRoleIndexes.SetNum(NumRoles);
-	for (int32 Index = 0; Index < NumRoles; ++Index)
-	{
-		RelevantRoleIndexes[Index] = Result.AnimContexts[Index] != nullptr;
-	}
-	MultiAnimAsset->CalculateWarpTransforms(Time, TimeOffset, ActorTransforms, OutFullAlignedTransforms, MirrorDataTables, RelevantRoleIndexes);
+	MultiAnimAsset->CalculateWarpTransforms(Time, TimeOffset, ActorTransforms, OutFullAlignedTransforms, MirrorDataTables);
 }
 
 FTransform CalculateDeltaAlignment(const FTransform& MeshWithoutOffset, const FTransform& MeshWithOffset, const FTransform& FullAlignedTransform, float WarpingRotationRatio, float WarpingTranslationRatio)
diff --git a/Engine/Plugins/Animation/PoseSearch/Source/Runtime/Public/PoseSearch/MultiAnimAsset.h b/Engine/Plugins/Animation/PoseSearch/Source/Runtime/Public/PoseSearch/MultiAnimAsset.h
--- a/Engine/Plugins/Animation/PoseSearch/Source/Runtime/Public/PoseSearch/MultiAnimAsset.h
+++ b/Engine/Plugins/Animation/PoseSearch/Source/Runtime/Public/PoseSearch/MultiAnimAsset.h
@@ -33,13 +33,13 @@ public:
 	[[nodiscard]] virtual FTransform GetOrigin(const FName& Role) const PURE_VIRTUAL(UMultiAnimAsset::GetOrigin, return FTransform::Identity;);
 
 	// calculates the full aligned root bone transforms (FullAlignedActorRootBoneTransforms) at time "Time + TimeOffset" given the current actors root bone transforms (ActorRootBoneTransforms) at "Time"
-	// the mirror data tables, the RelevantRoleIndexes (array of what roles are present or "relevant"). DebugAnimContexts can be used to draw debug lines
+	// the mirror data tables. DebugAnimContexts can be used to draw debug lines
 	// For example if the actors are playing the associated GetAnimationAsset(Role) at time Time, by providing the ActorRootBoneTransforms and an TimeOffset of zero, CalculateWarpTransforms will return the
 	// FullAlignedActorRootBoneTransforms should be NOW to be in full alignment. Any TimeOffset > 0 will have the method returning a future FullAlignedActorRootBoneTransforms. This is useful in synergy with
 	// Motion Warping to provide a transform to use as full aligment at the ned of the notify state time
-	virtual void CalculateWarpTransforms(float Time, float TimeOffset, const TConstArrayView<const FTransform> ActorRootBoneTransforms, TArrayView<FTransform> FullAlignedActorRootBoneTransforms,
-		const TConstArrayView<const UMirrorDataTable*> MirrorDataTables = TConstArrayView<const UMirrorDataTable*>(), const TConstArrayView<bool> RelevantRoleIndexes = TConstArrayView<bool>(),
-		const TConstArrayView<TObjectPtr<const UObject>> DebugAnimContexts = TConstArrayView<TObjectPtr<const UObject>>()) const PURE_VIRTUAL(UMultiAnimAsset::CalculateWarpTransforms, );
+	virtual void CalculateWarpTransforms(float Time, float TimeOffset, const TConstArrayView<const FTransform> ActorRootBoneTransforms, TArrayView<FTransform> FullAlignedActorRootBoneTransforms, 
+		const TConstArrayView<const UMirrorDataTable*> MirrorDataTables = TConstArrayView<const UMirrorDataTable*>(), const TConstArrayView<TObjectPtr<const UObject>> DebugAnimContexts = TConstArrayView<TObjectPtr<const UObject>>()) const 
+		PURE_VIRTUAL(UMultiAnimAsset::CalculateWarpTransforms, );
 
 	UFUNCTION(BlueprintPure, Category = "Animation", meta=(BlueprintThreadSafe, DisplayName = "Get Animation Asset"))
 	UAnimationAsset* BP_GetAnimationAsset(const FName& Role) const { return GetAnimationAsset(Role); }
diff --git a/Engine/Plugins/Animation/PoseSearch/Source/Runtime/Public/PoseSearch/PoseSearchInteractionAsset.h b/Engine/Plugins/Animation/PoseSearch/Source/Runtime/Public/PoseSearch/PoseSearchInteractionAsset.h
--- a/Engine/Plugins/Animation/PoseSearch/Source/Runtime/Public/PoseSearch/PoseSearchInteractionAsset.h
+++ b/Engine/Plugins/Animation/PoseSearch/Source/Runtime/Public/PoseSearch/PoseSearchInteractionAsset.h
@@ -98,14 +98,7 @@ public:
 #endif // WITH_EDITOR
 
 	UE_API virtual void CalculateWarpTransforms(float Time, float TimeOffset, const TConstArrayView<const FTransform> ActorRootBoneTransforms, TArrayView<FTransform> FullAlignedActorRootBoneTransforms,
-		const TConstArrayView<const UMirrorDataTable*> MirrorDataTables = TConstArrayView<const UMirrorDataTable*>(), const TConstArrayView<bool> RelevantRoleIndexes = TConstArrayView<bool>(),
-		const TConstArrayView<TObjectPtr<const UObject>> DebugAnimContexts = TConstArrayView<TObjectPtr<const UObject>>()) const override;
-	
-private:
-	FVector FindReferencePosition(const TArrayView<const FTransform> Transforms, const TArrayView<float> NormalizedWarpingWeightTranslation) const;
-	FQuat FindReferenceOrientation(const TArrayView<const FTransform> Transforms, const TArrayView<int32> SortedByWarpingWeightRotationItemIndex, const TArrayView<float> NormalizedWarpingWeightRotation) const;
-	FQuat FindReferenceOrientationNoBanking(const TArrayView<const FTransform> Transforms, const TArrayView<int32> SortedByWarpingWeightRotationItemIndex, const TArrayView<float> NormalizedWarpingWeightRotation) const;
-	FQuat FindReferenceOrientationFullBanking(const TArrayView<const FTransform> Transforms, const TArrayView<int32> SortedByWarpingWeightRotationItemIndex, const TArrayView<float> NormalizedWarpingWeightRotation) const;
+		const TConstArrayView<const UMirrorDataTable*> MirrorDataTables = TConstArrayView<const UMirrorDataTable*>(), const TConstArrayView<TObjectPtr<const UObject>> DebugAnimContexts = TConstArrayView<TObjectPtr<const UObject>>()) const override;
 };
 
 #undef UE_API
