PoseSearch - firx for PoseSearchInteractionAsset.cpp(348): warning V1008: Consider inspecting the 'for' operator. No more than one iteration of the loop will be performed.

#Synced-CL 49824781

diff --git a/Engine/Plugins/Animation/PoseSearch/Source/Runtime/Private/PoseSearchInteractionAsset.cpp b/Engine/Plugins/Animation/PoseSearch/Source/Runtime/Private/PoseSearchInteractionAsset.cpp
--- a/Engine/Plugins/Animation/PoseSearch/Source/Runtime/Private/PoseSearchInteractionAsset.cpp
+++ b/Engine/Plugins/Animation/PoseSearch/Source/Runtime/Private/PoseSearchInteractionAsset.cpp
@@ -343,15 +343,21 @@ void UPoseSearchInteractionAsset::CalculateWarpTransforms(float Time, float Time
 	check(FullAlignedActorRootBoneTransforms.Num() == GetNumRoles());
 	
 	const int32 ItemsNum = Items.Num();
-	if (ItemsNum < 2)
+	switch (ItemsNum)
 	{
-		for (int32 ItemIndex = 0; ItemIndex < ItemsNum; ++ItemIndex)
-		{
-			FullAlignedActorRootBoneTransforms[ItemIndex] = ActorRootBoneTransforms[ItemIndex];
-		}
+	case 0:
+	{
+		break;
 	}
-	else
+	case 1:
+	{
+		FullAlignedActorRootBoneTransforms[0] = ActorRootBoneTransforms[0];
+		break;
+	}
+	default:
 	{
+		check(ItemsNum > 1);
+
 		TArray<FTransform, TInlineAllocator<PreallocatedRolesNum>> AssetRootBoneTransforms;
 		AssetRootBoneTransforms.SetNum(ItemsNum);
 
@@ -416,7 +422,7 @@ void UPoseSearchInteractionAsset::CalculateWarpTransforms(float Time, float Time
 			NormalizedWarpingWeightTranslation[ItemIndex] = Items[ItemIndex].WarpingWeightTranslation;
 		}
 		Normalize(NormalizedWarpingWeightTranslation);
-		
+
 		const FVector AssetReferencePosition = FindReferencePosition(AssetRootBoneTransforms, NormalizedWarpingWeightTranslation);
 		const FVector ActorsReferencePosition = FindReferencePosition(ActorRootBoneTransforms, NormalizedWarpingWeightTranslation);
 
@@ -440,7 +446,7 @@ void UPoseSearchInteractionAsset::CalculateWarpTransforms(float Time, float Time
 			// finding FullAlignedActorRootBoneTransforms with the minimal possible translations
 			for (int32 ItemIndex = 0; ItemIndex < ItemsNum; ++ItemIndex)
 			{
-				if (!CalculateFullAlignedActorRootBoneTransform(ActorRootBoneTransforms[ItemIndex], AssetRootBoneTransforms[ItemIndex], 
+				if (!CalculateFullAlignedActorRootBoneTransform(ActorRootBoneTransforms[ItemIndex], AssetRootBoneTransforms[ItemIndex],
 					ActorsReferencePosition, AssetReferencePosition, WarpingBankingWeight, FullAlignedActorRootBoneTransforms[ItemIndex], DebugWorld))
 				{
 					// finding the references without the ItemIndex, since the current actors and asset references are too close to ActorRootBoneTransforms[ItemIndex] and AssetRootBoneTransforms[ItemIndex]
@@ -475,14 +481,14 @@ void UPoseSearchInteractionAsset::CalculateWarpTransforms(float Time, float Time
 			DrawDebugCircle(DebugWorld, FMatrix(FVector::ZAxisVector, FVector::XAxisVector, FVector::YAxisVector, ActorsReferencePosition), 5.f, 16, FColorList::Blue, false, 0.f, SDPG_Foreground, 0.f, false);
 #endif // ENABLE_ANIM_DEBUG && WITH_EDITOR
 		}
-		
+
 		if (MinimalTranslationWeight < 1.f)
 		{
 			TArray<int32, TInlineAllocator<PreallocatedRolesNum>> SortedByWarpingWeightRotationItemIndex;
 			TArray<float, TInlineAllocator<PreallocatedRolesNum>> NormalizedWarpingWeightRotation;
 			SortedByWarpingWeightRotationItemIndex.SetNum(ItemsNum);
 			NormalizedWarpingWeightRotation.SetNum(ItemsNum);
-		
+
 			float WarpingWeightRotationSum = 0.f;
 			for (int32 ItemIndex = 0; ItemIndex < ItemsNum; ++ItemIndex)
 			{
@@ -578,5 +584,8 @@ void UPoseSearchInteractionAsset::CalculateWarpTransforms(float Time, float Time
 				FullAlignedActorRootBoneTransforms[ItemIndex] = DeltaAssetRootBoneTransform * FullAlignedActorRootBoneTransforms[ItemIndex];
 			}
 		}
+
+		break;
+	}
 	}
 }
