Fix montage trajectories not working when operating on pawns not characters
Make montage prediction CVar `a.PoseSearchTrajectory.UseWarpedMontage` affect Mover / when using predictor as well

#tests PIE

#Synced-CL 49819669

diff --git a/Engine/Plugins/Animation/PoseSearch/Source/Runtime/Private/PoseSearchTrajectoryLibrary.cpp b/Engine/Plugins/Animation/PoseSearch/Source/Runtime/Private/PoseSearchTrajectoryLibrary.cpp
--- a/Engine/Plugins/Animation/PoseSearch/Source/Runtime/Private/PoseSearchTrajectoryLibrary.cpp
+++ b/Engine/Plugins/Animation/PoseSearch/Source/Runtime/Private/PoseSearchTrajectoryLibrary.cpp
@@ -582,23 +582,23 @@ void UPoseSearchTrajectoryLibrary::UpdatePrediction_SimulateMontage(const UObjec
 
 	float AccumulatedSeconds = DeltaTime;
 
-	const ACharacter* Character = Cast<ACharacter>(InContext);
-	if (!Character)
+	const AActor* OwningActor = Cast<AActor>(InContext);
+	if (!OwningActor)
 	{
 		if (const UAnimInstance* AnimInstance = Cast<UAnimInstance>(InContext))
 		{
-			Character = Cast<ACharacter>(AnimInstance->GetOwningActor());
+			OwningActor = Cast<AActor>(AnimInstance->GetOwningActor());
 		}
 		else if (const UActorComponent* AnimNextComponent = Cast<UActorComponent>(InContext))
 		{
-			Character = Cast<ACharacter>(AnimNextComponent->GetOwner());
+			OwningActor = Cast<AActor>(AnimNextComponent->GetOwner());
 		}
 	}
 
 	UMotionWarpingComponent* WarpingComp = InWarpingComponent;
 	if (!WarpingComp && bSimulateWarping)
 	{
-		WarpingComp = Character ? Character->FindComponentByClass<UMotionWarpingComponent>() : nullptr;
+		WarpingComp = OwningActor ? OwningActor->FindComponentByClass<UMotionWarpingComponent>() : nullptr;
 	}
 	UMotionWarpingMontageTrajectoryAdapter* PredictionAdapter = nullptr;
 
@@ -780,20 +780,22 @@ void UPoseSearchTrajectoryLibrary::PoseSearchGenerateWarpedTransformTrajectory(c
 	InitTrajectorySamples(InOutTrajectory, TrajectoryDataDerived.Position, TrajectoryDataDerived.Facing, TrajectoryDataSampling, InDeltaTime);
 	UpdateHistory_TransformHistory(InOutTrajectory, TrajectoryDataDerived.Position, TrajectoryDataDerived.Velocity, TrajectoryDataSampling, InDeltaTime);
 
-	const ACharacter* Character = Cast<ACharacter>(InContext);
-	if (!Character)
+	FAnimMontageInstance* AnimMontageInstance = nullptr;
+
+	if (const UAnimInstance* AnimInstance = Cast<UAnimInstance>(InContext))
 	{
-		if (const UAnimInstance* AnimInstance = Cast<UAnimInstance>(InContext))
-		{
-			Character = Cast<ACharacter>(AnimInstance->GetOwningActor());
-		}
-		else if (const UActorComponent* AnimNextComponent = Cast<UActorComponent>(InContext))
-		{
-			Character = Cast<ACharacter>(AnimNextComponent->GetOwner());
-		}
+		AnimMontageInstance = AnimInstance->GetActiveMontageInstance();
+	}
+	else if (const ACharacter* Character = Cast<ACharacter>(InContext))
+	{
+		AnimMontageInstance = Character->GetRootMotionAnimMontageInstance();
+	}
+	else if (const UActorComponent* AnimNextComponent = Cast<UActorComponent>(InContext))
+	{
+		ACharacter* OwningCharacter = Cast<ACharacter>(AnimNextComponent->GetOwner());
+		AnimMontageInstance = OwningCharacter ? OwningCharacter->GetRootMotionAnimMontageInstance() : nullptr;
 	}
 
-	const FAnimMontageInstance* AnimMontageInstance = Character ? Character->GetRootMotionAnimMontageInstance() : nullptr;
 	if (AnimMontageInstance && bSimulateMontages)
 	{
 		UpdatePrediction_SimulateMontage(InContext, AnimMontageInstance, InOutTrajectory, InTrajectoryData, TrajectoryDataDerived, TrajectoryDataSampling, InDeltaTime, bSimulateWarping);
@@ -913,7 +915,34 @@ void UPoseSearchTrajectoryLibrary::PoseSearchGenerateTransformTrajectoryWithPred
 	InOutTrajectory.Samples[InTrajectoryHistoryCount].Position = CurrentPosition;
 	InOutTrajectory.Samples[InTrajectoryHistoryCount].Facing = CurrentFacing;
 
-	if (InPredictor != nullptr)
+	auto TryGenerateTrajectoryWithMontage = [&]()
+	{
+		if (GVarPoseSearchTrajectoryUseWarpedMontage && InPredictor.GetObject())
+		{
+			if (const UAnimInstance* AnimInstance = Cast<UAnimInstance>(InPredictor.GetObject()->GetOuter()))
+			{
+				if (FAnimMontageInstance* AnimMontageInstance = AnimInstance->GetActiveMontageInstance())
+				{
+					bool bSimulateWarping = true;
+					bool bSimulateCharacterMovement = false;
+
+					// When not simulating CMC, simulate montage only needs current position / facing
+					FPoseSearchTrajectoryData TrajectoryData;
+					FPoseSearchTrajectoryData::FDerived TrajectoryDataDerived;
+					TrajectoryDataDerived.Position = CurrentPosition;
+					TrajectoryDataDerived.Facing = CurrentFacing;
+
+					UpdatePrediction_SimulateMontage(AnimInstance, AnimMontageInstance, InOutTrajectory, TrajectoryData, TrajectoryDataDerived, TrajectoryDataSampling, InDeltaTime, bSimulateWarping, bSimulateCharacterMovement);
+
+					return true;
+				}
+			}
+		}
+
+		return false;
+	};
+
+	if (InPredictor && !TryGenerateTrajectoryWithMontage())
 	{
 		InPredictor->Predict(InOutTrajectory, InTrajectoryPredictionCount, InPredictionSamplingInterval, InTrajectoryHistoryCount);
 	}
diff --git a/Engine/Plugins/Animation/PoseSearch/Source/Runtime/Public/PoseSearch/PoseSearchTrajectoryLibrary.h b/Engine/Plugins/Animation/PoseSearch/Source/Runtime/Public/PoseSearch/PoseSearchTrajectoryLibrary.h
--- a/Engine/Plugins/Animation/PoseSearch/Source/Runtime/Public/PoseSearch/PoseSearchTrajectoryLibrary.h
+++ b/Engine/Plugins/Animation/PoseSearch/Source/Runtime/Public/PoseSearch/PoseSearchTrajectoryLibrary.h
@@ -325,7 +325,7 @@ public:
 		float InHistorySamplingInterval = 0.04f, int32 InTrajectoryHistoryCount = 10, float InPredictionSamplingInterval = 0.2f, int32 InTrajectoryPredictionCount = 8);
 
 	UE_EXPERIMENTAL(5.8, "Experimental pose search function, API may change in the future")
-	UFUNCTION(BlueprintCallable, Category = "Animation|PoseSearch", meta = (BlueprintThreadSafe, DisplayName = "Pose Search Generate Trajectory (for Character)"))
+	UFUNCTION(BlueprintCallable, Category = "Animation|PoseSearch|Experimental", meta = (BlueprintThreadSafe, DisplayName = "Pose Search Generate Warped Trajectory (Experimental)"))
 	static UE_API void PoseSearchGenerateWarpedTransformTrajectory(const UObject* InAnimInstance,
 		UPARAM(ref) const FPoseSearchTrajectoryData& InTrajectoryData, float InDeltaTime,
 		UPARAM(ref) FTransformTrajectory& InOutTrajectory, UPARAM(ref) float& InOutDesiredControllerYawLastUpdate, FTransformTrajectory& OutTrajectory,
