Fix missing debug tracks in rewind debugger for PoseSearch, SequenceInfo, and EvaluationTask when using UAF
Use strongly typed references to class names instead of hardcoded strings to make rewind debugger tracks more robust against class renames
Move UAF specific pose search debugger code to the UAFPoseSearch plugin

#Synced-CL 49767845

diff --git a/Engine/Plugins/Animation/GameplayInsights/Source/GameplayInsights/Private/AnimCurvesTrack.cpp b/Engine/Plugins/Animation/GameplayInsights/Source/GameplayInsights/Private/AnimCurvesTrack.cpp
--- a/Engine/Plugins/Animation/GameplayInsights/Source/GameplayInsights/Private/AnimCurvesTrack.cpp
+++ b/Engine/Plugins/Animation/GameplayInsights/Source/GameplayInsights/Private/AnimCurvesTrack.cpp
@@ -6,6 +6,9 @@
 #include "GameplayProvider.h"
 #include "AnimationProvider.h"
 #include "SAnimationCurvesView.h"
+#if WITH_ENGINE
+#include "Components/SkeletalMeshComponent.h"
+#endif // WITH_ENGINE
 
 #define LOCTEXT_NAMESPACE "AnimCurvesTrack"
 
@@ -211,13 +214,11 @@ TSharedPtr<SWidget> FAnimCurveTrack::GetTimelineViewInternal()
 		.CurveData_Raw(this, &FAnimCurveTrack::GetCurveData);
 }
 
-
+#if WITH_ENGINE
 static const FName AnimationCurvesName("AnimationCurves");
-
 FName FAnimationCurvesTrackCreator::GetTargetTypeNameInternal() const
 {
-	static const FName SkelMeshName("SkeletalMeshComponent");
-	return SkelMeshName;
+	return USkeletalMeshComponent::StaticClass()->GetFName();
 }
 	
 FName FAnimationCurvesTrackCreator::GetNameInternal() const
@@ -250,8 +251,8 @@ bool FAnimationCurvesTrackCreator::HasDebugInfoInternal(const FObjectId& InObjec
 	}
 	return bHasData;
 }
+#endif // WITH_ENGINE
 
-	
 }
 
 #undef LOCTEXT_NAMESPACE
\ No newline at end of file
diff --git a/Engine/Plugins/Animation/GameplayInsights/Source/GameplayInsights/Private/AnimCurvesTrack.h b/Engine/Plugins/Animation/GameplayInsights/Source/GameplayInsights/Private/AnimCurvesTrack.h
--- a/Engine/Plugins/Animation/GameplayInsights/Source/GameplayInsights/Private/AnimCurvesTrack.h
+++ b/Engine/Plugins/Animation/GameplayInsights/Source/GameplayInsights/Private/AnimCurvesTrack.h
@@ -67,7 +67,7 @@ private:
 	TArray<TSharedPtr<FAnimCurveTrack>> Children;
 };
 
-	
+#if WITH_ENGINE
 class FAnimationCurvesTrackCreator : public IRewindDebuggerTrackCreator
 {
 private:
@@ -77,5 +77,6 @@ private:
 	virtual TSharedPtr<FRewindDebuggerTrack> CreateTrackInternal(const FObjectId& InObjectId) const override;
 	virtual bool HasDebugInfoInternal(const FObjectId& InObjectId) const override;
 };
+#endif // WITH_ENGINE
 
 }
\ No newline at end of file
diff --git a/Engine/Plugins/Animation/GameplayInsights/Source/GameplayInsights/Private/BlendWeightsTrack.cpp b/Engine/Plugins/Animation/GameplayInsights/Source/GameplayInsights/Private/BlendWeightsTrack.cpp
--- a/Engine/Plugins/Animation/GameplayInsights/Source/GameplayInsights/Private/BlendWeightsTrack.cpp
+++ b/Engine/Plugins/Animation/GameplayInsights/Source/GameplayInsights/Private/BlendWeightsTrack.cpp
@@ -418,10 +418,10 @@ bool FBlendWeightTrack::HandleDoubleClickInternal()
 	return false;
 }
 
+#if WITH_ENGINE
 FName FBlendWeightsTrackCreator::GetTargetTypeNameInternal() const
 {
-	static const FName TargetTypeName = "AnimInstance";
-	return TargetTypeName;
+	return UAnimInstance::StaticClass()->GetFName();
 }
 	
 static const FName BlendWeightsName("BlendWeights");
@@ -457,6 +457,7 @@ bool FBlendWeightsTrackCreator::HasDebugInfoInternal(const FObjectId& InObjectId
 	
 	return bHasData;
 }
+#endif // #if WITH_ENGINE
 
 }
 #undef LOCTEXT_NAMESPACE
diff --git a/Engine/Plugins/Animation/GameplayInsights/Source/GameplayInsights/Private/BlendWeightsTrack.h b/Engine/Plugins/Animation/GameplayInsights/Source/GameplayInsights/Private/BlendWeightsTrack.h
--- a/Engine/Plugins/Animation/GameplayInsights/Source/GameplayInsights/Private/BlendWeightsTrack.h
+++ b/Engine/Plugins/Animation/GameplayInsights/Source/GameplayInsights/Private/BlendWeightsTrack.h
@@ -80,7 +80,8 @@ private:
 	uint64 ObjectId;
 	TArray<TSharedPtr<FBlendWeightTrack>> Children;
 };
-	
+
+#if WITH_ENGINE
 class FBlendWeightsTrackCreator : public IRewindDebuggerTrackCreator
 {
 public:
@@ -90,5 +91,6 @@ public:
 	virtual TSharedPtr<FRewindDebuggerTrack> CreateTrackInternal(const FObjectId& InObjectId) const override;
 	virtual bool HasDebugInfoInternal(const FObjectId& InObjectId) const override;
 };
+#endif
 
 }
\ No newline at end of file
diff --git a/Engine/Plugins/Animation/GameplayInsights/Source/GameplayInsights/Private/ExternalMorphTrack.cpp b/Engine/Plugins/Animation/GameplayInsights/Source/GameplayInsights/Private/ExternalMorphTrack.cpp
--- a/Engine/Plugins/Animation/GameplayInsights/Source/GameplayInsights/Private/ExternalMorphTrack.cpp
+++ b/Engine/Plugins/Animation/GameplayInsights/Source/GameplayInsights/Private/ExternalMorphTrack.cpp
@@ -6,6 +6,9 @@
 #include "GameplayProvider.h"
 #include "AnimationProvider.h"
 #include "SAnimationCurvesView.h"
+#if WITH_ENGINE
+#include "Components/SkeletalMeshComponent.h"
+#endif //WITH_ENGINE 
 
 #define LOCTEXT_NAMESPACE "ExternalMorphTrack"
 
@@ -262,11 +265,10 @@ namespace RewindDebugger
 	}
 
 	//----------------------------------------------------------------------------------------------
-
+#if WITH_ENGINE
 	FName FExternalMorphSetGroupTrackCreator::GetTargetTypeNameInternal() const
 	{
-		static const FName SkelMeshName("SkeletalMeshComponent");
-		return SkelMeshName;
+		return USkeletalMeshComponent::StaticClass()->GetFName();
 	}
 	
 	static const FName ExternalMorphSetName("ExternalMorphSetGroup");
@@ -296,6 +298,7 @@ namespace RewindDebugger
 		}
 		return false;
 	}
+#endif // WITH_ENGINE
 }
 
 #undef LOCTEXT_NAMESPACE
diff --git a/Engine/Plugins/Animation/GameplayInsights/Source/GameplayInsights/Private/ExternalMorphTrack.h b/Engine/Plugins/Animation/GameplayInsights/Source/GameplayInsights/Private/ExternalMorphTrack.h
--- a/Engine/Plugins/Animation/GameplayInsights/Source/GameplayInsights/Private/ExternalMorphTrack.h
+++ b/Engine/Plugins/Animation/GameplayInsights/Source/GameplayInsights/Private/ExternalMorphTrack.h
@@ -97,6 +97,7 @@ namespace RewindDebugger
 	};
 
 
+#if WITH_ENGINE
 	// The creator for the external morph set track.
 	class FExternalMorphSetGroupTrackCreator
 		: public IRewindDebuggerTrackCreator
@@ -110,5 +111,6 @@ namespace RewindDebugger
 		virtual bool HasDebugInfoInternal(const FObjectId& InObjectId) const override;
 		// ~END IRewindDebuggerTrackCreator overrides.
 	};
+#endif // WITH_ENGINE
 
 }	// namespace RewindDebugger
\ No newline at end of file
diff --git a/Engine/Plugins/Animation/GameplayInsights/Source/GameplayInsights/Private/InertializationsTrack.cpp b/Engine/Plugins/Animation/GameplayInsights/Source/GameplayInsights/Private/InertializationsTrack.cpp
--- a/Engine/Plugins/Animation/GameplayInsights/Source/GameplayInsights/Private/InertializationsTrack.cpp
+++ b/Engine/Plugins/Animation/GameplayInsights/Source/GameplayInsights/Private/InertializationsTrack.cpp
@@ -246,10 +246,10 @@ bool FInertializationTrack::HandleDoubleClickInternal()
 	return false;
 }
 
+#if WITH_ENGINE
 FName FInertializationsTrackCreator::GetTargetTypeNameInternal() const
 {
-	static const FName TargetTypeName = "AnimInstance";
-	return TargetTypeName;
+	return UAnimInstance::StaticClass()->GetFName();
 }
 
 static const FName InertializationsName("Inertializations");
@@ -284,6 +284,7 @@ bool FInertializationsTrackCreator::HasDebugInfoInternal(const FObjectId& InObje
 
 	return false;
 }
+#endif // WITH_ENGINE
 
 }
 #undef LOCTEXT_NAMESPACE
diff --git a/Engine/Plugins/Animation/GameplayInsights/Source/GameplayInsights/Private/InertializationsTrack.h b/Engine/Plugins/Animation/GameplayInsights/Source/GameplayInsights/Private/InertializationsTrack.h
--- a/Engine/Plugins/Animation/GameplayInsights/Source/GameplayInsights/Private/InertializationsTrack.h
+++ b/Engine/Plugins/Animation/GameplayInsights/Source/GameplayInsights/Private/InertializationsTrack.h
@@ -61,7 +61,8 @@ private:
 	uint64 ObjectId = INDEX_NONE;
 	TArray<TSharedPtr<FInertializationTrack>> Children;
 };
-	
+
+#if WITH_ENGINE
 class FInertializationsTrackCreator : public IRewindDebuggerTrackCreator
 {
 public:
@@ -71,5 +72,6 @@ public:
 	virtual TSharedPtr<FRewindDebuggerTrack> CreateTrackInternal(const FObjectId& InObjectId) const override;
 	virtual bool HasDebugInfoInternal(const FObjectId& InObjectId) const override;
 };
+#endif // WITH_ENGINE
 
 }
\ No newline at end of file
diff --git a/Engine/Plugins/Animation/GameplayInsights/Source/GameplayInsights/Private/MontagesTrack.cpp b/Engine/Plugins/Animation/GameplayInsights/Source/GameplayInsights/Private/MontagesTrack.cpp
--- a/Engine/Plugins/Animation/GameplayInsights/Source/GameplayInsights/Private/MontagesTrack.cpp
+++ b/Engine/Plugins/Animation/GameplayInsights/Source/GameplayInsights/Private/MontagesTrack.cpp
@@ -376,13 +376,18 @@ bool FMontageTrack::HandleDoubleClickInternal()
 	return false;
 }
 
+static const FName MontagesName("Montages");
+FName FMontagesTrack::GetNameInternal() const
+{
+	return MontagesName;
+}
+
+#if WITH_ENGINE
 FName FMontagesTrackCreator::GetTargetTypeNameInternal() const
 {
-	static const FName TargetTypeName = "AnimInstance";
-	return TargetTypeName;
+	return UAnimInstance::StaticClass()->GetFName();
 }
 
-static const FName MontagesName("Montages");
 FName FMontagesTrackCreator::GetNameInternal() const
 {
 	return MontagesName;
@@ -392,11 +397,6 @@ void FMontagesTrackCreator::GetTrackTypesInternal(TArray<FRewindDebuggerTrackTyp
 {
 	Types.Add({MontagesName, LOCTEXT("Montages", "Montages")});
 }
-	
-FName FMontagesTrack::GetNameInternal() const
-{
-	return MontagesName;
-}
 
 TSharedPtr<RewindDebugger::FRewindDebuggerTrack> FMontagesTrackCreator::CreateTrackInternal(const FObjectId& InObjectId) const
 {
@@ -419,6 +419,7 @@ bool FMontagesTrackCreator::HasDebugInfoInternal(const FObjectId& InObjectId) co
 	
 	return bHasData;
 }
+#endif // WITH_ENGINE
 
 }
 #undef LOCTEXT_NAMESPACE
diff --git a/Engine/Plugins/Animation/GameplayInsights/Source/GameplayInsights/Private/MontagesTrack.h b/Engine/Plugins/Animation/GameplayInsights/Source/GameplayInsights/Private/MontagesTrack.h
--- a/Engine/Plugins/Animation/GameplayInsights/Source/GameplayInsights/Private/MontagesTrack.h
+++ b/Engine/Plugins/Animation/GameplayInsights/Source/GameplayInsights/Private/MontagesTrack.h
@@ -73,7 +73,8 @@ private:
 	uint64 ObjectId;
 	TArray<TSharedPtr<FMontageTrack>> Children;
 };
-	
+
+#if WITH_ENGINE
 class FMontagesTrackCreator : public IRewindDebuggerTrackCreator
 {
 public:
@@ -83,5 +84,6 @@ public:
 	virtual TSharedPtr<FRewindDebuggerTrack> CreateTrackInternal(const FObjectId& InObjectId) const override;
 	virtual bool HasDebugInfoInternal(const FObjectId& InObjectId) const override;
 };
+#endif // WITH_ENGINE
 
 }
\ No newline at end of file
diff --git a/Engine/Plugins/Animation/GameplayInsights/Source/GameplayInsights/Private/NotifiesTrack.cpp b/Engine/Plugins/Animation/GameplayInsights/Source/GameplayInsights/Private/NotifiesTrack.cpp
--- a/Engine/Plugins/Animation/GameplayInsights/Source/GameplayInsights/Private/NotifiesTrack.cpp
+++ b/Engine/Plugins/Animation/GameplayInsights/Source/GameplayInsights/Private/NotifiesTrack.cpp
@@ -6,6 +6,9 @@
 #include "GameplayProvider.h"
 #include "AnimationProvider.h"
 #include "SNotifiesView.h"
+#if WITH_ENGINE
+#include "Animation/AnimInstance.h"
+#endif // WITH_ENGINE
 
 #define LOCTEXT_NAMESPACE "NotifiesTrack"
 
@@ -151,11 +154,10 @@ TSharedPtr<SWidget> FNotifyTrack::GetTimelineViewInternal()
 }
 
 static const FName NotifiesName("Notifies");
-
+#if WITH_ENGINE
 FName FNotifiesTrackCreator::GetTargetTypeNameInternal() const
 {
-	static const FName AnimInstanceName("AnimInstance");
-	return AnimInstanceName;
+	return UAnimInstance::StaticClass()->GetFName();
 }
 	
 FName FNotifiesTrackCreator::GetNameInternal() const
@@ -168,12 +170,11 @@ void FNotifiesTrackCreator::GetTrackTypesInternal(TArray<FRewindDebuggerTrackTyp
 	Types.Add({ NotifiesName, LOCTEXT("Notifies", "Notifies") });
 }
 
-
 TSharedPtr<RewindDebugger::FRewindDebuggerTrack> FNotifiesTrackCreator::CreateTrackInternal(const FObjectId& InObjectId) const
 {
 	return MakeShared<RewindDebugger::FNotifiesTrack>(InObjectId.GetMainId());
 }
-
+#endif // WITH_ENGINE
 		
 FNotifiesTrack::FNotifiesTrack(uint64 InObjectId) :
 	ObjectId(InObjectId)
@@ -261,6 +262,7 @@ TConstArrayView<TSharedPtr<FRewindDebuggerTrack>> FNotifiesTrack::GetChildrenInt
 	return MakeConstArrayView<TSharedPtr<FRewindDebuggerTrack>>(reinterpret_cast<const TSharedPtr<FRewindDebuggerTrack>*>(Children.GetData()), Children.Num());
 }
 
+#if WITH_ENGINE
 bool FNotifiesTrackCreator::HasDebugInfoInternal(const FObjectId& InObjectId) const
 {
 	TRACE_CPUPROFILER_EVENT_SCOPE(FNotifiesTrack::HasDebugInfoInternal);
@@ -277,7 +279,7 @@ bool FNotifiesTrackCreator::HasDebugInfoInternal(const FObjectId& InObjectId) co
 	}
 	return bHasData;
 }
-
+#endif // WITH_ENGINE
 	
 }
 
diff --git a/Engine/Plugins/Animation/GameplayInsights/Source/GameplayInsights/Private/NotifiesTrack.h b/Engine/Plugins/Animation/GameplayInsights/Source/GameplayInsights/Private/NotifiesTrack.h
--- a/Engine/Plugins/Animation/GameplayInsights/Source/GameplayInsights/Private/NotifiesTrack.h
+++ b/Engine/Plugins/Animation/GameplayInsights/Source/GameplayInsights/Private/NotifiesTrack.h
@@ -95,7 +95,7 @@ private:
 	TArray<TSharedPtr<FNotifyTrack>> Children;
 };
 
-
+#if WITH_ENGINE
 class FNotifiesTrackCreator : public IRewindDebuggerTrackCreator
 {
 private:
@@ -105,5 +105,6 @@ private:
 	virtual TSharedPtr<FRewindDebuggerTrack> CreateTrackInternal(const FObjectId& InObjectId) const override;
 	virtual bool HasDebugInfoInternal(const FObjectId& InObjectId) const override;
 };
+#endif // WITH_ENGINE
 
 }
\ No newline at end of file
diff --git a/Engine/Plugins/Animation/GameplayInsights/Source/GameplayInsights/Private/PoseWatchTrack.cpp b/Engine/Plugins/Animation/GameplayInsights/Source/GameplayInsights/Private/PoseWatchTrack.cpp
--- a/Engine/Plugins/Animation/GameplayInsights/Source/GameplayInsights/Private/PoseWatchTrack.cpp
+++ b/Engine/Plugins/Animation/GameplayInsights/Source/GameplayInsights/Private/PoseWatchTrack.cpp
@@ -240,8 +240,7 @@ TConstArrayView<TSharedPtr<FRewindDebuggerTrack>> FPoseWatchTrack::GetChildrenIn
 
 FName FPoseWatchesTrackCreator::GetTargetTypeNameInternal() const
 {
-	static const FName AnimInstanceName("AnimInstance");
-	return AnimInstanceName;
+	return UAnimInstance::StaticClass()->GetFName();
 }
 	
 static const FName PoseWatchesName("PoseWatches");
diff --git a/Engine/Plugins/Animation/GameplayInsights/Source/GameplayInsights/Private/PoseWatchTrack.h b/Engine/Plugins/Animation/GameplayInsights/Source/GameplayInsights/Private/PoseWatchTrack.h
--- a/Engine/Plugins/Animation/GameplayInsights/Source/GameplayInsights/Private/PoseWatchTrack.h
+++ b/Engine/Plugins/Animation/GameplayInsights/Source/GameplayInsights/Private/PoseWatchTrack.h
@@ -83,7 +83,7 @@ private:
 	TArray<TSharedPtr<FPoseWatchTrack>> Children;
 };
 
-
+#if WITH_ENGINE
 class FPoseWatchesTrackCreator : public IRewindDebuggerTrackCreator
 {
 private:
@@ -93,6 +93,7 @@ private:
 	virtual TSharedPtr<FRewindDebuggerTrack> CreateTrackInternal(const FObjectId& InObjectId) const override;
 	virtual bool HasDebugInfoInternal(const FObjectId& InObjectId) const override;
 };
+#endif // WITH_ENGINE
 
 }
 
diff --git a/Engine/Plugins/Animation/GameplayInsights/Source/GameplayInsights/Private/PropertiesTrack.cpp b/Engine/Plugins/Animation/GameplayInsights/Source/GameplayInsights/Private/PropertiesTrack.cpp
--- a/Engine/Plugins/Animation/GameplayInsights/Source/GameplayInsights/Private/PropertiesTrack.cpp
+++ b/Engine/Plugins/Animation/GameplayInsights/Source/GameplayInsights/Private/PropertiesTrack.cpp
@@ -285,11 +285,10 @@ namespace RewindDebugger
 	///////////////////////////////////////////////////////////////////////////////
 	// Properties Track Creator
 	///////////////////////////////////////////////////////////////////////////////
-	
+#if WITH_ENGINE
 	FName FPropertiesTrackCreator::GetTargetTypeNameInternal() const
 	{
-		static const FName TargetTypeName("Object");
-		return TargetTypeName;
+		return UObject::StaticClass()->GetFName();
 	}
 
 	static const FName PropertiesTrackCreatorName("PropertyWatch");
@@ -333,6 +332,7 @@ namespace RewindDebugger
 
 		return bHasData;
 	}
+#endif // WITH_ENGINE
 		
 }
 
diff --git a/Engine/Plugins/Animation/GameplayInsights/Source/GameplayInsights/Private/PropertiesTrack.h b/Engine/Plugins/Animation/GameplayInsights/Source/GameplayInsights/Private/PropertiesTrack.h
--- a/Engine/Plugins/Animation/GameplayInsights/Source/GameplayInsights/Private/PropertiesTrack.h
+++ b/Engine/Plugins/Animation/GameplayInsights/Source/GameplayInsights/Private/PropertiesTrack.h
@@ -66,7 +66,8 @@ namespace RewindDebugger
 		/** Count of recorded properties for traced object at the start of recording */
 		int64 TotalPropertyCountAtStart;
 	};
-	
+
+#if WITH_ENGINE
 	/**
 	 * Used to automatically create an FPropertiesTrack in Rewind Debugger
 	 */
@@ -80,4 +81,5 @@ namespace RewindDebugger
 		virtual bool HasDebugInfoInternal(const FObjectId& InObjectId) const override;
 		/** End IRewindDebuggerTrackCreator interface */
 	};
+#endif // WITH_ENGINE
 }
\ No newline at end of file
diff --git a/Engine/Plugins/Animation/GameplayInsights/Source/GameplayInsights/Private/SAnimGraphSchematicView.cpp b/Engine/Plugins/Animation/GameplayInsights/Source/GameplayInsights/Private/SAnimGraphSchematicView.cpp
--- a/Engine/Plugins/Animation/GameplayInsights/Source/GameplayInsights/Private/SAnimGraphSchematicView.cpp
+++ b/Engine/Plugins/Animation/GameplayInsights/Source/GameplayInsights/Private/SAnimGraphSchematicView.cpp
@@ -946,10 +946,10 @@ void SAnimGraphSchematicView::RefreshDetails(const TArray<TSharedRef<FAnimGraphS
 	}
 }
 
+#if WITH_ENGINE
 FName FAnimGraphSchematicTrackCreator::GetTargetTypeNameInternal() const
 {
-	static FName TargetTypeName = "AnimInstance";
-	return TargetTypeName;
+	return UAnimInstance::StaticClass()->GetFName();
 }
 
 FName FAnimGraphSchematicTrackCreator::GetNameInternal() const
@@ -966,6 +966,7 @@ TSharedPtr<RewindDebugger::FRewindDebuggerTrack> FAnimGraphSchematicTrackCreator
 {
 	return MakeShared<FAnimGraphSchematicTrack>(InObjectId.GetMainId());
 }
+#endif // WITH_ENGINE
 
 FText FAnimGraphSchematicTrack::GetDisplayNameInternal() const
 {
@@ -1017,6 +1018,7 @@ bool FAnimGraphSchematicTrack::UpdateInternal()
 	return false;
 }
 
+#if WITH_ENGINE
 bool FAnimGraphSchematicTrackCreator::HasDebugInfoInternal(const RewindDebugger::FObjectId& InObjectId) const
 {
 	const TraceServices::IAnalysisSession* AnalysisSession = IRewindDebugger::Instance()->GetAnalysisSession();
@@ -1033,5 +1035,6 @@ bool FAnimGraphSchematicTrackCreator::HasDebugInfoInternal(const RewindDebugger:
 	
 	return bHasData;
 }
+#endif // WITH_ENGINE
 
 #undef LOCTEXT_NAMESPACE
diff --git a/Engine/Plugins/Animation/GameplayInsights/Source/GameplayInsights/Private/SAnimGraphSchematicView.h b/Engine/Plugins/Animation/GameplayInsights/Source/GameplayInsights/Private/SAnimGraphSchematicView.h
--- a/Engine/Plugins/Animation/GameplayInsights/Source/GameplayInsights/Private/SAnimGraphSchematicView.h
+++ b/Engine/Plugins/Animation/GameplayInsights/Source/GameplayInsights/Private/SAnimGraphSchematicView.h
@@ -139,6 +139,7 @@ private:
 	uint64 AnimInstanceId;
 };
 
+#if WITH_ENGINE
 class FAnimGraphSchematicTrackCreator : public RewindDebugger::IRewindDebuggerTrackCreator
 {
 public:
@@ -147,4 +148,5 @@ public:
 	virtual void GetTrackTypesInternal(TArray<RewindDebugger::FRewindDebuggerTrackType>& Types) const override;
 	virtual TSharedPtr<RewindDebugger::FRewindDebuggerTrack> CreateTrackInternal(const RewindDebugger::FObjectId& InObjectId) const override;
 	virtual bool HasDebugInfoInternal(const RewindDebugger::FObjectId& InObjectId) const override;
-};
\ No newline at end of file
+};
+#endif // WITH_ENGINE
diff --git a/Engine/Plugins/Animation/GameplayInsights/Source/RewindDebuggerVLog/Private/VisualLogTrack.cpp b/Engine/Plugins/Animation/GameplayInsights/Source/RewindDebuggerVLog/Private/VisualLogTrack.cpp
--- a/Engine/Plugins/Animation/GameplayInsights/Source/RewindDebuggerVLog/Private/VisualLogTrack.cpp
+++ b/Engine/Plugins/Animation/GameplayInsights/Source/RewindDebuggerVLog/Private/VisualLogTrack.cpp
@@ -286,8 +286,7 @@ TSharedPtr<SWidget> FVisualLogTrack::GetDetailsViewInternal()
 //----------------------------------------------------------------------//
 FName FVisualLogTrackCreator::GetTargetTypeNameInternal() const
 {
-	static const FName TargetTypeName("Object");
-	return TargetTypeName;
+	return UObject::StaticClass()->GetFName();
 }
 
 FName FVisualLogTrackCreator::GetNameInternal() const
diff --git a/Engine/Plugins/Animation/GameplayInsights/Source/RewindDebuggerVLog/Private/VisualLogTrack.h b/Engine/Plugins/Animation/GameplayInsights/Source/RewindDebuggerVLog/Private/VisualLogTrack.h
--- a/Engine/Plugins/Animation/GameplayInsights/Source/RewindDebuggerVLog/Private/VisualLogTrack.h
+++ b/Engine/Plugins/Animation/GameplayInsights/Source/RewindDebuggerVLog/Private/VisualLogTrack.h
@@ -95,6 +95,7 @@ private:
 	TArray<TSharedPtr<FVisualLogCategoryTrack>> Children;
 };
 
+#if WITH_ENGINE
 /**
  * Track creator for VisualLogger recorded entries
  */
@@ -106,5 +107,6 @@ class FVisualLogTrackCreator : public IRewindDebuggerTrackCreator
 	virtual TSharedPtr<FRewindDebuggerTrack> CreateTrackInternal(const FObjectId& InObjectId) const override;
 	virtual bool HasDebugInfoInternal(const FObjectId& InObjectId) const override;
 };
+#endif // WITH_ENGINE
 
 }
diff --git a/Engine/Plugins/Animation/PoseSearch/Source/Editor/Private/PoseSearchDebugger.cpp b/Engine/Plugins/Animation/PoseSearch/Source/Editor/Private/PoseSearchDebugger.cpp
--- a/Engine/Plugins/Animation/PoseSearch/Source/Editor/Private/PoseSearchDebugger.cpp
+++ b/Engine/Plugins/Animation/PoseSearch/Source/Editor/Private/PoseSearchDebugger.cpp
@@ -556,44 +556,6 @@ TConstArrayView<TSharedPtr<RewindDebugger::FRewindDebuggerTrack>> FDebuggerTrack
 	return MakeConstArrayView<TSharedPtr<FRewindDebuggerTrack>>(reinterpret_cast<const TSharedPtr<FRewindDebuggerTrack>*>(SearchTracks.GetData()), SearchTracks.Num());
 }
 
-// FDebuggerTrackCreator
-///////////////////////////////////////////////////
-
-void FDebuggerTrackCreator::GetTrackTypesInternal(TArray<RewindDebugger::FRewindDebuggerTrackType>& Types) const
-{
-	Types.Add({ GetNameInternal(), LOCTEXT("Pose Search", "Pose Search") });
-}
-
-TSharedPtr<RewindDebugger::FRewindDebuggerTrack> FDebuggerTrackCreator::CreateTrackInternal(const RewindDebugger::FObjectId& InObjectId) const
-{
-	return MakeShared<FDebuggerTrack>(InObjectId.GetMainId());
-}
-
-bool FDebuggerTrackCreator::HasDebugInfoInternal(const RewindDebugger::FObjectId& InObjectId) const
-{
-	TRACE_CPUPROFILER_EVENT_SCOPE(PoseSearchDebugger::HasDebugInfoInternal);
-	// Get provider and validate
-	const TraceServices::IAnalysisSession* Session = IRewindDebugger::Instance()->GetAnalysisSession();
-	TraceServices::FAnalysisSessionReadScope SessionReadScope(*Session);
-
-	const FTraceProvider* PoseSearchProvider = Session->ReadProvider<FTraceProvider>(FTraceProvider::ProviderName);
-	const IAnimationProvider* AnimationProvider = Session->ReadProvider<IAnimationProvider>("AnimationProvider");
-	const IGameplayProvider* GameplayProvider = Session->ReadProvider<IGameplayProvider>("GameplayProvider");
-	if (!(PoseSearchProvider && AnimationProvider && GameplayProvider))
-	{
-		return false;
-	}
-
-	bool bHasData = false;
-
-	PoseSearchProvider->EnumerateMotionMatchingStateTimelines(InObjectId.GetMainId(), [&bHasData](const FTraceProvider::FMotionMatchingStateTimeline& InTimeline)
-	{
-		bHasData = true;
-	});
-
-	return bHasData;
-}
-
 } // namespace UE::PoseSearch
 
 #undef LOCTEXT_NAMESPACE
diff --git a/Engine/Plugins/Animation/PoseSearch/Source/Editor/Private/PoseSearchDebugger.h b/Engine/Plugins/Animation/PoseSearch/Source/Editor/Private/PoseSearchDebugger.h
--- a/Engine/Plugins/Animation/PoseSearch/Source/Editor/Private/PoseSearchDebugger.h
+++ b/Engine/Plugins/Animation/PoseSearch/Source/Editor/Private/PoseSearchDebugger.h
@@ -4,7 +4,6 @@
 
 #include "CoreMinimal.h"
 #include "IRewindDebuggerExtension.h"
-#include "IRewindDebuggerTrackCreator.h"
 #include "PoseSearch/PoseSearchDefines.h"
 #include "RewindDebuggerTrack.h"
 #include "SCurveTimelineView.h"
@@ -109,21 +108,5 @@ private:
 	TArray<TSharedPtr<FSearchTrack>> SearchTracks;
 };
 
-class FDebuggerTrackCreator : public RewindDebugger::IRewindDebuggerTrackCreator
-{
-private:
-	virtual FName GetTargetTypeNameInternal() const override { return "AnimInstance"; }
-	virtual FName GetNameInternal() const override { return "PoseSearchDebugger"; }
-	virtual void GetTrackTypesInternal(TArray<RewindDebugger::FRewindDebuggerTrackType>& Types) const override;
-	virtual TSharedPtr<RewindDebugger::FRewindDebuggerTrack> CreateTrackInternal(const RewindDebugger::FObjectId& InObjectId) const override;
-	virtual bool HasDebugInfoInternal(const RewindDebugger::FObjectId& InObjectId) const override;
-	virtual int32 GetSortOrderPriorityInternal() const override { return 10; };
-};
-
-class FAnimNextDebuggerTrackCreator : public FDebuggerTrackCreator
-{
-private:
-	virtual FName GetTargetTypeNameInternal() const override { return "AnimNextComponent"; }
-};
 
 } // namespace UE::PoseSearch
diff --git a/Engine/Plugins/Animation/PoseSearch/Source/Editor/Private/PoseSearchDebuggerTrackCreator.cpp b/Engine/Plugins/Animation/PoseSearch/Source/Editor/Private/PoseSearchDebuggerTrackCreator.cpp
new file mode 100644
--- /dev/null
+++ b/Engine/Plugins/Animation/PoseSearch/Source/Editor/Private/PoseSearchDebuggerTrackCreator.cpp
@@ -0,0 +1,60 @@
+// Copyright Epic Games, Inc. All Rights Reserved.
+
+#include "PoseSearchDebuggerTrackCreator.h"
+
+#include "IGameplayProvider.h"
+#include "IAnimationProvider.h"
+#include "PoseSearchDebugger.h"
+#include "Trace/PoseSearchTraceProvider.h"
+
+#define LOCTEXT_NAMESPACE "PoseSearchDebugger"
+
+// FDebuggerTrackCreator
+///////////////////////////////////////////////////
+
+namespace UE::PoseSearch
+{
+
+FName FDebuggerTrackCreator::GetTargetTypeNameInternal() const
+{
+	return UAnimInstance::StaticClass()->GetFName();
+}
+
+void FDebuggerTrackCreator::GetTrackTypesInternal(TArray<RewindDebugger::FRewindDebuggerTrackType>& Types) const
+{
+	Types.Add({ GetNameInternal(), LOCTEXT("Pose Search", "Pose Search") });
+}
+
+TSharedPtr<RewindDebugger::FRewindDebuggerTrack> FDebuggerTrackCreator::CreateTrackInternal(const RewindDebugger::FObjectId& InObjectId) const
+{
+	return MakeShared<UE::PoseSearch::FDebuggerTrack>(InObjectId.GetMainId());
+}
+
+bool FDebuggerTrackCreator::HasDebugInfoInternal(const RewindDebugger::FObjectId& InObjectId) const
+{
+	TRACE_CPUPROFILER_EVENT_SCOPE(PoseSearchDebugger::HasDebugInfoInternal);
+	// Get provider and validate
+	const TraceServices::IAnalysisSession* Session = IRewindDebugger::Instance()->GetAnalysisSession();
+	TraceServices::FAnalysisSessionReadScope SessionReadScope(*Session);
+
+	const UE::PoseSearch::FTraceProvider* PoseSearchProvider = Session->ReadProvider<FTraceProvider>(FTraceProvider::ProviderName);
+	const IAnimationProvider* AnimationProvider = Session->ReadProvider<IAnimationProvider>("AnimationProvider");
+	const IGameplayProvider* GameplayProvider = Session->ReadProvider<IGameplayProvider>("GameplayProvider");
+	if (!(PoseSearchProvider && AnimationProvider && GameplayProvider))
+	{
+		return false;
+	}
+
+	bool bHasData = false;
+
+	PoseSearchProvider->EnumerateMotionMatchingStateTimelines(InObjectId.GetMainId(), [&bHasData](const FTraceProvider::FMotionMatchingStateTimeline& InTimeline)
+	{
+		bHasData = true;
+	});
+
+	return bHasData;
+}
+
+} // namespace UE::PoseSearch
+
+#undef LOCTEXT_NAMESPACE
diff --git a/Engine/Plugins/Animation/PoseSearch/Source/Editor/Private/PoseSearchEditor.cpp b/Engine/Plugins/Animation/PoseSearch/Source/Editor/Private/PoseSearchEditor.cpp
--- a/Engine/Plugins/Animation/PoseSearch/Source/Editor/Private/PoseSearchEditor.cpp
+++ b/Engine/Plugins/Animation/PoseSearch/Source/Editor/Private/PoseSearchEditor.cpp
@@ -20,6 +20,7 @@
 #include "Trace/PoseSearchTraceAnalyzer.h"
 #include "Trace/PoseSearchTraceModule.h"
 #include "PoseSearchColumnEditor.h"
+#include "PoseSearchDebuggerTrackCreator.h"
 
 DEFINE_LOG_CATEGORY(LogPoseSearchEditor);
 
@@ -41,7 +42,6 @@ private:
 	
 	/** Creates the track for the Rewind Debugger */
 	TSharedPtr<FDebuggerTrackCreator> DebuggerTrackCreator;
-	TSharedPtr<FAnimNextDebuggerTrackCreator> AnimNextDebuggerTrackCreator;
 
 	/** Enables dedicated PoseSearch trace module */
 	TSharedPtr<FTraceModule> TraceModule;
@@ -68,10 +68,8 @@ void FEditorModule::StartupModule()
 		FDebugger::Initialize();
 		TraceModule = MakeShared<FTraceModule>();
 		DebuggerTrackCreator = MakeShared<FDebuggerTrackCreator>();
-		AnimNextDebuggerTrackCreator = MakeShared<FAnimNextDebuggerTrackCreator>();
 		
 		IModularFeatures::Get().RegisterModularFeature(FDebuggerTrackCreator::ModularFeatureName, DebuggerTrackCreator.Get());
-		IModularFeatures::Get().RegisterModularFeature(FDebuggerTrackCreator::ModularFeatureName, AnimNextDebuggerTrackCreator.Get());
 		IModularFeatures::Get().RegisterModularFeature(TraceServices::ModuleFeatureName, TraceModule.Get());
 		
 		// Register Ed Modes used by PoseSearchDatabase and PoseSearchInteractionAsset
@@ -103,8 +101,16 @@ void FEditorModule::ShutdownModule()
 
 	// Unregister Asset Editor Commands
 	FDatabaseEditorCommands::Unregister();
-	
-	IModularFeatures::Get().UnregisterModularFeature(TraceServices::ModuleFeatureName, TraceModule.Get());
+
+	if (DebuggerTrackCreator.Get() != nullptr)
+	{
+		IModularFeatures::Get().UnregisterModularFeature(FDebuggerTrackCreator::ModularFeatureName, DebuggerTrackCreator.Get());
+	}
+
+	if (TraceModule.Get() != nullptr)
+	{
+		IModularFeatures::Get().UnregisterModularFeature(TraceServices::ModuleFeatureName, TraceModule.Get());
+	}
 	FDebugger::Shutdown();
 }
 
diff --git a/Engine/Plugins/Animation/PoseSearch/Source/Editor/Public/PoseSearchDebuggerTrackCreator.h b/Engine/Plugins/Animation/PoseSearch/Source/Editor/Public/PoseSearchDebuggerTrackCreator.h
new file mode 100644
--- /dev/null
+++ b/Engine/Plugins/Animation/PoseSearch/Source/Editor/Public/PoseSearchDebuggerTrackCreator.h
@@ -0,0 +1,27 @@
+// Copyright Epic Games, Inc. All Rights Reserved.
+
+#pragma once
+
+#include "CoreMinimal.h"
+#include "IRewindDebuggerTrackCreator.h"
+
+#define UE_API POSESEARCHEDITOR_API
+
+namespace UE::PoseSearch
+{
+
+class FDebuggerTrackCreator : public RewindDebugger::IRewindDebuggerTrackCreator
+{
+private:
+	virtual UE_API FName GetTargetTypeNameInternal() const override;
+	virtual FName GetNameInternal() const override { return "PoseSearchDebugger"; }
+	virtual UE_API void  GetTrackTypesInternal(TArray<RewindDebugger::FRewindDebuggerTrackType>& Types) const override;
+	virtual UE_API TSharedPtr<RewindDebugger::FRewindDebuggerTrack> CreateTrackInternal(const RewindDebugger::FObjectId& InObjectId) const override;
+	virtual UE_API bool HasDebugInfoInternal(const RewindDebugger::FObjectId& InObjectId) const override;
+	virtual int32 GetSortOrderPriorityInternal() const override { return 10; };
+};
+
+} // UE::PoseSearch
+
+#undef UE_API
+
diff --git a/Engine/Plugins/Cameras/GameplayCameras/Source/GameplayCamerasEditor/Private/Trace/CameraSystemRewindDebuggerTrack.h b/Engine/Plugins/Cameras/GameplayCameras/Source/GameplayCamerasEditor/Private/Trace/CameraSystemRewindDebuggerTrack.h
--- a/Engine/Plugins/Cameras/GameplayCameras/Source/GameplayCamerasEditor/Private/Trace/CameraSystemRewindDebuggerTrack.h
+++ b/Engine/Plugins/Cameras/GameplayCameras/Source/GameplayCamerasEditor/Private/Trace/CameraSystemRewindDebuggerTrack.h
@@ -53,6 +53,7 @@ private:
 	TSharedPtr<FCameraSystemRewindDebuggerTrackTimelineData> TimelineData;
 };
 
+#if WITH_ENGINE
 /**
  * Factory class that creates the rewind debugger track for the camera system evaluation trace.
  */
@@ -69,6 +70,7 @@ protected:
 	virtual TSharedPtr<RewindDebugger::FRewindDebuggerTrack> CreateTrackInternal(uint64 InObjectId) const override;
 #endif
 };
+#endif // WITH_ENGINE
 
 }  // namespace UE::Cameras
 
diff --git a/Engine/Plugins/Chooser/Source/ChooserEditor/Private/ChooserTrack.cpp b/Engine/Plugins/Chooser/Source/ChooserEditor/Private/ChooserTrack.cpp
--- a/Engine/Plugins/Chooser/Source/ChooserEditor/Private/ChooserTrack.cpp
+++ b/Engine/Plugins/Chooser/Source/ChooserEditor/Private/ChooserTrack.cpp
@@ -163,10 +163,9 @@ TSharedPtr<SWidget> FChooserTrack::GetTimelineViewInternal()
 
 static const FName ChoosersName("Choosers");
 
- FName FChoosersTrackCreator::GetTargetTypeNameInternal() const
- {
- 	static const FName ObjectName("Object");
-	return ObjectName;
+FName FChoosersTrackCreator::GetTargetTypeNameInternal() const
+{
+	return UObject::StaticClass()->GetFName();
 }
 	
 FName FChoosersTrackCreator::GetNameInternal() const
diff --git a/Engine/Plugins/Chooser/Source/ChooserEditor/Private/ChooserTrack.h b/Engine/Plugins/Chooser/Source/ChooserEditor/Private/ChooserTrack.h
--- a/Engine/Plugins/Chooser/Source/ChooserEditor/Private/ChooserTrack.h
+++ b/Engine/Plugins/Chooser/Source/ChooserEditor/Private/ChooserTrack.h
@@ -63,7 +63,7 @@ private:
 	TArray<TSharedPtr<FChooserTrack>> Children;
 };
 
-
+#if WITH_ENGINE
 class FChoosersTrackCreator : public RewindDebugger::IRewindDebuggerTrackCreator
 {
 private:
@@ -73,5 +73,6 @@ private:
 	virtual TSharedPtr<RewindDebugger::FRewindDebuggerTrack> CreateTrackInternal(const RewindDebugger::FObjectId& InObjectId) const override;
 	virtual bool HasDebugInfoInternal(const RewindDebugger::FObjectId& InObjectId) const override;
 };
+#endif // WITH_ENGINE
 
 }
\ No newline at end of file
diff --git a/Engine/Plugins/Developer/FieldNotificationTrace/Source/FieldNotificationTraceEditor/Private/FieldNotificationTraceEditorModule.h b/Engine/Plugins/Developer/FieldNotificationTrace/Source/FieldNotificationTraceEditor/Private/FieldNotificationTraceEditorModule.h
--- a/Engine/Plugins/Developer/FieldNotificationTrace/Source/FieldNotificationTraceEditor/Private/FieldNotificationTraceEditorModule.h
+++ b/Engine/Plugins/Developer/FieldNotificationTrace/Source/FieldNotificationTraceEditor/Private/FieldNotificationTraceEditorModule.h
@@ -27,7 +27,9 @@ public:
 private:
 	FRewindDebuggerExtension RewindDebuggerExtension;
 	FRewindDebuggerRuntimeExtension RewindDebuggerRuntimeExtension;
+#if WITH_ENGINE
 	FTracksCreator TrackCreator;
+#endif // WITH_ENGINE 
 	FTraceServiceModule TraceModule;
 };
 
diff --git a/Engine/Plugins/Developer/FieldNotificationTrace/Source/FieldNotificationTraceEditor/Private/FieldNotificationTrack.cpp b/Engine/Plugins/Developer/FieldNotificationTrace/Source/FieldNotificationTraceEditor/Private/FieldNotificationTrack.cpp
--- a/Engine/Plugins/Developer/FieldNotificationTrace/Source/FieldNotificationTraceEditor/Private/FieldNotificationTrack.cpp
+++ b/Engine/Plugins/Developer/FieldNotificationTrace/Source/FieldNotificationTraceEditor/Private/FieldNotificationTrack.cpp
@@ -228,8 +228,7 @@ TConstArrayView<TSharedPtr<RewindDebugger::FRewindDebuggerTrack>> FObjectTrack::
 
 FName FTracksCreator::GetTargetTypeNameInternal() const
 {
-	static const FName ObjectName("Object");
-	return ObjectName;
+	return UObject::StaticClass()->GetFName();
 }
 
 FName FTracksCreator::GetNameInternal() const
diff --git a/Engine/Plugins/Developer/FieldNotificationTrace/Source/FieldNotificationTraceEditor/Private/FieldNotificationTrack.h b/Engine/Plugins/Developer/FieldNotificationTrace/Source/FieldNotificationTraceEditor/Private/FieldNotificationTrack.h
--- a/Engine/Plugins/Developer/FieldNotificationTrace/Source/FieldNotificationTraceEditor/Private/FieldNotificationTrack.h
+++ b/Engine/Plugins/Developer/FieldNotificationTrace/Source/FieldNotificationTraceEditor/Private/FieldNotificationTrack.h
@@ -2,7 +2,9 @@
 #pragma once
 
 #include "FieldNotificationId.h"
+#if WITH_ENGINE
 #include "IRewindDebuggerTrackCreator.h"
+#endif // WITH_ENGINE
 #include "RewindDebuggerTrack.h"
 #include "Textures/SlateIcon.h"
 #include "SEventTimelineView.h"
@@ -64,6 +66,7 @@ private:
 };
 
 
+#if WITH_ENGINE
 class FTracksCreator : public RewindDebugger::IRewindDebuggerTrackCreator
 {
 private:
@@ -73,5 +76,6 @@ private:
 	virtual TSharedPtr<RewindDebugger::FRewindDebuggerTrack> CreateTrackInternal(const RewindDebugger::FObjectId& InObjectId) const override;
 	virtual bool HasDebugInfoInternal(const RewindDebugger::FObjectId& InObjectId) const override;
 };
+#endif // WITH_ENGINE
 
 } // namespace
diff --git a/Engine/Plugins/Experimental/UAF/UAF/Source/UAFEditor/Private/RewindDebugger/AnimNextModuleTrack.cpp b/Engine/Plugins/Experimental/UAF/UAF/Source/UAFEditor/Private/RewindDebugger/AnimNextModuleTrack.cpp
--- a/Engine/Plugins/Experimental/UAF/UAF/Source/UAFEditor/Private/RewindDebugger/AnimNextModuleTrack.cpp
+++ b/Engine/Plugins/Experimental/UAF/UAF/Source/UAFEditor/Private/RewindDebugger/AnimNextModuleTrack.cpp
@@ -14,6 +14,7 @@
 #include "Modules/ModuleManager.h"
 #include "ObjectAsTraceIdProxyArchiveReader.h"
 #include "PropertyEditorModule.h"
+#include "UAFAssetInstance.h"
 #include "Serialization/MemoryReader.h"
 #include "StructUtils/PropertyBag.h"
 #include "Subsystems/AssetEditorSubsystem.h"
@@ -27,8 +28,7 @@ static const FName AnimNextModulesName("AnimNextModules");
 
 FName FAnimNextModuleTrackCreator::GetTargetTypeNameInternal() const
 {
-	static const FName ObjectName("UAFAssetInstance");
-	return ObjectName;
+	return FUAFAssetInstance::StaticStruct()->GetFName();
 }
 
 FText FAnimNextModuleTrack::GetDisplayNameInternal() const
diff --git a/Engine/Plugins/Experimental/UAF/UAF/Source/UAFEditor/Private/RewindDebugger/AnimNextModuleTrack.h b/Engine/Plugins/Experimental/UAF/UAF/Source/UAFEditor/Private/RewindDebugger/AnimNextModuleTrack.h
--- a/Engine/Plugins/Experimental/UAF/UAF/Source/UAFEditor/Private/RewindDebugger/AnimNextModuleTrack.h
+++ b/Engine/Plugins/Experimental/UAF/UAF/Source/UAFEditor/Private/RewindDebugger/AnimNextModuleTrack.h
@@ -79,7 +79,7 @@ private:
 	TArray<FPropertyBagPropertyDesc> PropertyDescriptions;
 };
 
-
+#if WITH_ENGINE
 class FAnimNextModuleTrackCreator : public RewindDebugger::IRewindDebuggerTrackCreator
 {
 private:
@@ -99,6 +99,7 @@ private:
 		return true;
 	}
 };
+#endif // WITH_ENGINE
 
 }
 #endif // UAF_TRACE_ENABLED
diff --git a/Engine/Plugins/Experimental/UAF/UAFAnimGraph/Source/UAFAnimGraphEditor/Private/RewindDebugger/EvaluationProgramTrack.cpp b/Engine/Plugins/Experimental/UAF/UAFAnimGraph/Source/UAFAnimGraphEditor/Private/RewindDebugger/EvaluationProgramTrack.cpp
--- a/Engine/Plugins/Experimental/UAF/UAFAnimGraph/Source/UAFAnimGraphEditor/Private/RewindDebugger/EvaluationProgramTrack.cpp
+++ b/Engine/Plugins/Experimental/UAF/UAFAnimGraph/Source/UAFAnimGraphEditor/Private/RewindDebugger/EvaluationProgramTrack.cpp
@@ -10,6 +10,7 @@
 #include "Modules/ModuleManager.h"
 #include "PropertyEditorModule.h"
 #include "ObjectAsTraceIdProxyArchiveReader.h"
+#include "Component/AnimNextComponent.h"
 #include "Serialization/ObjectReader.h"
 
 #include UE_INLINE_GENERATED_CPP_BY_NAME(EvaluationProgramTrack)
@@ -24,8 +25,7 @@ static const FName AnimNextModulesName("AnimNextModules");
 
  FName FEvaluationProgramTrackCreator::GetTargetTypeNameInternal() const
  {
- 	static const FName ObjectName("AnimNextComponent");
-	return ObjectName;
+	return UUAFComponent::StaticClass()->GetFName();
 }
 
 	
diff --git a/Engine/Plugins/Experimental/UAF/UAFAnimGraph/Source/UAFAnimGraphEditor/Private/RewindDebugger/EvaluationProgramTrack.h b/Engine/Plugins/Experimental/UAF/UAFAnimGraph/Source/UAFAnimGraphEditor/Private/RewindDebugger/EvaluationProgramTrack.h
--- a/Engine/Plugins/Experimental/UAF/UAFAnimGraph/Source/UAFAnimGraphEditor/Private/RewindDebugger/EvaluationProgramTrack.h
+++ b/Engine/Plugins/Experimental/UAF/UAFAnimGraph/Source/UAFAnimGraphEditor/Private/RewindDebugger/EvaluationProgramTrack.h
@@ -55,7 +55,7 @@ private:
 	// TArray<TSharedPtr<FEvaluationProgramTrack>> Children;
 };
 
-
+#if WITH_ENGINE
 class FEvaluationProgramTrackCreator : public RewindDebugger::IRewindDebuggerTrackCreator
 {
 private:
@@ -65,5 +65,6 @@ private:
 	virtual TSharedPtr<RewindDebugger::FRewindDebuggerTrack> CreateTrackInternal(const RewindDebugger::FObjectId& InObjectId) const override;
 	virtual bool HasDebugInfoInternal(const RewindDebugger::FObjectId& InObjectId) const override;
 };
+#endif // WITH_ENGINE
 
 }
\ No newline at end of file
diff --git a/Engine/Plugins/Experimental/UAF/UAFAnimGraph/Source/UAFAnimGraphEditor/Private/RewindDebugger/SequenceInfoTrack.cpp b/Engine/Plugins/Experimental/UAF/UAFAnimGraph/Source/UAFAnimGraphEditor/Private/RewindDebugger/SequenceInfoTrack.cpp
--- a/Engine/Plugins/Experimental/UAF/UAFAnimGraph/Source/UAFAnimGraphEditor/Private/RewindDebugger/SequenceInfoTrack.cpp
+++ b/Engine/Plugins/Experimental/UAF/UAFAnimGraph/Source/UAFAnimGraphEditor/Private/RewindDebugger/SequenceInfoTrack.cpp
@@ -12,6 +12,7 @@
 #include "Animation/AnimSequence.h"
 #include "EvaluationVM/Tasks/PushAnimSequenceKeyframe.h"
 #include "ObjectAsTraceIdProxyArchiveReader.h"
+#include "Component/AnimNextComponent.h"
 #include "Widgets/SCanvas.h"
 #include "Widgets/Text/STextBlock.h"
 #include "Widgets/Colors/SColorBlock.h"
@@ -24,8 +25,7 @@ namespace UE::UAF::Editor
 {
 FName FSequenceInfoTrackCreator::GetTargetTypeNameInternal() const
 {
-	static const FName ObjectName("AnimNextComponent");
-	return ObjectName;
+	return UUAFComponent::StaticClass()->GetFName();
 }
 
 
diff --git a/Engine/Plugins/Experimental/UAF/UAFAnimGraph/Source/UAFAnimGraphEditor/Private/RewindDebugger/SequenceInfoTrack.h b/Engine/Plugins/Experimental/UAF/UAFAnimGraph/Source/UAFAnimGraphEditor/Private/RewindDebugger/SequenceInfoTrack.h
--- a/Engine/Plugins/Experimental/UAF/UAFAnimGraph/Source/UAFAnimGraphEditor/Private/RewindDebugger/SequenceInfoTrack.h
+++ b/Engine/Plugins/Experimental/UAF/UAFAnimGraph/Source/UAFAnimGraphEditor/Private/RewindDebugger/SequenceInfoTrack.h
@@ -107,7 +107,7 @@ private:
 	// TArray<TSharedPtr<FEvaluationProgramTrack>> Children;
 };
 
-
+#if WITH_ENGINE
 class FSequenceInfoTrackCreator : public RewindDebugger::IRewindDebuggerTrackCreator
 {
 private:
@@ -117,5 +117,6 @@ private:
 	virtual TSharedPtr<RewindDebugger::FRewindDebuggerTrack> CreateTrackInternal(const RewindDebugger::FObjectId& InObjectId) const override;
 	virtual bool HasDebugInfoInternal(const RewindDebugger::FObjectId& InObjectId) const override;
 };
+#endif // WITH_ENGINE
 
 }
\ No newline at end of file
diff --git a/Engine/Plugins/Experimental/UAF/UAFPoseSearch/Source/UAFPoseSearchUncookedOnly/Private/UAFPoseSearchDebugger.cpp b/Engine/Plugins/Experimental/UAF/UAFPoseSearch/Source/UAFPoseSearchUncookedOnly/Private/UAFPoseSearchDebugger.cpp
new file mode 100644
--- /dev/null
+++ b/Engine/Plugins/Experimental/UAF/UAFPoseSearch/Source/UAFPoseSearchUncookedOnly/Private/UAFPoseSearchDebugger.cpp
@@ -0,0 +1,14 @@
+// Copyright Epic Games, Inc. All Rights Reserved.
+
+#include "UAFPoseSearchDebugger.h"
+
+#include "Component/AnimNextComponent.h"
+
+namespace UE::PoseSearch
+{
+
+FName FUAFDebuggerTrackCreator::GetTargetTypeNameInternal() const
+{
+	return UUAFComponent::StaticClass()->GetFName();
+} 
+}
diff --git a/Engine/Plugins/Experimental/UAF/UAFPoseSearch/Source/UAFPoseSearchUncookedOnly/Private/UAFPoseSearchDebugger.h b/Engine/Plugins/Experimental/UAF/UAFPoseSearch/Source/UAFPoseSearchUncookedOnly/Private/UAFPoseSearchDebugger.h
new file mode 100644
--- /dev/null
+++ b/Engine/Plugins/Experimental/UAF/UAFPoseSearch/Source/UAFPoseSearchUncookedOnly/Private/UAFPoseSearchDebugger.h
@@ -0,0 +1,19 @@
+// Copyright Epic Games, Inc. All Rights Reserved.
+
+#pragma once
+
+#include "CoreMinimal.h"
+#include "PoseSearchDebuggerTrackCreator.h"
+
+namespace UE::PoseSearch
+{
+
+#if WITH_ENGINE
+class FUAFDebuggerTrackCreator : public FDebuggerTrackCreator
+{
+private:
+	virtual FName GetTargetTypeNameInternal() const override;
+};
+#endif // WITH_ENGINE
+
+}
\ No newline at end of file
diff --git a/Engine/Plugins/Experimental/UAF/UAFPoseSearch/Source/UAFPoseSearchUncookedOnly/Private/UAFPoseSearchUncookedOnlyModule.cpp b/Engine/Plugins/Experimental/UAF/UAFPoseSearch/Source/UAFPoseSearchUncookedOnly/Private/UAFPoseSearchUncookedOnlyModule.cpp
--- a/Engine/Plugins/Experimental/UAF/UAFPoseSearch/Source/UAFPoseSearchUncookedOnly/Private/UAFPoseSearchUncookedOnlyModule.cpp
+++ b/Engine/Plugins/Experimental/UAF/UAFPoseSearch/Source/UAFPoseSearchUncookedOnly/Private/UAFPoseSearchUncookedOnlyModule.cpp
@@ -1,5 +1,44 @@
 ﻿// Copyright Epic Games, Inc. All Rights Reserved.
 
+#include "PoseSearchDebuggerTrackCreator.h"
+#include "UAFPoseSearchDebugger.h"
+#include "Features/IModularFeatures.h"
 #include "Modules/ModuleManager.h"
 
-IMPLEMENT_MODULE(FDefaultModuleImpl, UAFPoseSearchUncookedOnly)
\ No newline at end of file
+namespace UE::PoseSearch
+{
+class FUAFDebuggerTrackCreator;
+
+//////////////////////////////////////////////////////////////////////////
+// FPoseSearchEditorModule
+class FUAFPoseSearchUncookedOnlyModule : public IModuleInterface
+{
+public:
+	void StartupModule() override;
+	void ShutdownModule() override;
+
+private:
+	
+	/** Creates the track for the Rewind Debugger */
+	TSharedPtr<FUAFDebuggerTrackCreator> UAFDebuggerTrackCreator;
+};
+
+void FUAFPoseSearchUncookedOnlyModule::StartupModule()
+{
+	IModuleInterface::StartupModule();
+
+	UAFDebuggerTrackCreator = MakeShared<FUAFDebuggerTrackCreator>();
+	IModularFeatures::Get().RegisterModularFeature(RewindDebugger::IRewindDebuggerTrackCreator::ModularFeatureName, UAFDebuggerTrackCreator.Get());
+}
+
+void FUAFPoseSearchUncookedOnlyModule::ShutdownModule()
+{
+	IModularFeatures::Get().UnregisterModularFeature(RewindDebugger::IRewindDebuggerTrackCreator::ModularFeatureName, UAFDebuggerTrackCreator.Get());
+	
+	IModuleInterface::ShutdownModule();
+}
+
+} // namespace UE::PoseSearch
+
+IMPLEMENT_MODULE(UE::PoseSearch::FUAFPoseSearchUncookedOnlyModule, UAFPoseSearchUncookedOnly)
+
diff --git a/Engine/Plugins/Experimental/UAF/UAFPoseSearch/Source/UAFPoseSearchUncookedOnly/UAFPoseSearchUncookedOnly.Build.cs b/Engine/Plugins/Experimental/UAF/UAFPoseSearch/Source/UAFPoseSearchUncookedOnly/UAFPoseSearchUncookedOnly.Build.cs
--- a/Engine/Plugins/Experimental/UAF/UAFPoseSearch/Source/UAFPoseSearchUncookedOnly/UAFPoseSearchUncookedOnly.Build.cs
+++ b/Engine/Plugins/Experimental/UAF/UAFPoseSearch/Source/UAFPoseSearchUncookedOnly/UAFPoseSearchUncookedOnly.Build.cs
@@ -20,7 +20,9 @@ namespace UnrealBuildTool.Rules
 					"RigVMDeveloper",
 					"RigVM",
 					"Slate",
-					"SlateCore",
+					"SlateCore", 
+					"RewindDebuggerInterface", 
+					"PoseSearchEditor",
 				}
 			);
 		}
diff --git a/Engine/Plugins/Runtime/RigVM/Source/RigVMEditor/Private/Trace/RigVMTraceTrack.cpp b/Engine/Plugins/Runtime/RigVM/Source/RigVMEditor/Private/Trace/RigVMTraceTrack.cpp
--- a/Engine/Plugins/Runtime/RigVM/Source/RigVMEditor/Private/Trace/RigVMTraceTrack.cpp
+++ b/Engine/Plugins/Runtime/RigVM/Source/RigVMEditor/Private/Trace/RigVMTraceTrack.cpp
@@ -4,6 +4,7 @@
 #include "RigVMTraceProvider.h"
 #include "RigVMTraceAnalyzer.h"
 #include "DetailLayoutBuilder.h"
+#include "RigVMHost.h"
 #include "Widgets/Layout/SGridPanel.h"
 #include "Widgets/Input/SNumericEntryBox.h"
 
@@ -262,8 +263,7 @@ TConstArrayView<TSharedPtr<RewindDebugger::FRewindDebuggerTrack>> FRewindDebugge
 
 FName FRewindDebuggerRigVMTrackCreator::GetTargetTypeNameInternal() const
 {
-	static const FLazyName RigVMHostName("RigVMHost");
-	return RigVMHostName;
+	return URigVMHost::StaticClass()->GetFName();
 }
 	
 FName FRewindDebuggerRigVMTrackCreator::GetNameInternal() const
diff --git a/Engine/Plugins/Runtime/RigVM/Source/RigVMEditor/Public/Trace/RigVMTraceTrack.h b/Engine/Plugins/Runtime/RigVM/Source/RigVMEditor/Public/Trace/RigVMTraceTrack.h
--- a/Engine/Plugins/Runtime/RigVM/Source/RigVMEditor/Public/Trace/RigVMTraceTrack.h
+++ b/Engine/Plugins/Runtime/RigVM/Source/RigVMEditor/Public/Trace/RigVMTraceTrack.h
@@ -83,6 +83,7 @@ private:
 	TArray<TSharedPtr<FRewindDebuggerTrack>> Children;
 };
 
+#if WITH_ENGINE
 /**
  * Factory class to create RigVM rewind debugger tracks
  */
@@ -95,6 +96,7 @@ private:
 	virtual TSharedPtr<RewindDebugger::FRewindDebuggerTrack> CreateTrackInternal(const RewindDebugger::FObjectId& InObjectId) const override;
 	virtual bool HasDebugInfoInternal(const RewindDebugger::FObjectId& InObjectId) const override;
 };
+#endif // WITH_ENGINE
 
 #undef UE_API
 
diff --git a/Engine/Plugins/Runtime/StateTree/Source/StateTreeEditorModule/Private/Debugger/StateTreeRewindDebuggerTrack.cpp b/Engine/Plugins/Runtime/StateTree/Source/StateTreeEditorModule/Private/Debugger/StateTreeRewindDebuggerTrack.cpp
--- a/Engine/Plugins/Runtime/StateTree/Source/StateTreeEditorModule/Private/Debugger/StateTreeRewindDebuggerTrack.cpp
+++ b/Engine/Plugins/Runtime/StateTree/Source/StateTreeEditorModule/Private/Debugger/StateTreeRewindDebuggerTrack.cpp
@@ -56,8 +56,7 @@ bool FRewindDebuggerTrackCreator::IsCreatingPrimaryChildTrackInternal() const
 
 FName FRewindDebuggerTrackCreator::GetTargetTypeNameInternal() const
 {
-	static const FName AssociatedObjectName("StateTreeInstanceData");
-	return AssociatedObjectName;
+	return FStateTreeInstanceData::StaticStruct()->GetFName();
 }
 
 //----------------------------------------------------------------------//
diff --git a/Engine/Plugins/Runtime/StateTree/Source/StateTreeEditorModule/Private/Debugger/StateTreeRewindDebuggerTrack.h b/Engine/Plugins/Runtime/StateTree/Source/StateTreeEditorModule/Private/Debugger/StateTreeRewindDebuggerTrack.h
--- a/Engine/Plugins/Runtime/StateTree/Source/StateTreeEditorModule/Private/Debugger/StateTreeRewindDebuggerTrack.h
+++ b/Engine/Plugins/Runtime/StateTree/Source/StateTreeEditorModule/Private/Debugger/StateTreeRewindDebuggerTrack.h
@@ -16,6 +16,7 @@ namespace UE::StateTreeDebugger
 {
 class SFrameEventsView;
 
+#if WITH_ENGINE
 /**
  * RewindDebugger track creator for StateTree instances
  */
@@ -32,6 +33,7 @@ protected:
 	virtual bool HasDebugInfoInternal(const RewindDebugger::FObjectId& InObjectId) const override;
 	virtual bool IsCreatingPrimaryChildTrackInternal() const override;
 };
+#endif // WITH_ENGINE
 
 /**
  * RewindDebugger track representing a single StateTree instance
diff --git a/Engine/Source/Runtime/Core/Private/Features/ModularFeatures.cpp b/Engine/Source/Runtime/Core/Private/Features/ModularFeatures.cpp
--- a/Engine/Source/Runtime/Core/Private/Features/ModularFeatures.cpp
+++ b/Engine/Source/Runtime/Core/Private/Features/ModularFeatures.cpp
@@ -54,6 +54,9 @@ IModularFeature* FModularFeatures::GetModularFeatureImplementation( const FName
 
 void FModularFeatures::RegisterModularFeature( const FName Type, IModularFeature* ModularFeature )
 {
+	check(ModularFeature != nullptr);
+	check(Type.IsValid());
+
 	UE::TScopeLock ScopeLock(ModularFeaturesMapCriticalSection);
 
 	ModularFeaturesMap.AddUnique( Type, ModularFeature );
@@ -63,6 +66,9 @@ void FModularFeatures::RegisterModularFeature( const FName Type, IModularFeature
 
 void FModularFeatures::UnregisterModularFeature( const FName Type, IModularFeature* ModularFeature )
 {
+	check(ModularFeature != nullptr);
+	check(Type.IsValid());
+
 	UE::TScopeLock ScopeLock(ModularFeaturesMapCriticalSection);
 
 	ModularFeaturesMap.RemoveSingle( Type, ModularFeature );
