[Backout] - CL49859133
#rnx
Original CL Desc
-----------------------------------------------------------------
Fix missing debug tracks in rewind debugger for PoseSearch, SequenceInfo, and EvaluationTask when using UAF
Use strongly typed references to class names instead of hardcoded strings to make rewind debugger tracks more robust against class renames
Move UAF specific pose search debugger code to the UAFPoseSearch plugin

#Synced-CL 49767845

diff --git a/Engine/Plugins/Animation/GameplayInsights/Source/GameplayInsights/Private/AnimCurvesTrack.cpp b/Engine/Plugins/Animation/GameplayInsights/Source/GameplayInsights/Private/AnimCurvesTrack.cpp
--- a/Engine/Plugins/Animation/GameplayInsights/Source/GameplayInsights/Private/AnimCurvesTrack.cpp
+++ b/Engine/Plugins/Animation/GameplayInsights/Source/GameplayInsights/Private/AnimCurvesTrack.cpp
@@ -6,9 +6,6 @@
 #include "GameplayProvider.h"
 #include "AnimationProvider.h"
 #include "SAnimationCurvesView.h"
-#if WITH_ENGINE
-#include "Components/SkeletalMeshComponent.h"
-#endif // WITH_ENGINE
 
 #define LOCTEXT_NAMESPACE "AnimCurvesTrack"
 
@@ -214,11 +211,13 @@ TSharedPtr<SWidget> FAnimCurveTrack::GetTimelineViewInternal()
 		.CurveData_Raw(this, &FAnimCurveTrack::GetCurveData);
 }
 
-#if WITH_ENGINE
+
 static const FName AnimationCurvesName("AnimationCurves");
+
 FName FAnimationCurvesTrackCreator::GetTargetTypeNameInternal() const
 {
-	return USkeletalMeshComponent::StaticClass()->GetFName();
+	static const FName SkelMeshName("SkeletalMeshComponent");
+	return SkelMeshName;
 }
 	
 FName FAnimationCurvesTrackCreator::GetNameInternal() const
@@ -251,8 +250,8 @@ bool FAnimationCurvesTrackCreator::HasDebugInfoInternal(const FObjectId& InObjec
 	}
 	return bHasData;
 }
-#endif // WITH_ENGINE
 
+	
 }
 
 #undef LOCTEXT_NAMESPACE
\ No newline at end of file
diff --git a/Engine/Plugins/Animation/GameplayInsights/Source/GameplayInsights/Private/AnimCurvesTrack.h b/Engine/Plugins/Animation/GameplayInsights/Source/GameplayInsights/Private/AnimCurvesTrack.h
--- a/Engine/Plugins/Animation/GameplayInsights/Source/GameplayInsights/Private/AnimCurvesTrack.h
+++ b/Engine/Plugins/Animation/GameplayInsights/Source/GameplayInsights/Private/AnimCurvesTrack.h
@@ -67,7 +67,7 @@ private:
 	TArray<TSharedPtr<FAnimCurveTrack>> Children;
 };
 
-#if WITH_ENGINE
+	
 class FAnimationCurvesTrackCreator : public IRewindDebuggerTrackCreator
 {
 private:
@@ -77,6 +77,5 @@ private:
 	virtual TSharedPtr<FRewindDebuggerTrack> CreateTrackInternal(const FObjectId& InObjectId) const override;
 	virtual bool HasDebugInfoInternal(const FObjectId& InObjectId) const override;
 };
-#endif // WITH_ENGINE
 
 }
\ No newline at end of file
diff --git a/Engine/Plugins/Animation/GameplayInsights/Source/GameplayInsights/Private/BlendWeightsTrack.cpp b/Engine/Plugins/Animation/GameplayInsights/Source/GameplayInsights/Private/BlendWeightsTrack.cpp
--- a/Engine/Plugins/Animation/GameplayInsights/Source/GameplayInsights/Private/BlendWeightsTrack.cpp
+++ b/Engine/Plugins/Animation/GameplayInsights/Source/GameplayInsights/Private/BlendWeightsTrack.cpp
@@ -418,10 +418,10 @@ bool FBlendWeightTrack::HandleDoubleClickInternal()
 	return false;
 }
 
-#if WITH_ENGINE
 FName FBlendWeightsTrackCreator::GetTargetTypeNameInternal() const
 {
-	return UAnimInstance::StaticClass()->GetFName();
+	static const FName TargetTypeName = "AnimInstance";
+	return TargetTypeName;
 }
 	
 static const FName BlendWeightsName("BlendWeights");
@@ -457,7 +457,6 @@ bool FBlendWeightsTrackCreator::HasDebugInfoInternal(const FObjectId& InObjectId
 	
 	return bHasData;
 }
-#endif // #if WITH_ENGINE
 
 }
 #undef LOCTEXT_NAMESPACE
diff --git a/Engine/Plugins/Animation/GameplayInsights/Source/GameplayInsights/Private/BlendWeightsTrack.h b/Engine/Plugins/Animation/GameplayInsights/Source/GameplayInsights/Private/BlendWeightsTrack.h
--- a/Engine/Plugins/Animation/GameplayInsights/Source/GameplayInsights/Private/BlendWeightsTrack.h
+++ b/Engine/Plugins/Animation/GameplayInsights/Source/GameplayInsights/Private/BlendWeightsTrack.h
@@ -80,8 +80,7 @@ private:
 	uint64 ObjectId;
 	TArray<TSharedPtr<FBlendWeightTrack>> Children;
 };
-
-#if WITH_ENGINE
+	
 class FBlendWeightsTrackCreator : public IRewindDebuggerTrackCreator
 {
 public:
@@ -91,6 +90,5 @@ public:
 	virtual TSharedPtr<FRewindDebuggerTrack> CreateTrackInternal(const FObjectId& InObjectId) const override;
 	virtual bool HasDebugInfoInternal(const FObjectId& InObjectId) const override;
 };
-#endif
 
 }
\ No newline at end of file
diff --git a/Engine/Plugins/Animation/GameplayInsights/Source/GameplayInsights/Private/ExternalMorphTrack.cpp b/Engine/Plugins/Animation/GameplayInsights/Source/GameplayInsights/Private/ExternalMorphTrack.cpp
--- a/Engine/Plugins/Animation/GameplayInsights/Source/GameplayInsights/Private/ExternalMorphTrack.cpp
+++ b/Engine/Plugins/Animation/GameplayInsights/Source/GameplayInsights/Private/ExternalMorphTrack.cpp
@@ -6,9 +6,6 @@
 #include "GameplayProvider.h"
 #include "AnimationProvider.h"
 #include "SAnimationCurvesView.h"
-#if WITH_ENGINE
-#include "Components/SkeletalMeshComponent.h"
-#endif //WITH_ENGINE 
 
 #define LOCTEXT_NAMESPACE "ExternalMorphTrack"
 
@@ -265,10 +262,11 @@ namespace RewindDebugger
 	}
 
 	//----------------------------------------------------------------------------------------------
-#if WITH_ENGINE
+
 	FName FExternalMorphSetGroupTrackCreator::GetTargetTypeNameInternal() const
 	{
-		return USkeletalMeshComponent::StaticClass()->GetFName();
+		static const FName SkelMeshName("SkeletalMeshComponent");
+		return SkelMeshName;
 	}
 	
 	static const FName ExternalMorphSetName("ExternalMorphSetGroup");
@@ -298,7 +296,6 @@ namespace RewindDebugger
 		}
 		return false;
 	}
-#endif // WITH_ENGINE
 }
 
 #undef LOCTEXT_NAMESPACE
diff --git a/Engine/Plugins/Animation/GameplayInsights/Source/GameplayInsights/Private/ExternalMorphTrack.h b/Engine/Plugins/Animation/GameplayInsights/Source/GameplayInsights/Private/ExternalMorphTrack.h
--- a/Engine/Plugins/Animation/GameplayInsights/Source/GameplayInsights/Private/ExternalMorphTrack.h
+++ b/Engine/Plugins/Animation/GameplayInsights/Source/GameplayInsights/Private/ExternalMorphTrack.h
@@ -97,7 +97,6 @@ namespace RewindDebugger
 	};
 
 
-#if WITH_ENGINE
 	// The creator for the external morph set track.
 	class FExternalMorphSetGroupTrackCreator
 		: public IRewindDebuggerTrackCreator
@@ -111,6 +110,5 @@ namespace RewindDebugger
 		virtual bool HasDebugInfoInternal(const FObjectId& InObjectId) const override;
 		// ~END IRewindDebuggerTrackCreator overrides.
 	};
-#endif // WITH_ENGINE
 
 }	// namespace RewindDebugger
\ No newline at end of file
diff --git a/Engine/Plugins/Animation/GameplayInsights/Source/GameplayInsights/Private/InertializationsTrack.cpp b/Engine/Plugins/Animation/GameplayInsights/Source/GameplayInsights/Private/InertializationsTrack.cpp
--- a/Engine/Plugins/Animation/GameplayInsights/Source/GameplayInsights/Private/InertializationsTrack.cpp
+++ b/Engine/Plugins/Animation/GameplayInsights/Source/GameplayInsights/Private/InertializationsTrack.cpp
@@ -246,10 +246,10 @@ bool FInertializationTrack::HandleDoubleClickInternal()
 	return false;
 }
 
-#if WITH_ENGINE
 FName FInertializationsTrackCreator::GetTargetTypeNameInternal() const
 {
-	return UAnimInstance::StaticClass()->GetFName();
+	static const FName TargetTypeName = "AnimInstance";
+	return TargetTypeName;
 }
 
 static const FName InertializationsName("Inertializations");
@@ -284,7 +284,6 @@ bool FInertializationsTrackCreator::HasDebugInfoInternal(const FObjectId& InObje
 
 	return false;
 }
-#endif // WITH_ENGINE
 
 }
 #undef LOCTEXT_NAMESPACE
diff --git a/Engine/Plugins/Animation/GameplayInsights/Source/GameplayInsights/Private/InertializationsTrack.h b/Engine/Plugins/Animation/GameplayInsights/Source/GameplayInsights/Private/InertializationsTrack.h
--- a/Engine/Plugins/Animation/GameplayInsights/Source/GameplayInsights/Private/InertializationsTrack.h
+++ b/Engine/Plugins/Animation/GameplayInsights/Source/GameplayInsights/Private/InertializationsTrack.h
@@ -61,8 +61,7 @@ private:
 	uint64 ObjectId = INDEX_NONE;
 	TArray<TSharedPtr<FInertializationTrack>> Children;
 };
-
-#if WITH_ENGINE
+	
 class FInertializationsTrackCreator : public IRewindDebuggerTrackCreator
 {
 public:
@@ -72,6 +71,5 @@ public:
 	virtual TSharedPtr<FRewindDebuggerTrack> CreateTrackInternal(const FObjectId& InObjectId) const override;
 	virtual bool HasDebugInfoInternal(const FObjectId& InObjectId) const override;
 };
-#endif // WITH_ENGINE
 
 }
\ No newline at end of file
diff --git a/Engine/Plugins/Animation/GameplayInsights/Source/GameplayInsights/Private/MontagesTrack.cpp b/Engine/Plugins/Animation/GameplayInsights/Source/GameplayInsights/Private/MontagesTrack.cpp
--- a/Engine/Plugins/Animation/GameplayInsights/Source/GameplayInsights/Private/MontagesTrack.cpp
+++ b/Engine/Plugins/Animation/GameplayInsights/Source/GameplayInsights/Private/MontagesTrack.cpp
@@ -376,18 +376,13 @@ bool FMontageTrack::HandleDoubleClickInternal()
 	return false;
 }
 
-static const FName MontagesName("Montages");
-FName FMontagesTrack::GetNameInternal() const
-{
-	return MontagesName;
-}
-
-#if WITH_ENGINE
 FName FMontagesTrackCreator::GetTargetTypeNameInternal() const
 {
-	return UAnimInstance::StaticClass()->GetFName();
+	static const FName TargetTypeName = "AnimInstance";
+	return TargetTypeName;
 }
 
+static const FName MontagesName("Montages");
 FName FMontagesTrackCreator::GetNameInternal() const
 {
 	return MontagesName;
@@ -397,6 +392,11 @@ void FMontagesTrackCreator::GetTrackTypesInternal(TArray<FRewindDebuggerTrackTyp
 {
 	Types.Add({MontagesName, LOCTEXT("Montages", "Montages")});
 }
+	
+FName FMontagesTrack::GetNameInternal() const
+{
+	return MontagesName;
+}
 
 TSharedPtr<RewindDebugger::FRewindDebuggerTrack> FMontagesTrackCreator::CreateTrackInternal(const FObjectId& InObjectId) const
 {
@@ -419,7 +419,6 @@ bool FMontagesTrackCreator::HasDebugInfoInternal(const FObjectId& InObjectId) co
 	
 	return bHasData;
 }
-#endif // WITH_ENGINE
 
 }
 #undef LOCTEXT_NAMESPACE
diff --git a/Engine/Plugins/Animation/GameplayInsights/Source/GameplayInsights/Private/MontagesTrack.h b/Engine/Plugins/Animation/GameplayInsights/Source/GameplayInsights/Private/MontagesTrack.h
--- a/Engine/Plugins/Animation/GameplayInsights/Source/GameplayInsights/Private/MontagesTrack.h
+++ b/Engine/Plugins/Animation/GameplayInsights/Source/GameplayInsights/Private/MontagesTrack.h
@@ -73,8 +73,7 @@ private:
 	uint64 ObjectId;
 	TArray<TSharedPtr<FMontageTrack>> Children;
 };
-
-#if WITH_ENGINE
+	
 class FMontagesTrackCreator : public IRewindDebuggerTrackCreator
 {
 public:
@@ -84,6 +83,5 @@ public:
 	virtual TSharedPtr<FRewindDebuggerTrack> CreateTrackInternal(const FObjectId& InObjectId) const override;
 	virtual bool HasDebugInfoInternal(const FObjectId& InObjectId) const override;
 };
-#endif // WITH_ENGINE
 
 }
\ No newline at end of file
diff --git a/Engine/Plugins/Animation/GameplayInsights/Source/GameplayInsights/Private/NotifiesTrack.cpp b/Engine/Plugins/Animation/GameplayInsights/Source/GameplayInsights/Private/NotifiesTrack.cpp
--- a/Engine/Plugins/Animation/GameplayInsights/Source/GameplayInsights/Private/NotifiesTrack.cpp
+++ b/Engine/Plugins/Animation/GameplayInsights/Source/GameplayInsights/Private/NotifiesTrack.cpp
@@ -6,9 +6,6 @@
 #include "GameplayProvider.h"
 #include "AnimationProvider.h"
 #include "SNotifiesView.h"
-#if WITH_ENGINE
-#include "Animation/AnimInstance.h"
-#endif // WITH_ENGINE
 
 #define LOCTEXT_NAMESPACE "NotifiesTrack"
 
@@ -154,10 +151,11 @@ TSharedPtr<SWidget> FNotifyTrack::GetTimelineViewInternal()
 }
 
 static const FName NotifiesName("Notifies");
-#if WITH_ENGINE
+
 FName FNotifiesTrackCreator::GetTargetTypeNameInternal() const
 {
-	return UAnimInstance::StaticClass()->GetFName();
+	static const FName AnimInstanceName("AnimInstance");
+	return AnimInstanceName;
 }
 	
 FName FNotifiesTrackCreator::GetNameInternal() const
@@ -170,11 +168,12 @@ void FNotifiesTrackCreator::GetTrackTypesInternal(TArray<FRewindDebuggerTrackTyp
 	Types.Add({ NotifiesName, LOCTEXT("Notifies", "Notifies") });
 }
 
+
 TSharedPtr<RewindDebugger::FRewindDebuggerTrack> FNotifiesTrackCreator::CreateTrackInternal(const FObjectId& InObjectId) const
 {
 	return MakeShared<RewindDebugger::FNotifiesTrack>(InObjectId.GetMainId());
 }
-#endif // WITH_ENGINE
+
 		
 FNotifiesTrack::FNotifiesTrack(uint64 InObjectId) :
 	ObjectId(InObjectId)
@@ -262,7 +261,6 @@ TConstArrayView<TSharedPtr<FRewindDebuggerTrack>> FNotifiesTrack::GetChildrenInt
 	return MakeConstArrayView<TSharedPtr<FRewindDebuggerTrack>>(reinterpret_cast<const TSharedPtr<FRewindDebuggerTrack>*>(Children.GetData()), Children.Num());
 }
 
-#if WITH_ENGINE
 bool FNotifiesTrackCreator::HasDebugInfoInternal(const FObjectId& InObjectId) const
 {
 	TRACE_CPUPROFILER_EVENT_SCOPE(FNotifiesTrack::HasDebugInfoInternal);
@@ -279,7 +277,7 @@ bool FNotifiesTrackCreator::HasDebugInfoInternal(const FObjectId& InObjectId) co
 	}
 	return bHasData;
 }
-#endif // WITH_ENGINE
+
 	
 }
 
diff --git a/Engine/Plugins/Animation/GameplayInsights/Source/GameplayInsights/Private/NotifiesTrack.h b/Engine/Plugins/Animation/GameplayInsights/Source/GameplayInsights/Private/NotifiesTrack.h
--- a/Engine/Plugins/Animation/GameplayInsights/Source/GameplayInsights/Private/NotifiesTrack.h
+++ b/Engine/Plugins/Animation/GameplayInsights/Source/GameplayInsights/Private/NotifiesTrack.h
@@ -95,7 +95,7 @@ private:
 	TArray<TSharedPtr<FNotifyTrack>> Children;
 };
 
-#if WITH_ENGINE
+
 class FNotifiesTrackCreator : public IRewindDebuggerTrackCreator
 {
 private:
@@ -105,6 +105,5 @@ private:
 	virtual TSharedPtr<FRewindDebuggerTrack> CreateTrackInternal(const FObjectId& InObjectId) const override;
 	virtual bool HasDebugInfoInternal(const FObjectId& InObjectId) const override;
 };
-#endif // WITH_ENGINE
 
 }
\ No newline at end of file
diff --git a/Engine/Plugins/Animation/GameplayInsights/Source/GameplayInsights/Private/PoseWatchTrack.cpp b/Engine/Plugins/Animation/GameplayInsights/Source/GameplayInsights/Private/PoseWatchTrack.cpp
--- a/Engine/Plugins/Animation/GameplayInsights/Source/GameplayInsights/Private/PoseWatchTrack.cpp
+++ b/Engine/Plugins/Animation/GameplayInsights/Source/GameplayInsights/Private/PoseWatchTrack.cpp
@@ -240,7 +240,8 @@ TConstArrayView<TSharedPtr<FRewindDebuggerTrack>> FPoseWatchTrack::GetChildrenIn
 
 FName FPoseWatchesTrackCreator::GetTargetTypeNameInternal() const
 {
-	return UAnimInstance::StaticClass()->GetFName();
+	static const FName AnimInstanceName("AnimInstance");
+	return AnimInstanceName;
 }
 	
 static const FName PoseWatchesName("PoseWatches");
diff --git a/Engine/Plugins/Animation/GameplayInsights/Source/GameplayInsights/Private/PoseWatchTrack.h b/Engine/Plugins/Animation/GameplayInsights/Source/GameplayInsights/Private/PoseWatchTrack.h
--- a/Engine/Plugins/Animation/GameplayInsights/Source/GameplayInsights/Private/PoseWatchTrack.h
+++ b/Engine/Plugins/Animation/GameplayInsights/Source/GameplayInsights/Private/PoseWatchTrack.h
@@ -83,7 +83,7 @@ private:
 	TArray<TSharedPtr<FPoseWatchTrack>> Children;
 };
 
-#if WITH_ENGINE
+
 class FPoseWatchesTrackCreator : public IRewindDebuggerTrackCreator
 {
 private:
@@ -93,7 +93,6 @@ private:
 	virtual TSharedPtr<FRewindDebuggerTrack> CreateTrackInternal(const FObjectId& InObjectId) const override;
 	virtual bool HasDebugInfoInternal(const FObjectId& InObjectId) const override;
 };
-#endif // WITH_ENGINE
 
 }
 
diff --git a/Engine/Plugins/Animation/GameplayInsights/Source/GameplayInsights/Private/PropertiesTrack.cpp b/Engine/Plugins/Animation/GameplayInsights/Source/GameplayInsights/Private/PropertiesTrack.cpp
--- a/Engine/Plugins/Animation/GameplayInsights/Source/GameplayInsights/Private/PropertiesTrack.cpp
+++ b/Engine/Plugins/Animation/GameplayInsights/Source/GameplayInsights/Private/PropertiesTrack.cpp
@@ -285,10 +285,11 @@ namespace RewindDebugger
 	///////////////////////////////////////////////////////////////////////////////
 	// Properties Track Creator
 	///////////////////////////////////////////////////////////////////////////////
-#if WITH_ENGINE
+	
 	FName FPropertiesTrackCreator::GetTargetTypeNameInternal() const
 	{
-		return UObject::StaticClass()->GetFName();
+		static const FName TargetTypeName("Object");
+		return TargetTypeName;
 	}
 
 	static const FName PropertiesTrackCreatorName("PropertyWatch");
@@ -332,7 +333,6 @@ namespace RewindDebugger
 
 		return bHasData;
 	}
-#endif // WITH_ENGINE
 		
 }
 
diff --git a/Engine/Plugins/Animation/GameplayInsights/Source/GameplayInsights/Private/PropertiesTrack.h b/Engine/Plugins/Animation/GameplayInsights/Source/GameplayInsights/Private/PropertiesTrack.h
--- a/Engine/Plugins/Animation/GameplayInsights/Source/GameplayInsights/Private/PropertiesTrack.h
+++ b/Engine/Plugins/Animation/GameplayInsights/Source/GameplayInsights/Private/PropertiesTrack.h
@@ -66,8 +66,7 @@ namespace RewindDebugger
 		/** Count of recorded properties for traced object at the start of recording */
 		int64 TotalPropertyCountAtStart;
 	};
-
-#if WITH_ENGINE
+	
 	/**
 	 * Used to automatically create an FPropertiesTrack in Rewind Debugger
 	 */
@@ -81,5 +80,4 @@ namespace RewindDebugger
 		virtual bool HasDebugInfoInternal(const FObjectId& InObjectId) const override;
 		/** End IRewindDebuggerTrackCreator interface */
 	};
-#endif // WITH_ENGINE
 }
\ No newline at end of file
diff --git a/Engine/Plugins/Animation/GameplayInsights/Source/GameplayInsights/Private/SAnimGraphSchematicView.cpp b/Engine/Plugins/Animation/GameplayInsights/Source/GameplayInsights/Private/SAnimGraphSchematicView.cpp
--- a/Engine/Plugins/Animation/GameplayInsights/Source/GameplayInsights/Private/SAnimGraphSchematicView.cpp
+++ b/Engine/Plugins/Animation/GameplayInsights/Source/GameplayInsights/Private/SAnimGraphSchematicView.cpp
@@ -946,10 +946,10 @@ void SAnimGraphSchematicView::RefreshDetails(const TArray<TSharedRef<FAnimGraphS
 	}
 }
 
-#if WITH_ENGINE
 FName FAnimGraphSchematicTrackCreator::GetTargetTypeNameInternal() const
 {
-	return UAnimInstance::StaticClass()->GetFName();
+	static FName TargetTypeName = "AnimInstance";
+	return TargetTypeName;
 }
 
 FName FAnimGraphSchematicTrackCreator::GetNameInternal() const
@@ -966,7 +966,6 @@ TSharedPtr<RewindDebugger::FRewindDebuggerTrack> FAnimGraphSchematicTrackCreator
 {
 	return MakeShared<FAnimGraphSchematicTrack>(InObjectId.GetMainId());
 }
-#endif // WITH_ENGINE
 
 FText FAnimGraphSchematicTrack::GetDisplayNameInternal() const
 {
@@ -1018,7 +1017,6 @@ bool FAnimGraphSchematicTrack::UpdateInternal()
 	return false;
 }
 
-#if WITH_ENGINE
 bool FAnimGraphSchematicTrackCreator::HasDebugInfoInternal(const RewindDebugger::FObjectId& InObjectId) const
 {
 	const TraceServices::IAnalysisSession* AnalysisSession = IRewindDebugger::Instance()->GetAnalysisSession();
@@ -1035,6 +1033,5 @@ bool FAnimGraphSchematicTrackCreator::HasDebugInfoInternal(const RewindDebugger:
 	
 	return bHasData;
 }
-#endif // WITH_ENGINE
 
 #undef LOCTEXT_NAMESPACE
diff --git a/Engine/Plugins/Animation/GameplayInsights/Source/GameplayInsights/Private/SAnimGraphSchematicView.h b/Engine/Plugins/Animation/GameplayInsights/Source/GameplayInsights/Private/SAnimGraphSchematicView.h
--- a/Engine/Plugins/Animation/GameplayInsights/Source/GameplayInsights/Private/SAnimGraphSchematicView.h
+++ b/Engine/Plugins/Animation/GameplayInsights/Source/GameplayInsights/Private/SAnimGraphSchematicView.h
@@ -139,7 +139,6 @@ private:
 	uint64 AnimInstanceId;
 };
 
-#if WITH_ENGINE
 class FAnimGraphSchematicTrackCreator : public RewindDebugger::IRewindDebuggerTrackCreator
 {
 public:
@@ -148,5 +147,4 @@ public:
 	virtual void GetTrackTypesInternal(TArray<RewindDebugger::FRewindDebuggerTrackType>& Types) const override;
 	virtual TSharedPtr<RewindDebugger::FRewindDebuggerTrack> CreateTrackInternal(const RewindDebugger::FObjectId& InObjectId) const override;
 	virtual bool HasDebugInfoInternal(const RewindDebugger::FObjectId& InObjectId) const override;
-};
-#endif // WITH_ENGINE
+};
\ No newline at end of file
diff --git a/Engine/Plugins/Animation/GameplayInsights/Source/RewindDebuggerVLog/Private/VisualLogTrack.cpp b/Engine/Plugins/Animation/GameplayInsights/Source/RewindDebuggerVLog/Private/VisualLogTrack.cpp
--- a/Engine/Plugins/Animation/GameplayInsights/Source/RewindDebuggerVLog/Private/VisualLogTrack.cpp
+++ b/Engine/Plugins/Animation/GameplayInsights/Source/RewindDebuggerVLog/Private/VisualLogTrack.cpp
@@ -286,7 +286,8 @@ TSharedPtr<SWidget> FVisualLogTrack::GetDetailsViewInternal()
 //----------------------------------------------------------------------//
 FName FVisualLogTrackCreator::GetTargetTypeNameInternal() const
 {
-	return UObject::StaticClass()->GetFName();
+	static const FName TargetTypeName("Object");
+	return TargetTypeName;
 }
 
 FName FVisualLogTrackCreator::GetNameInternal() const
diff --git a/Engine/Plugins/Animation/GameplayInsights/Source/RewindDebuggerVLog/Private/VisualLogTrack.h b/Engine/Plugins/Animation/GameplayInsights/Source/RewindDebuggerVLog/Private/VisualLogTrack.h
--- a/Engine/Plugins/Animation/GameplayInsights/Source/RewindDebuggerVLog/Private/VisualLogTrack.h
+++ b/Engine/Plugins/Animation/GameplayInsights/Source/RewindDebuggerVLog/Private/VisualLogTrack.h
@@ -95,7 +95,6 @@ private:
 	TArray<TSharedPtr<FVisualLogCategoryTrack>> Children;
 };
 
-#if WITH_ENGINE
 /**
  * Track creator for VisualLogger recorded entries
  */
@@ -107,6 +106,5 @@ class FVisualLogTrackCreator : public IRewindDebuggerTrackCreator
 	virtual TSharedPtr<FRewindDebuggerTrack> CreateTrackInternal(const FObjectId& InObjectId) const override;
 	virtual bool HasDebugInfoInternal(const FObjectId& InObjectId) const override;
 };
-#endif // WITH_ENGINE
 
 }
diff --git a/Engine/Plugins/Animation/PoseSearch/Source/Editor/Private/PoseSearchDebugger.cpp b/Engine/Plugins/Animation/PoseSearch/Source/Editor/Private/PoseSearchDebugger.cpp
--- a/Engine/Plugins/Animation/PoseSearch/Source/Editor/Private/PoseSearchDebugger.cpp
+++ b/Engine/Plugins/Animation/PoseSearch/Source/Editor/Private/PoseSearchDebugger.cpp
@@ -556,6 +556,44 @@ TConstArrayView<TSharedPtr<RewindDebugger::FRewindDebuggerTrack>> FDebuggerTrack
 	return MakeConstArrayView<TSharedPtr<FRewindDebuggerTrack>>(reinterpret_cast<const TSharedPtr<FRewindDebuggerTrack>*>(SearchTracks.GetData()), SearchTracks.Num());
 }
 
+// FDebuggerTrackCreator
+///////////////////////////////////////////////////
+
+void FDebuggerTrackCreator::GetTrackTypesInternal(TArray<RewindDebugger::FRewindDebuggerTrackType>& Types) const
+{
+	Types.Add({ GetNameInternal(), LOCTEXT("Pose Search", "Pose Search") });
+}
+
+TSharedPtr<RewindDebugger::FRewindDebuggerTrack> FDebuggerTrackCreator::CreateTrackInternal(const RewindDebugger::FObjectId& InObjectId) const
+{
+	return MakeShared<FDebuggerTrack>(InObjectId.GetMainId());
+}
+
+bool FDebuggerTrackCreator::HasDebugInfoInternal(const RewindDebugger::FObjectId& InObjectId) const
+{
+	TRACE_CPUPROFILER_EVENT_SCOPE(PoseSearchDebugger::HasDebugInfoInternal);
+	// Get provider and validate
+	const TraceServices::IAnalysisSession* Session = IRewindDebugger::Instance()->GetAnalysisSession();
+	TraceServices::FAnalysisSessionReadScope SessionReadScope(*Session);
+
+	const FTraceProvider* PoseSearchProvider = Session->ReadProvider<FTraceProvider>(FTraceProvider::ProviderName);
+	const IAnimationProvider* AnimationProvider = Session->ReadProvider<IAnimationProvider>("AnimationProvider");
+	const IGameplayProvider* GameplayProvider = Session->ReadProvider<IGameplayProvider>("GameplayProvider");
+	if (!(PoseSearchProvider && AnimationProvider && GameplayProvider))
+	{
+		return false;
+	}
+
+	bool bHasData = false;
+
+	PoseSearchProvider->EnumerateMotionMatchingStateTimelines(InObjectId.GetMainId(), [&bHasData](const FTraceProvider::FMotionMatchingStateTimeline& InTimeline)
+	{
+		bHasData = true;
+	});
+
+	return bHasData;
+}
+
 } // namespace UE::PoseSearch
 
 #undef LOCTEXT_NAMESPACE
diff --git a/Engine/Plugins/Animation/PoseSearch/Source/Editor/Private/PoseSearchDebugger.h b/Engine/Plugins/Animation/PoseSearch/Source/Editor/Private/PoseSearchDebugger.h
--- a/Engine/Plugins/Animation/PoseSearch/Source/Editor/Private/PoseSearchDebugger.h
+++ b/Engine/Plugins/Animation/PoseSearch/Source/Editor/Private/PoseSearchDebugger.h
@@ -4,6 +4,7 @@
 
 #include "CoreMinimal.h"
 #include "IRewindDebuggerExtension.h"
+#include "IRewindDebuggerTrackCreator.h"
 #include "PoseSearch/PoseSearchDefines.h"
 #include "RewindDebuggerTrack.h"
 #include "SCurveTimelineView.h"
@@ -108,5 +109,21 @@ private:
 	TArray<TSharedPtr<FSearchTrack>> SearchTracks;
 };
 
+class FDebuggerTrackCreator : public RewindDebugger::IRewindDebuggerTrackCreator
+{
+private:
+	virtual FName GetTargetTypeNameInternal() const override { return "AnimInstance"; }
+	virtual FName GetNameInternal() const override { return "PoseSearchDebugger"; }
+	virtual void GetTrackTypesInternal(TArray<RewindDebugger::FRewindDebuggerTrackType>& Types) const override;
+	virtual TSharedPtr<RewindDebugger::FRewindDebuggerTrack> CreateTrackInternal(const RewindDebugger::FObjectId& InObjectId) const override;
+	virtual bool HasDebugInfoInternal(const RewindDebugger::FObjectId& InObjectId) const override;
+	virtual int32 GetSortOrderPriorityInternal() const override { return 10; };
+};
+
+class FAnimNextDebuggerTrackCreator : public FDebuggerTrackCreator
+{
+private:
+	virtual FName GetTargetTypeNameInternal() const override { return "AnimNextComponent"; }
+};
 
 } // namespace UE::PoseSearch
diff --git a/Engine/Plugins/Animation/PoseSearch/Source/Editor/Private/PoseSearchDebuggerTrackCreator.cpp b/Engine/Plugins/Animation/PoseSearch/Source/Editor/Private/PoseSearchDebuggerTrackCreator.cpp
deleted file mode 100644
diff --git a/Engine/Plugins/Animation/PoseSearch/Source/Editor/Private/PoseSearchEditor.cpp b/Engine/Plugins/Animation/PoseSearch/Source/Editor/Private/PoseSearchEditor.cpp
--- a/Engine/Plugins/Animation/PoseSearch/Source/Editor/Private/PoseSearchEditor.cpp
+++ b/Engine/Plugins/Animation/PoseSearch/Source/Editor/Private/PoseSearchEditor.cpp
@@ -20,7 +20,6 @@
 #include "Trace/PoseSearchTraceAnalyzer.h"
 #include "Trace/PoseSearchTraceModule.h"
 #include "PoseSearchColumnEditor.h"
-#include "PoseSearchDebuggerTrackCreator.h"
 
 DEFINE_LOG_CATEGORY(LogPoseSearchEditor);
 
@@ -42,6 +41,7 @@ private:
 	
 	/** Creates the track for the Rewind Debugger */
 	TSharedPtr<FDebuggerTrackCreator> DebuggerTrackCreator;
+	TSharedPtr<FAnimNextDebuggerTrackCreator> AnimNextDebuggerTrackCreator;
 
 	/** Enables dedicated PoseSearch trace module */
 	TSharedPtr<FTraceModule> TraceModule;
@@ -68,8 +68,10 @@ void FEditorModule::StartupModule()
 		FDebugger::Initialize();
 		TraceModule = MakeShared<FTraceModule>();
 		DebuggerTrackCreator = MakeShared<FDebuggerTrackCreator>();
+		AnimNextDebuggerTrackCreator = MakeShared<FAnimNextDebuggerTrackCreator>();
 		
 		IModularFeatures::Get().RegisterModularFeature(FDebuggerTrackCreator::ModularFeatureName, DebuggerTrackCreator.Get());
+		IModularFeatures::Get().RegisterModularFeature(FDebuggerTrackCreator::ModularFeatureName, AnimNextDebuggerTrackCreator.Get());
 		IModularFeatures::Get().RegisterModularFeature(TraceServices::ModuleFeatureName, TraceModule.Get());
 		
 		// Register Ed Modes used by PoseSearchDatabase and PoseSearchInteractionAsset
@@ -101,16 +103,8 @@ void FEditorModule::ShutdownModule()
 
 	// Unregister Asset Editor Commands
 	FDatabaseEditorCommands::Unregister();
-
-	if (DebuggerTrackCreator.Get() != nullptr)
-	{
-		IModularFeatures::Get().UnregisterModularFeature(FDebuggerTrackCreator::ModularFeatureName, DebuggerTrackCreator.Get());
-	}
-
-	if (TraceModule.Get() != nullptr)
-	{
-		IModularFeatures::Get().UnregisterModularFeature(TraceServices::ModuleFeatureName, TraceModule.Get());
-	}
+	
+	IModularFeatures::Get().UnregisterModularFeature(TraceServices::ModuleFeatureName, TraceModule.Get());
 	FDebugger::Shutdown();
 }
 
diff --git a/Engine/Plugins/Animation/PoseSearch/Source/Editor/Public/PoseSearchDebuggerTrackCreator.h b/Engine/Plugins/Animation/PoseSearch/Source/Editor/Public/PoseSearchDebuggerTrackCreator.h
deleted file mode 100644
diff --git a/Engine/Plugins/Cameras/GameplayCameras/Source/GameplayCamerasEditor/Private/Trace/CameraSystemRewindDebuggerTrack.h b/Engine/Plugins/Cameras/GameplayCameras/Source/GameplayCamerasEditor/Private/Trace/CameraSystemRewindDebuggerTrack.h
--- a/Engine/Plugins/Cameras/GameplayCameras/Source/GameplayCamerasEditor/Private/Trace/CameraSystemRewindDebuggerTrack.h
+++ b/Engine/Plugins/Cameras/GameplayCameras/Source/GameplayCamerasEditor/Private/Trace/CameraSystemRewindDebuggerTrack.h
@@ -53,7 +53,6 @@ private:
 	TSharedPtr<FCameraSystemRewindDebuggerTrackTimelineData> TimelineData;
 };
 
-#if WITH_ENGINE
 /**
  * Factory class that creates the rewind debugger track for the camera system evaluation trace.
  */
@@ -70,7 +69,6 @@ protected:
 	virtual TSharedPtr<RewindDebugger::FRewindDebuggerTrack> CreateTrackInternal(uint64 InObjectId) const override;
 #endif
 };
-#endif // WITH_ENGINE
 
 }  // namespace UE::Cameras
 
diff --git a/Engine/Plugins/Chooser/Source/ChooserEditor/Private/ChooserTrack.cpp b/Engine/Plugins/Chooser/Source/ChooserEditor/Private/ChooserTrack.cpp
--- a/Engine/Plugins/Chooser/Source/ChooserEditor/Private/ChooserTrack.cpp
+++ b/Engine/Plugins/Chooser/Source/ChooserEditor/Private/ChooserTrack.cpp
@@ -163,9 +163,10 @@ TSharedPtr<SWidget> FChooserTrack::GetTimelineViewInternal()
 
 static const FName ChoosersName("Choosers");
 
-FName FChoosersTrackCreator::GetTargetTypeNameInternal() const
-{
-	return UObject::StaticClass()->GetFName();
+ FName FChoosersTrackCreator::GetTargetTypeNameInternal() const
+ {
+ 	static const FName ObjectName("Object");
+	return ObjectName;
 }
 	
 FName FChoosersTrackCreator::GetNameInternal() const
diff --git a/Engine/Plugins/Chooser/Source/ChooserEditor/Private/ChooserTrack.h b/Engine/Plugins/Chooser/Source/ChooserEditor/Private/ChooserTrack.h
--- a/Engine/Plugins/Chooser/Source/ChooserEditor/Private/ChooserTrack.h
+++ b/Engine/Plugins/Chooser/Source/ChooserEditor/Private/ChooserTrack.h
@@ -63,7 +63,7 @@ private:
 	TArray<TSharedPtr<FChooserTrack>> Children;
 };
 
-#if WITH_ENGINE
+
 class FChoosersTrackCreator : public RewindDebugger::IRewindDebuggerTrackCreator
 {
 private:
@@ -73,6 +73,5 @@ private:
 	virtual TSharedPtr<RewindDebugger::FRewindDebuggerTrack> CreateTrackInternal(const RewindDebugger::FObjectId& InObjectId) const override;
 	virtual bool HasDebugInfoInternal(const RewindDebugger::FObjectId& InObjectId) const override;
 };
-#endif // WITH_ENGINE
 
 }
\ No newline at end of file
diff --git a/Engine/Plugins/Developer/FieldNotificationTrace/Source/FieldNotificationTraceEditor/Private/FieldNotificationTraceEditorModule.h b/Engine/Plugins/Developer/FieldNotificationTrace/Source/FieldNotificationTraceEditor/Private/FieldNotificationTraceEditorModule.h
--- a/Engine/Plugins/Developer/FieldNotificationTrace/Source/FieldNotificationTraceEditor/Private/FieldNotificationTraceEditorModule.h
+++ b/Engine/Plugins/Developer/FieldNotificationTrace/Source/FieldNotificationTraceEditor/Private/FieldNotificationTraceEditorModule.h
@@ -27,9 +27,7 @@ public:
 private:
 	FRewindDebuggerExtension RewindDebuggerExtension;
 	FRewindDebuggerRuntimeExtension RewindDebuggerRuntimeExtension;
-#if WITH_ENGINE
 	FTracksCreator TrackCreator;
-#endif // WITH_ENGINE 
 	FTraceServiceModule TraceModule;
 };
 
diff --git a/Engine/Plugins/Developer/FieldNotificationTrace/Source/FieldNotificationTraceEditor/Private/FieldNotificationTrack.cpp b/Engine/Plugins/Developer/FieldNotificationTrace/Source/FieldNotificationTraceEditor/Private/FieldNotificationTrack.cpp
--- a/Engine/Plugins/Developer/FieldNotificationTrace/Source/FieldNotificationTraceEditor/Private/FieldNotificationTrack.cpp
+++ b/Engine/Plugins/Developer/FieldNotificationTrace/Source/FieldNotificationTraceEditor/Private/FieldNotificationTrack.cpp
@@ -228,7 +228,8 @@ TConstArrayView<TSharedPtr<RewindDebugger::FRewindDebuggerTrack>> FObjectTrack::
 
 FName FTracksCreator::GetTargetTypeNameInternal() const
 {
-	return UObject::StaticClass()->GetFName();
+	static const FName ObjectName("Object");
+	return ObjectName;
 }
 
 FName FTracksCreator::GetNameInternal() const
diff --git a/Engine/Plugins/Developer/FieldNotificationTrace/Source/FieldNotificationTraceEditor/Private/FieldNotificationTrack.h b/Engine/Plugins/Developer/FieldNotificationTrace/Source/FieldNotificationTraceEditor/Private/FieldNotificationTrack.h
--- a/Engine/Plugins/Developer/FieldNotificationTrace/Source/FieldNotificationTraceEditor/Private/FieldNotificationTrack.h
+++ b/Engine/Plugins/Developer/FieldNotificationTrace/Source/FieldNotificationTraceEditor/Private/FieldNotificationTrack.h
@@ -2,9 +2,7 @@
 #pragma once
 
 #include "FieldNotificationId.h"
-#if WITH_ENGINE
 #include "IRewindDebuggerTrackCreator.h"
-#endif // WITH_ENGINE
 #include "RewindDebuggerTrack.h"
 #include "Textures/SlateIcon.h"
 #include "SEventTimelineView.h"
@@ -66,7 +64,6 @@ private:
 };
 
 
-#if WITH_ENGINE
 class FTracksCreator : public RewindDebugger::IRewindDebuggerTrackCreator
 {
 private:
@@ -76,6 +73,5 @@ private:
 	virtual TSharedPtr<RewindDebugger::FRewindDebuggerTrack> CreateTrackInternal(const RewindDebugger::FObjectId& InObjectId) const override;
 	virtual bool HasDebugInfoInternal(const RewindDebugger::FObjectId& InObjectId) const override;
 };
-#endif // WITH_ENGINE
 
 } // namespace
diff --git a/Engine/Plugins/Experimental/UAF/UAF/Source/UAFEditor/Private/RewindDebugger/AnimNextModuleTrack.cpp b/Engine/Plugins/Experimental/UAF/UAF/Source/UAFEditor/Private/RewindDebugger/AnimNextModuleTrack.cpp
--- a/Engine/Plugins/Experimental/UAF/UAF/Source/UAFEditor/Private/RewindDebugger/AnimNextModuleTrack.cpp
+++ b/Engine/Plugins/Experimental/UAF/UAF/Source/UAFEditor/Private/RewindDebugger/AnimNextModuleTrack.cpp
@@ -14,7 +14,6 @@
 #include "Modules/ModuleManager.h"
 #include "ObjectAsTraceIdProxyArchiveReader.h"
 #include "PropertyEditorModule.h"
-#include "UAFAssetInstance.h"
 #include "Serialization/MemoryReader.h"
 #include "StructUtils/PropertyBag.h"
 #include "Subsystems/AssetEditorSubsystem.h"
@@ -28,7 +27,8 @@ static const FName AnimNextModulesName("AnimNextModules");
 
 FName FAnimNextModuleTrackCreator::GetTargetTypeNameInternal() const
 {
-	return FUAFAssetInstance::StaticStruct()->GetFName();
+	static const FName ObjectName("UAFAssetInstance");
+	return ObjectName;
 }
 
 FText FAnimNextModuleTrack::GetDisplayNameInternal() const
diff --git a/Engine/Plugins/Experimental/UAF/UAF/Source/UAFEditor/Private/RewindDebugger/AnimNextModuleTrack.h b/Engine/Plugins/Experimental/UAF/UAF/Source/UAFEditor/Private/RewindDebugger/AnimNextModuleTrack.h
--- a/Engine/Plugins/Experimental/UAF/UAF/Source/UAFEditor/Private/RewindDebugger/AnimNextModuleTrack.h
+++ b/Engine/Plugins/Experimental/UAF/UAF/Source/UAFEditor/Private/RewindDebugger/AnimNextModuleTrack.h
@@ -79,7 +79,7 @@ private:
 	TArray<FPropertyBagPropertyDesc> PropertyDescriptions;
 };
 
-#if WITH_ENGINE
+
 class FAnimNextModuleTrackCreator : public RewindDebugger::IRewindDebuggerTrackCreator
 {
 private:
@@ -99,7 +99,6 @@ private:
 		return true;
 	}
 };
-#endif // WITH_ENGINE
 
 }
 #endif // UAF_TRACE_ENABLED
diff --git a/Engine/Plugins/Experimental/UAF/UAFAnimGraph/Source/UAFAnimGraphEditor/Private/RewindDebugger/EvaluationProgramTrack.cpp b/Engine/Plugins/Experimental/UAF/UAFAnimGraph/Source/UAFAnimGraphEditor/Private/RewindDebugger/EvaluationProgramTrack.cpp
--- a/Engine/Plugins/Experimental/UAF/UAFAnimGraph/Source/UAFAnimGraphEditor/Private/RewindDebugger/EvaluationProgramTrack.cpp
+++ b/Engine/Plugins/Experimental/UAF/UAFAnimGraph/Source/UAFAnimGraphEditor/Private/RewindDebugger/EvaluationProgramTrack.cpp
@@ -10,7 +10,6 @@
 #include "Modules/ModuleManager.h"
 #include "PropertyEditorModule.h"
 #include "ObjectAsTraceIdProxyArchiveReader.h"
-#include "Component/AnimNextComponent.h"
 #include "Serialization/ObjectReader.h"
 
 #include UE_INLINE_GENERATED_CPP_BY_NAME(EvaluationProgramTrack)
@@ -25,7 +24,8 @@ static const FName AnimNextModulesName("AnimNextModules");
 
  FName FEvaluationProgramTrackCreator::GetTargetTypeNameInternal() const
  {
-	return UUAFComponent::StaticClass()->GetFName();
+ 	static const FName ObjectName("AnimNextComponent");
+	return ObjectName;
 }
 
 	
diff --git a/Engine/Plugins/Experimental/UAF/UAFAnimGraph/Source/UAFAnimGraphEditor/Private/RewindDebugger/EvaluationProgramTrack.h b/Engine/Plugins/Experimental/UAF/UAFAnimGraph/Source/UAFAnimGraphEditor/Private/RewindDebugger/EvaluationProgramTrack.h
--- a/Engine/Plugins/Experimental/UAF/UAFAnimGraph/Source/UAFAnimGraphEditor/Private/RewindDebugger/EvaluationProgramTrack.h
+++ b/Engine/Plugins/Experimental/UAF/UAFAnimGraph/Source/UAFAnimGraphEditor/Private/RewindDebugger/EvaluationProgramTrack.h
@@ -55,7 +55,7 @@ private:
 	// TArray<TSharedPtr<FEvaluationProgramTrack>> Children;
 };
 
-#if WITH_ENGINE
+
 class FEvaluationProgramTrackCreator : public RewindDebugger::IRewindDebuggerTrackCreator
 {
 private:
@@ -65,6 +65,5 @@ private:
 	virtual TSharedPtr<RewindDebugger::FRewindDebuggerTrack> CreateTrackInternal(const RewindDebugger::FObjectId& InObjectId) const override;
 	virtual bool HasDebugInfoInternal(const RewindDebugger::FObjectId& InObjectId) const override;
 };
-#endif // WITH_ENGINE
 
 }
\ No newline at end of file
diff --git a/Engine/Plugins/Experimental/UAF/UAFAnimGraph/Source/UAFAnimGraphEditor/Private/RewindDebugger/SequenceInfoTrack.cpp b/Engine/Plugins/Experimental/UAF/UAFAnimGraph/Source/UAFAnimGraphEditor/Private/RewindDebugger/SequenceInfoTrack.cpp
--- a/Engine/Plugins/Experimental/UAF/UAFAnimGraph/Source/UAFAnimGraphEditor/Private/RewindDebugger/SequenceInfoTrack.cpp
+++ b/Engine/Plugins/Experimental/UAF/UAFAnimGraph/Source/UAFAnimGraphEditor/Private/RewindDebugger/SequenceInfoTrack.cpp
@@ -12,7 +12,6 @@
 #include "Animation/AnimSequence.h"
 #include "EvaluationVM/Tasks/PushAnimSequenceKeyframe.h"
 #include "ObjectAsTraceIdProxyArchiveReader.h"
-#include "Component/AnimNextComponent.h"
 #include "Widgets/SCanvas.h"
 #include "Widgets/Text/STextBlock.h"
 #include "Widgets/Colors/SColorBlock.h"
@@ -25,7 +24,8 @@ namespace UE::UAF::Editor
 {
 FName FSequenceInfoTrackCreator::GetTargetTypeNameInternal() const
 {
-	return UUAFComponent::StaticClass()->GetFName();
+	static const FName ObjectName("AnimNextComponent");
+	return ObjectName;
 }
 
 
diff --git a/Engine/Plugins/Experimental/UAF/UAFAnimGraph/Source/UAFAnimGraphEditor/Private/RewindDebugger/SequenceInfoTrack.h b/Engine/Plugins/Experimental/UAF/UAFAnimGraph/Source/UAFAnimGraphEditor/Private/RewindDebugger/SequenceInfoTrack.h
--- a/Engine/Plugins/Experimental/UAF/UAFAnimGraph/Source/UAFAnimGraphEditor/Private/RewindDebugger/SequenceInfoTrack.h
+++ b/Engine/Plugins/Experimental/UAF/UAFAnimGraph/Source/UAFAnimGraphEditor/Private/RewindDebugger/SequenceInfoTrack.h
@@ -107,7 +107,7 @@ private:
 	// TArray<TSharedPtr<FEvaluationProgramTrack>> Children;
 };
 
-#if WITH_ENGINE
+
 class FSequenceInfoTrackCreator : public RewindDebugger::IRewindDebuggerTrackCreator
 {
 private:
@@ -117,6 +117,5 @@ private:
 	virtual TSharedPtr<RewindDebugger::FRewindDebuggerTrack> CreateTrackInternal(const RewindDebugger::FObjectId& InObjectId) const override;
 	virtual bool HasDebugInfoInternal(const RewindDebugger::FObjectId& InObjectId) const override;
 };
-#endif // WITH_ENGINE
 
 }
\ No newline at end of file
diff --git a/Engine/Plugins/Experimental/UAF/UAFPoseSearch/Source/UAFPoseSearchUncookedOnly/Private/UAFPoseSearchDebugger.cpp b/Engine/Plugins/Experimental/UAF/UAFPoseSearch/Source/UAFPoseSearchUncookedOnly/Private/UAFPoseSearchDebugger.cpp
deleted file mode 100644
diff --git a/Engine/Plugins/Experimental/UAF/UAFPoseSearch/Source/UAFPoseSearchUncookedOnly/Private/UAFPoseSearchDebugger.h b/Engine/Plugins/Experimental/UAF/UAFPoseSearch/Source/UAFPoseSearchUncookedOnly/Private/UAFPoseSearchDebugger.h
deleted file mode 100644
diff --git a/Engine/Plugins/Experimental/UAF/UAFPoseSearch/Source/UAFPoseSearchUncookedOnly/Private/UAFPoseSearchUncookedOnlyModule.cpp b/Engine/Plugins/Experimental/UAF/UAFPoseSearch/Source/UAFPoseSearchUncookedOnly/Private/UAFPoseSearchUncookedOnlyModule.cpp
--- a/Engine/Plugins/Experimental/UAF/UAFPoseSearch/Source/UAFPoseSearchUncookedOnly/Private/UAFPoseSearchUncookedOnlyModule.cpp
+++ b/Engine/Plugins/Experimental/UAF/UAFPoseSearch/Source/UAFPoseSearchUncookedOnly/Private/UAFPoseSearchUncookedOnlyModule.cpp
@@ -1,44 +1,5 @@
 ﻿// Copyright Epic Games, Inc. All Rights Reserved.
 
-#include "PoseSearchDebuggerTrackCreator.h"
-#include "UAFPoseSearchDebugger.h"
-#include "Features/IModularFeatures.h"
 #include "Modules/ModuleManager.h"
 
-namespace UE::PoseSearch
-{
-class FUAFDebuggerTrackCreator;
-
-//////////////////////////////////////////////////////////////////////////
-// FPoseSearchEditorModule
-class FUAFPoseSearchUncookedOnlyModule : public IModuleInterface
-{
-public:
-	void StartupModule() override;
-	void ShutdownModule() override;
-
-private:
-	
-	/** Creates the track for the Rewind Debugger */
-	TSharedPtr<FUAFDebuggerTrackCreator> UAFDebuggerTrackCreator;
-};
-
-void FUAFPoseSearchUncookedOnlyModule::StartupModule()
-{
-	IModuleInterface::StartupModule();
-
-	UAFDebuggerTrackCreator = MakeShared<FUAFDebuggerTrackCreator>();
-	IModularFeatures::Get().RegisterModularFeature(RewindDebugger::IRewindDebuggerTrackCreator::ModularFeatureName, UAFDebuggerTrackCreator.Get());
-}
-
-void FUAFPoseSearchUncookedOnlyModule::ShutdownModule()
-{
-	IModularFeatures::Get().UnregisterModularFeature(RewindDebugger::IRewindDebuggerTrackCreator::ModularFeatureName, UAFDebuggerTrackCreator.Get());
-	
-	IModuleInterface::ShutdownModule();
-}
-
-} // namespace UE::PoseSearch
-
-IMPLEMENT_MODULE(UE::PoseSearch::FUAFPoseSearchUncookedOnlyModule, UAFPoseSearchUncookedOnly)
-
+IMPLEMENT_MODULE(FDefaultModuleImpl, UAFPoseSearchUncookedOnly)
\ No newline at end of file
diff --git a/Engine/Plugins/Experimental/UAF/UAFPoseSearch/Source/UAFPoseSearchUncookedOnly/UAFPoseSearchUncookedOnly.Build.cs b/Engine/Plugins/Experimental/UAF/UAFPoseSearch/Source/UAFPoseSearchUncookedOnly/UAFPoseSearchUncookedOnly.Build.cs
--- a/Engine/Plugins/Experimental/UAF/UAFPoseSearch/Source/UAFPoseSearchUncookedOnly/UAFPoseSearchUncookedOnly.Build.cs
+++ b/Engine/Plugins/Experimental/UAF/UAFPoseSearch/Source/UAFPoseSearchUncookedOnly/UAFPoseSearchUncookedOnly.Build.cs
@@ -20,9 +20,7 @@ namespace UnrealBuildTool.Rules
 					"RigVMDeveloper",
 					"RigVM",
 					"Slate",
-					"SlateCore", 
-					"RewindDebuggerInterface", 
-					"PoseSearchEditor",
+					"SlateCore",
 				}
 			);
 		}
diff --git a/Engine/Plugins/Runtime/RigVM/Source/RigVMEditor/Private/Trace/RigVMTraceTrack.cpp b/Engine/Plugins/Runtime/RigVM/Source/RigVMEditor/Private/Trace/RigVMTraceTrack.cpp
--- a/Engine/Plugins/Runtime/RigVM/Source/RigVMEditor/Private/Trace/RigVMTraceTrack.cpp
+++ b/Engine/Plugins/Runtime/RigVM/Source/RigVMEditor/Private/Trace/RigVMTraceTrack.cpp
@@ -4,7 +4,6 @@
 #include "RigVMTraceProvider.h"
 #include "RigVMTraceAnalyzer.h"
 #include "DetailLayoutBuilder.h"
-#include "RigVMHost.h"
 #include "Widgets/Layout/SGridPanel.h"
 #include "Widgets/Input/SNumericEntryBox.h"
 
@@ -263,7 +262,8 @@ TConstArrayView<TSharedPtr<RewindDebugger::FRewindDebuggerTrack>> FRewindDebugge
 
 FName FRewindDebuggerRigVMTrackCreator::GetTargetTypeNameInternal() const
 {
-	return URigVMHost::StaticClass()->GetFName();
+	static const FLazyName RigVMHostName("RigVMHost");
+	return RigVMHostName;
 }
 	
 FName FRewindDebuggerRigVMTrackCreator::GetNameInternal() const
diff --git a/Engine/Plugins/Runtime/RigVM/Source/RigVMEditor/Public/Trace/RigVMTraceTrack.h b/Engine/Plugins/Runtime/RigVM/Source/RigVMEditor/Public/Trace/RigVMTraceTrack.h
--- a/Engine/Plugins/Runtime/RigVM/Source/RigVMEditor/Public/Trace/RigVMTraceTrack.h
+++ b/Engine/Plugins/Runtime/RigVM/Source/RigVMEditor/Public/Trace/RigVMTraceTrack.h
@@ -83,7 +83,6 @@ private:
 	TArray<TSharedPtr<FRewindDebuggerTrack>> Children;
 };
 
-#if WITH_ENGINE
 /**
  * Factory class to create RigVM rewind debugger tracks
  */
@@ -96,7 +95,6 @@ private:
 	virtual TSharedPtr<RewindDebugger::FRewindDebuggerTrack> CreateTrackInternal(const RewindDebugger::FObjectId& InObjectId) const override;
 	virtual bool HasDebugInfoInternal(const RewindDebugger::FObjectId& InObjectId) const override;
 };
-#endif // WITH_ENGINE
 
 #undef UE_API
 
diff --git a/Engine/Plugins/Runtime/StateTree/Source/StateTreeEditorModule/Private/Debugger/StateTreeRewindDebuggerTrack.cpp b/Engine/Plugins/Runtime/StateTree/Source/StateTreeEditorModule/Private/Debugger/StateTreeRewindDebuggerTrack.cpp
--- a/Engine/Plugins/Runtime/StateTree/Source/StateTreeEditorModule/Private/Debugger/StateTreeRewindDebuggerTrack.cpp
+++ b/Engine/Plugins/Runtime/StateTree/Source/StateTreeEditorModule/Private/Debugger/StateTreeRewindDebuggerTrack.cpp
@@ -56,7 +56,8 @@ bool FRewindDebuggerTrackCreator::IsCreatingPrimaryChildTrackInternal() const
 
 FName FRewindDebuggerTrackCreator::GetTargetTypeNameInternal() const
 {
-	return FStateTreeInstanceData::StaticStruct()->GetFName();
+	static const FName AssociatedObjectName("StateTreeInstanceData");
+	return AssociatedObjectName;
 }
 
 //----------------------------------------------------------------------//
diff --git a/Engine/Plugins/Runtime/StateTree/Source/StateTreeEditorModule/Private/Debugger/StateTreeRewindDebuggerTrack.h b/Engine/Plugins/Runtime/StateTree/Source/StateTreeEditorModule/Private/Debugger/StateTreeRewindDebuggerTrack.h
--- a/Engine/Plugins/Runtime/StateTree/Source/StateTreeEditorModule/Private/Debugger/StateTreeRewindDebuggerTrack.h
+++ b/Engine/Plugins/Runtime/StateTree/Source/StateTreeEditorModule/Private/Debugger/StateTreeRewindDebuggerTrack.h
@@ -16,7 +16,6 @@ namespace UE::StateTreeDebugger
 {
 class SFrameEventsView;
 
-#if WITH_ENGINE
 /**
  * RewindDebugger track creator for StateTree instances
  */
@@ -33,7 +32,6 @@ protected:
 	virtual bool HasDebugInfoInternal(const RewindDebugger::FObjectId& InObjectId) const override;
 	virtual bool IsCreatingPrimaryChildTrackInternal() const override;
 };
-#endif // WITH_ENGINE
 
 /**
  * RewindDebugger track representing a single StateTree instance
diff --git a/Engine/Source/Runtime/Core/Private/Features/ModularFeatures.cpp b/Engine/Source/Runtime/Core/Private/Features/ModularFeatures.cpp
--- a/Engine/Source/Runtime/Core/Private/Features/ModularFeatures.cpp
+++ b/Engine/Source/Runtime/Core/Private/Features/ModularFeatures.cpp
@@ -54,9 +54,6 @@ IModularFeature* FModularFeatures::GetModularFeatureImplementation( const FName
 
 void FModularFeatures::RegisterModularFeature( const FName Type, IModularFeature* ModularFeature )
 {
-	check(ModularFeature != nullptr);
-	check(Type.IsValid());
-
 	UE::TScopeLock ScopeLock(ModularFeaturesMapCriticalSection);
 
 	ModularFeaturesMap.AddUnique( Type, ModularFeature );
@@ -66,9 +63,6 @@ void FModularFeatures::RegisterModularFeature( const FName Type, IModularFeature
 
 void FModularFeatures::UnregisterModularFeature( const FName Type, IModularFeature* ModularFeature )
 {
-	check(ModularFeature != nullptr);
-	check(Type.IsValid());
-
 	UE::TScopeLock ScopeLock(ModularFeaturesMapCriticalSection);
 
 	ModularFeaturesMap.RemoveSingle( Type, ModularFeature );
