PoseSearch - logging warnings in UPoseSearchLibrary::MotionMatch if databases AsyncBuildIndex is in still in progress or failed

#Synced-CL 49859538

diff --git a/Engine/Plugins/Animation/PoseSearch/Source/Runtime/Private/PoseSearchLibrary.cpp b/Engine/Plugins/Animation/PoseSearch/Source/Runtime/Private/PoseSearchLibrary.cpp
--- a/Engine/Plugins/Animation/PoseSearch/Source/Runtime/Private/PoseSearchLibrary.cpp
+++ b/Engine/Plugins/Animation/PoseSearch/Source/Runtime/Private/PoseSearchLibrary.cpp
@@ -1371,6 +1371,13 @@ void UPoseSearchLibrary::MotionMatch(
 		SearchContext.SetAssetsToConsiderSet(CurrentAssetsToConsider);
 	}
 	SearchContext.SetAssetIndexesToConsiderPerDatabase(CurrentAssetIndexesToConsiderPerDatabase);
+
+#if WITH_EDITOR && !NO_LOGGING
+	if (SearchContext.IsAsyncBuildIndexInProgress())
+	{
+		UE_LOG(LogPoseSearch, Warning, TEXT("UPoseSearchLibrary::MotionMatch - some searches have been skipped, since databases AsyncBuildIndex are in still in progress or failed"));
+	}
+#endif // WITH_EDITOR && !NO_LOGGING
 }
 
 const FAnimNode_PoseSearchHistoryCollector_Base* UPoseSearchLibrary::FindPoseHistoryNode(
