PoseSearch - fix for Crash when Undo adding Animations to Pose Search Interaction Asset

#Synced-CL 50072177

diff --git a/Engine/Plugins/Animation/PoseSearch/Source/Editor/Private/PoseSearchInteractionAssetEditor.cpp b/Engine/Plugins/Animation/PoseSearch/Source/Editor/Private/PoseSearchInteractionAssetEditor.cpp
--- a/Engine/Plugins/Animation/PoseSearch/Source/Editor/Private/PoseSearchInteractionAssetEditor.cpp
+++ b/Engine/Plugins/Animation/PoseSearch/Source/Editor/Private/PoseSearchInteractionAssetEditor.cpp
@@ -341,101 +341,107 @@ void FInteractionAssetViewModel::Tick(float DeltaSeconds)
 		FRoleToIndex PreviewActorRoleToIndex;
 		PreviewActorRoleToIndex.Reserve(PreviewActors.Num());
 
-		TArray<TObjectPtr<const UObject>, TInlineAllocator<PreallocatedRolesNum>> AnimContexts;
-		AnimContexts.SetNum(NumRoles);
-
-		TArray<FRole, TInlineAllocator<PreallocatedRolesNum>> Roles;
-		Roles.SetNum(NumRoles);
-
-		for (int32 PreviewActorIndex = 0; PreviewActorIndex < PreviewActors.Num(); ++PreviewActorIndex)
+		if (PreviewActors.Num() == NumRoles)
 		{
-			PreviewActorRoleToIndex.Add(PreviewActors[PreviewActorIndex].GetRole()) = PreviewActorIndex;
+			TArray<TObjectPtr<const UObject>, TInlineAllocator<PreallocatedRolesNum>> AnimContexts;
+			AnimContexts.SetNum(NumRoles);
 
-			FInteractionAssetPreviewActor& PreviewActor = PreviewActors[PreviewActorIndex];
-			AnimContexts[PreviewActorIndex] = PreviewActor.GetAnimPreviewInstance();
-			Roles[PreviewActorIndex] = PreviewActor.GetRole();
-		}
+			TArray<FRole, TInlineAllocator<PreallocatedRolesNum>> Roles;
+			Roles.SetNum(NumRoles);
 
-		// testing CalculateWarpTransforms
-#if WITH_EDITORONLY_DATA
-		if (InteractionAsset->bEnableDebugWarp && PreviewActors.Num() == NumRoles && PreviewActorRoleToIndex.Num() == NumRoles)
-		{
-			TArray<FTransform, TInlineAllocator<PreallocatedRolesNum>> ActorTransforms;
-			ActorTransforms.SetNum(NumRoles);
-			for (int32 InteractionAssetRoleIndex = 0; InteractionAssetRoleIndex < NumRoles; ++InteractionAssetRoleIndex)
+			for (int32 PreviewActorIndex = 0; PreviewActorIndex < PreviewActors.Num(); ++PreviewActorIndex)
 			{
-				const FRole& InteractionAssetRole = InteractionAsset->GetRole(InteractionAssetRoleIndex);
-				const int32 PreviewActorIndex = PreviewActorRoleToIndex[InteractionAssetRole];
+				PreviewActorRoleToIndex.Add(PreviewActors[PreviewActorIndex].GetRole()) = PreviewActorIndex;
 
 				FInteractionAssetPreviewActor& PreviewActor = PreviewActors[PreviewActorIndex];
-				check(PreviewActor.GetRole() == InteractionAssetRole);
-
-				ActorTransforms[InteractionAssetRoleIndex] = PreviewActor.GetDebugActorTransformFromSampler();
+				AnimContexts[PreviewActorIndex] = PreviewActor.GetAnimPreviewInstance();
+				Roles[PreviewActorIndex] = PreviewActor.GetRole();
 			}
 
-			TArray<FTransform, TInlineAllocator<PreallocatedRolesNum>> FullAlignedActorTransforms;
-			FullAlignedActorTransforms.SetNum(NumRoles);
-			const TConstArrayView<const UMirrorDataTable*> MirrorDataTables;
+			if (PreviewActorRoleToIndex.Num() == NumRoles)
+			{
+				// testing CalculateWarpTransforms
+#if WITH_EDITORONLY_DATA
+				if (InteractionAsset->bEnableDebugWarp)
+				{
+					TArray<FTransform, TInlineAllocator<PreallocatedRolesNum>> ActorTransforms;
+					ActorTransforms.SetNum(NumRoles);
+					for (int32 InteractionAssetRoleIndex = 0; InteractionAssetRoleIndex < NumRoles; ++InteractionAssetRoleIndex)
+					{
+						const FRole& InteractionAssetRole = InteractionAsset->GetRole(InteractionAssetRoleIndex);
+						const int32 PreviewActorIndex = PreviewActorRoleToIndex[InteractionAssetRole];
+
+						FInteractionAssetPreviewActor& PreviewActor = PreviewActors[PreviewActorIndex];
+						check(PreviewActor.GetRole() == InteractionAssetRole);
+
+						ActorTransforms[InteractionAssetRoleIndex] = PreviewActor.GetDebugActorTransformFromSampler();
+					}
+
+					TArray<FTransform, TInlineAllocator<PreallocatedRolesNum>> FullAlignedActorTransforms;
+					FullAlignedActorTransforms.SetNum(NumRoles);
+					const TConstArrayView<const UMirrorDataTable*> MirrorDataTables;
 
-			InteractionAsset->CalculateWarpTransforms(PlayTime, 0.f, ActorTransforms, FullAlignedActorTransforms, MirrorDataTables, AnimContexts);
+					InteractionAsset->CalculateWarpTransforms(PlayTime, 0.f, ActorTransforms, FullAlignedActorTransforms, MirrorDataTables, AnimContexts);
 
 #if ENABLE_ANIM_DEBUG
-			const UWorld* DebugWorld = GetWorld();
+					const UWorld* DebugWorld = GetWorld();
 #endif // ENABLE_ANIM_DEBUG
 
-			for (int32 InteractionAssetRoleIndex = 0; InteractionAssetRoleIndex < NumRoles; ++InteractionAssetRoleIndex)
-			{
-				const FRole& InteractionAssetRole = InteractionAsset->GetRole(InteractionAssetRoleIndex);
-				const int32 PreviewActorIndex = PreviewActorRoleToIndex[InteractionAssetRole];
+					for (int32 InteractionAssetRoleIndex = 0; InteractionAssetRoleIndex < NumRoles; ++InteractionAssetRoleIndex)
+					{
+						const FRole& InteractionAssetRole = InteractionAsset->GetRole(InteractionAssetRoleIndex);
+						const int32 PreviewActorIndex = PreviewActorRoleToIndex[InteractionAssetRole];
 
-				FInteractionAssetPreviewActor& PreviewActor = PreviewActors[PreviewActorIndex];
-				check(PreviewActor.GetRole() == InteractionAssetRole);
-				FTransform DebugActorTransform;
-				DebugActorTransform.Blend(ActorTransforms[InteractionAssetRoleIndex], FullAlignedActorTransforms[InteractionAssetRoleIndex], InteractionAsset->DebugWarpAmount);
+						FInteractionAssetPreviewActor& PreviewActor = PreviewActors[PreviewActorIndex];
+						check(PreviewActor.GetRole() == InteractionAssetRole);
+						FTransform DebugActorTransform;
+						DebugActorTransform.Blend(ActorTransforms[InteractionAssetRoleIndex], FullAlignedActorTransforms[InteractionAssetRoleIndex], InteractionAsset->DebugWarpAmount);
 
 #if ENABLE_ANIM_DEBUG
-				// FullAlignedActorTransforms drawing 
-				const FVector FullAlignedActorPosition = FullAlignedActorTransforms[InteractionAssetRoleIndex].GetTranslation();
-				const FQuat FullAlignedActorRotation = FullAlignedActorTransforms[InteractionAssetRoleIndex].GetRotation();
-				DrawDebugCircle(DebugWorld, FMatrix(FVector::ZAxisVector, FVector::XAxisVector, FVector::YAxisVector, FullAlignedActorPosition), 5.f, 16, FColorList::Yellow, false, 0.f, SDPG_Foreground, 0.f, false);
-				DrawDebugLine(DebugWorld, FullAlignedActorPosition, FullAlignedActorPosition + FullAlignedActorRotation.RotateVector(FVector::XAxisVector) * 10, FColorList::Yellow, false, 0.f, SDPG_Foreground, 0.f);
+						// FullAlignedActorTransforms drawing 
+						const FVector FullAlignedActorPosition = FullAlignedActorTransforms[InteractionAssetRoleIndex].GetTranslation();
+						const FQuat FullAlignedActorRotation = FullAlignedActorTransforms[InteractionAssetRoleIndex].GetRotation();
+						DrawDebugCircle(DebugWorld, FMatrix(FVector::ZAxisVector, FVector::XAxisVector, FVector::YAxisVector, FullAlignedActorPosition), 5.f, 16, FColorList::Yellow, false, 0.f, SDPG_Foreground, 0.f, false);
+						DrawDebugLine(DebugWorld, FullAlignedActorPosition, FullAlignedActorPosition + FullAlignedActorRotation.RotateVector(FVector::XAxisVector) * 10, FColorList::Yellow, false, 0.f, SDPG_Foreground, 0.f);
 #endif // ENABLE_ANIM_DEBUG
 
-				PreviewActor.ForceDebugActorTransform(DebugActorTransform);
-			}
+						PreviewActor.ForceDebugActorTransform(DebugActorTransform);
+					}
 
 #if ENABLE_ANIM_DEBUG
-			if (InteractionAsset->DebugTimeOffset > UE_KINDA_SMALL_NUMBER)
-			{
-				InteractionAsset->CalculateWarpTransforms(PlayTime, InteractionAsset->DebugTimeOffset, ActorTransforms, FullAlignedActorTransforms, MirrorDataTables, AnimContexts);
+					if (InteractionAsset->DebugTimeOffset > UE_KINDA_SMALL_NUMBER)
+					{
+						InteractionAsset->CalculateWarpTransforms(PlayTime, InteractionAsset->DebugTimeOffset, ActorTransforms, FullAlignedActorTransforms, MirrorDataTables, AnimContexts);
 
-				for (int32 InteractionAssetRoleIndex = 0; InteractionAssetRoleIndex < NumRoles; ++InteractionAssetRoleIndex)
-				{
-					// FullAlignedActorTransforms drawing 
-					const FVector FullAlignedActorPosition = FullAlignedActorTransforms[InteractionAssetRoleIndex].GetTranslation();
-					const FQuat FullAlignedActorRotation = FullAlignedActorTransforms[InteractionAssetRoleIndex].GetRotation();
-					DrawDebugCircle(DebugWorld, FMatrix(FVector::ZAxisVector, FVector::XAxisVector, FVector::YAxisVector, FullAlignedActorPosition), 5.f, 16, FColorList::Pink, false, 0.f, SDPG_Foreground, 0.f, false);
-					DrawDebugLine(DebugWorld, FullAlignedActorPosition, FullAlignedActorPosition + FullAlignedActorRotation.RotateVector(FVector::XAxisVector) * 10, FColorList::Pink, false, 0.f, SDPG_Foreground, 0.f);
-				}
-			}
+						for (int32 InteractionAssetRoleIndex = 0; InteractionAssetRoleIndex < NumRoles; ++InteractionAssetRoleIndex)
+						{
+							// FullAlignedActorTransforms drawing 
+							const FVector FullAlignedActorPosition = FullAlignedActorTransforms[InteractionAssetRoleIndex].GetTranslation();
+							const FQuat FullAlignedActorRotation = FullAlignedActorTransforms[InteractionAssetRoleIndex].GetRotation();
+							DrawDebugCircle(DebugWorld, FMatrix(FVector::ZAxisVector, FVector::XAxisVector, FVector::YAxisVector, FullAlignedActorPosition), 5.f, 16, FColorList::Pink, false, 0.f, SDPG_Foreground, 0.f, false);
+							DrawDebugLine(DebugWorld, FullAlignedActorPosition, FullAlignedActorPosition + FullAlignedActorRotation.RotateVector(FVector::XAxisVector) * 10, FColorList::Pink, false, 0.f, SDPG_Foreground, 0.f);
+						}
+					}
 #endif // ENABLE_ANIM_DEBUG
-		}
+				}
 #endif // WITH_EDITORONLY_DATA
 
-		// extracting and drawing Constraints
+				// extracting and drawing Constraints
 
-		// @todo: fill up those parameters
-		const FVector BlendParameters = FVector::ZeroVector;
-		const bool bIsMirrored = false;
+				// @todo: fill up those parameters
+				const FVector BlendParameters = FVector::ZeroVector;
+				const bool bIsMirrored = false;
 
-		FConstraints Constraints;
-		Constraints.ProcessAll(InteractionAsset, StepDeltaTime, PlayTime, bIsMirrored, BlendParameters, AnimContexts, Roles);
+				FConstraints Constraints;
+				Constraints.ProcessAll(InteractionAsset, StepDeltaTime, PlayTime, bIsMirrored, BlendParameters, AnimContexts, Roles);
 
-		if (const UWorld* World = GetWorld())
-		{
-			for (const FConstraint& Constraint : Constraints.GetAll())
-			{
-				DrawDebugLine(World, Constraint.FromSocketTransform.GetTranslation(), Constraint.ToSocketTransform.GetTranslation(), FColor::Red, false, -1.f, SDPG_Foreground);
+				if (const UWorld* World = GetWorld())
+				{
+					for (const FConstraint& Constraint : Constraints.GetAll())
+					{
+						DrawDebugLine(World, Constraint.FromSocketTransform.GetTranslation(), Constraint.ToSocketTransform.GetTranslation(), FColor::Red, false, -1.f, SDPG_Foreground);
+					}
+				}
 			}
 		}
 	}
