Fixed view dangling reference errors (when enabled).

diff --git a/Engine/Plugins/Animation/PoseSearch/Source/Runtime/Private/PoseSearchContext.cpp b/Engine/Plugins/Animation/PoseSearch/Source/Runtime/Private/PoseSearchContext.cpp
--- a/Engine/Plugins/Animation/PoseSearch/Source/Runtime/Private/PoseSearchContext.cpp
+++ b/Engine/Plugins/Animation/PoseSearch/Source/Runtime/Private/PoseSearchContext.cpp
@@ -906,7 +906,7 @@ void FSearchContext::UpdateContinuingPoseValues()
 	}
 	else
 	{
-		ContinuingPoseValues = TStackAlignedArray<float>();
+		ContinuingPoseValues = {};
 	}
 }
 
