PoseSearch - FInteractionIsland::AddTickDependencies improvements

#Synced-CL 50113163

diff --git a/Engine/Plugins/Animation/PoseSearch/Source/Runtime/Private/PoseSearchInteractionIsland.cpp b/Engine/Plugins/Animation/PoseSearch/Source/Runtime/Private/PoseSearchInteractionIsland.cpp
--- a/Engine/Plugins/Animation/PoseSearch/Source/Runtime/Private/PoseSearchInteractionIsland.cpp
+++ b/Engine/Plugins/Animation/PoseSearch/Source/Runtime/Private/PoseSearchInteractionIsland.cpp
@@ -27,11 +27,15 @@ static bool GVarPoseSearchInteractionDiagnoseTickDependencies = false;
 static FAutoConsoleVariableRef CVarPoseSearchInteractionDiagnoseTickDependencies(TEXT("a.PoseSearchInteraction.DiagnoseTickDependencies"), GVarPoseSearchInteractionDiagnoseTickDependencies, TEXT("Enable Pose Search Interaction Tick Dependencies Diagnostic (SLOW!)"));
 
 // recursion safe FTickFunction logging functions
-static void ShowPrerequistes(const FTickFunction& NestedTick, int32 Indent, int32 MaxIndent)
+static void ShowPrerequistes(FTickFunction& NestedTick, int32 Indent, TSet<FTickFunction*>& VisitedTickFunctions)
 {
-	if (Indent >= MaxIndent)
+	bool bAlreadyInSet = false;
+	VisitedTickFunctions.Add(&NestedTick, &bAlreadyInSet);
+
+	if (bAlreadyInSet)
 	{
-		UE_LOG(LogPoseSearch, Log, TEXT("%s==== REACHED MAX INDENT ===="), FCString::Spc(Indent * 2));
+		check(Indent > 0);
+		UE_LOG(LogPoseSearch, Log, TEXT("%s ERROR! Cycle detected! Use 'a.PoseSearchInteraction.DiagnoseTickDependencies' CVar to gather more infos to diagnose the problem"), FCString::Spc((Indent-1) * 2));
 	}
 	else
 	{
@@ -40,35 +44,31 @@ static void ShowPrerequistes(const FTickFunction& NestedTick, int32 Indent, int3
 			if (Prereq.PrerequisiteTickFunction)
 			{
 				UE_LOG(LogPoseSearch, Log, TEXT("%s prereq %s"), FCString::Spc(Indent * 2), *Prereq.PrerequisiteTickFunction->DiagnosticMessage());
-				ShowPrerequistes(*Prereq.PrerequisiteTickFunction, Indent + 1, MaxIndent);
+				ShowPrerequistes(*Prereq.PrerequisiteTickFunction, Indent + 1, VisitedTickFunctions);
 			}
 		}
+		VisitedTickFunctions.Remove(&NestedTick);
 	}
 }
 
-static void LogTickFunction(FTickFunction& Tick, ENamedThreads::Type CurrentThread, bool bLogPrerequisites, int32 Indent, int32 MaxIndent = 50)
+static void LogTickFunction(FTickFunction& Tick, ENamedThreads::Type CurrentThread, bool bLogPrerequisites, int32 Indent)
 {
-	if (Indent >= MaxIndent)
-	{
-		UE_LOG(LogPoseSearch, Log, TEXT("%s==== REACHED MAX INDENT ===="), FCString::Spc(Indent * 2));
-	}
-	else
+	// scoping brackets to save some heap for the recursion
 	{
-		// scoping brackets to save some heap for the recursion
+		UE_LOG(LogPoseSearch, Log, TEXT("%stick %s [%1d, %1d] %6llu %2d %s"), FCString::Spc(Indent * 2), Tick.bHighPriority ? TEXT("*") : TEXT(" "), (int32)Tick.GetActualTickGroup(), (int32)Tick.GetActualEndTickGroup(), (uint64)GFrameCounter, (int32)CurrentThread, *Tick.DiagnosticMessage());
+		if (bLogPrerequisites)
 		{
-			UE_LOG(LogPoseSearch, Log, TEXT("%stick %s [%1d, %1d] %6llu %2d %s"), FCString::Spc(Indent * 2), Tick.bHighPriority ? TEXT("*") : TEXT(" "), (int32)Tick.GetActualTickGroup(), (int32)Tick.GetActualEndTickGroup(), (uint64)GFrameCounter, (int32)CurrentThread, *Tick.DiagnosticMessage());
-			if (bLogPrerequisites)
-			{
-				ShowPrerequistes(Tick, Indent, MaxIndent);
-			}
+			TSet<FTickFunction*> VisitedTickFunctions;
+			ShowPrerequistes(Tick, Indent, VisitedTickFunctions);
+			check(VisitedTickFunctions.IsEmpty());
 		}
-
-		// Handle nested ticks
-		Tick.ForEachNestedTick([CurrentThread, bLogPrerequisites, Indent, MaxIndent](FTickFunction& NestedTick)
-			{
-				LogTickFunction(NestedTick, CurrentThread, bLogPrerequisites, Indent + 1, MaxIndent);
-			});
 	}
+
+	// Handle nested ticks
+	Tick.ForEachNestedTick([CurrentThread, bLogPrerequisites, Indent](FTickFunction& NestedTick)
+		{
+			LogTickFunction(NestedTick, CurrentThread, bLogPrerequisites, Indent + 1);
+		});
 }
 
 // check if there's any cycle within the prerequisites of Tick
@@ -81,7 +81,6 @@ static bool ValidateTickDependenciesCycles(FTickFunction& Tick, TSet<FTickFuncti
 
 	if (bAlreadyInSet)
 	{
-		UE_LOG(LogPoseSearch, Error, TEXT("ValidateTickDependencies: TickFunction '%s' form a cycle"), *Tick.DiagnosticMessage());
 		bValidatedCorrectly = false;
 	}
 	else
@@ -92,7 +91,6 @@ static bool ValidateTickDependenciesCycles(FTickFunction& Tick, TSet<FTickFuncti
 			{
 				if (!ValidateTickDependenciesCycles(*Prereq.PrerequisiteTickFunction, VisitedTickFunctions))
 				{
-					VisitedTickFunctions.Remove(&Tick);
 					bValidatedCorrectly = false;
 					break;
 				}
@@ -112,9 +110,9 @@ static bool ValidateTickDependenciesCycles(FTickFunction& Tick, TSet<FTickFuncti
 					}
 				});
 		}
+		VisitedTickFunctions.Remove(&Tick);
 	}
 	
-	VisitedTickFunctions.Remove(&Tick);
 	return bValidatedCorrectly;
 }
 #endif // ENABLE_ANIM_DEBUG
@@ -1083,14 +1081,22 @@ void FInteractionIsland::AddTickDependencies(UActorComponent* TickActorComponent
 		{
 			if (bInIsMainActor)
 			{
-				// PostTickFunction prerequisites should be empty since we haven't add the main actor tick function yet
+				// PostTickFunction and PreTickFunction prerequisites should be empty since we haven't add the main actor tick function yet
 				check(PostTickFunction.GetPrerequisites().IsEmpty());
+				check(PreTickFunction.GetPrerequisites().IsEmpty());
 
-				// adding to PreTickFunction all the tick dependencies TickActorComponent has, so it runs after all the tick dependencies of ALL the TickActorComponents in this FInteractionIsland
+				// adding to PreTickFunction all the tick dependencies TickActorComponent has (excluding the TickActor->PrimaryActorTick),
+				// so it runs after all the tick dependencies of ALL the TickActorComponents in this FInteractionIsland
 				for (const FTickPrerequisite& TickActorComponentPrerequisite : TickActorComponentTickFunction->GetPrerequisites())
 				{
 					check(TickActorComponentPrerequisite.PrerequisiteTickFunction);
-					AddPrerequisite(PreTickFunction, TickActorComponentPrerequisite.PrerequisiteObject.Get(), *TickActorComponentPrerequisite.PrerequisiteTickFunction);
+
+					// making sure we're not adding the component nor the actor tick functions to the PreTickFunction
+					check(TickActorComponentPrerequisite.PrerequisiteTickFunction != &TickActorComponent->PrimaryComponentTick);
+					if (TickActorComponentPrerequisite.PrerequisiteTickFunction != &TickActor->PrimaryActorTick)
+					{
+						AddPrerequisite(PreTickFunction, TickActorComponentPrerequisite.PrerequisiteObject.Get(), *TickActorComponentPrerequisite.PrerequisiteTickFunction);
+					}
 				}
 
 				// adding to TickActorComponent and owner actor TickActor dependencies to the PreTickFunction
@@ -1109,11 +1115,13 @@ void FInteractionIsland::AddTickDependencies(UActorComponent* TickActorComponent
 				// PostTickFunction should contain only the tick function to the main actor's ones (component and owner)
 				check(PostTickFunction.GetPrerequisites().Num() == 2);
 
-				// adding to PreTickFunction all the tick dependencies TickActorComponent has, so it runs after all the tick dependencies of ALL the TickActorComponents in this FInteractionIsland
+				// adding to PreTickFunction all the tick dependencies TickActorComponent has (excluding the TickActor->PrimaryActorTick),
+				// so it runs after all the tick dependencies of ALL the TickActorComponents in this FInteractionIsland
 				// BUT excluding the main actor tick fuctions that are in PostTickFunction.GetPrerequisites()
 				for (const FTickPrerequisite& TickActorComponentPrerequisite : TickActorComponentTickFunction->GetPrerequisites())
 				{
-					if (!PostTickFunction.GetPrerequisites().ContainsByPredicate([&TickActorComponentPrerequisite](const FTickPrerequisite& MainActorTickPrerequisite)
+					if (TickActorComponentPrerequisite.PrerequisiteTickFunction != &TickActor->PrimaryActorTick &&
+						!PostTickFunction.GetPrerequisites().ContainsByPredicate([&TickActorComponentPrerequisite](const FTickPrerequisite& MainActorTickPrerequisite)
 						{
 							return TickActorComponentPrerequisite.PrerequisiteTickFunction == MainActorTickPrerequisite.Get();
 						}))
@@ -1146,14 +1154,22 @@ void FInteractionIsland::AddTickDependencies(UActorComponent* TickActorComponent
 	{
 		if (bInIsMainActor)
 		{
-			// PostTickFunction prerequisites should be empty since we haven't add the main actor tick function yet
+			// PostTickFunction and PreTickFunction prerequisites should be empty since we haven't add the main actor tick function yet
 			check(PostTickFunction.GetPrerequisites().IsEmpty());
+			check(PreTickFunction.GetPrerequisites().IsEmpty());
 
-			// adding to PreTickFunction all the tick dependencies TickActorComponent has, so it runs after all the tick dependencies of ALL the TickActorComponents in this FInteractionIsland
+			// adding to PreTickFunction all the tick dependencies TickActorComponent has (excluding the TickActor->PrimaryActorTick),
+			// so it runs after all the tick dependencies of ALL the TickActorComponents in this FInteractionIsland
 			for (FTickPrerequisite& TickActorComponentPrerequisite : TickActorComponent->PrimaryComponentTick.GetPrerequisites())
 			{
 				check(TickActorComponentPrerequisite.PrerequisiteTickFunction);
-				AddPrerequisite(PreTickFunction, TickActorComponentPrerequisite.PrerequisiteObject.Get(), *TickActorComponentPrerequisite.PrerequisiteTickFunction);
+
+				// making sure we're not adding the component nor the actor tick functions to the PreTickFunction
+				check(TickActorComponentPrerequisite.PrerequisiteTickFunction != &TickActorComponent->PrimaryComponentTick);
+				if (TickActorComponentPrerequisite.PrerequisiteTickFunction != &TickActor->PrimaryActorTick)
+				{
+					AddPrerequisite(PreTickFunction, TickActorComponentPrerequisite.PrerequisiteObject.Get(), *TickActorComponentPrerequisite.PrerequisiteTickFunction);
+				}
 			}
 
 			// adding to TickActorComponent and owner actor TickActor dependencies to the PreTickFunction
@@ -1172,11 +1188,13 @@ void FInteractionIsland::AddTickDependencies(UActorComponent* TickActorComponent
 			// PostTickFunction should contain only the tick function to the main actor's ones (component and owner)
 			check(PostTickFunction.GetPrerequisites().Num() == 2);
 
-			// adding to PreTickFunction all the tick dependencies TickActorComponent has, so it runs after all the tick dependencies of ALL the TickActorComponents in this FInteractionIsland
+			// adding to PreTickFunction all the tick dependencies TickActorComponent has (excluding the TickActor->PrimaryActorTick),
+			// so it runs after all the tick dependencies of ALL the TickActorComponents in this FInteractionIsland
 			// BUT excluding the main actor tick fuctions that are in PostTickFunction.GetPrerequisites()
 			for (FTickPrerequisite& TickActorComponentPrerequisite : TickActorComponent->PrimaryComponentTick.GetPrerequisites())
 			{
-				if (!PostTickFunction.GetPrerequisites().ContainsByPredicate([&TickActorComponentPrerequisite](const FTickPrerequisite& MainActorTickPrerequisite)
+				if (TickActorComponentPrerequisite.PrerequisiteTickFunction != &TickActor->PrimaryActorTick &&
+					!PostTickFunction.GetPrerequisites().ContainsByPredicate([&TickActorComponentPrerequisite](const FTickPrerequisite& MainActorTickPrerequisite)
 					{
 						return TickActorComponentPrerequisite.PrerequisiteTickFunction == MainActorTickPrerequisite.Get();
 					}))
