PoseSearch - fix for SearchThrottleTime not working properly in FAnimNode_MotionMatching::UpdateAssetPlayer and FMotionMatchingTrait::PreUpdate

#Synced-CL 50261394

diff --git a/Engine/Plugins/Animation/PoseSearch/Source/Runtime/Private/AnimNode_MotionMatching.cpp b/Engine/Plugins/Animation/PoseSearch/Source/Runtime/Private/AnimNode_MotionMatching.cpp
--- a/Engine/Plugins/Animation/PoseSearch/Source/Runtime/Private/AnimNode_MotionMatching.cpp
+++ b/Engine/Plugins/Animation/PoseSearch/Source/Runtime/Private/AnimNode_MotionMatching.cpp
@@ -186,9 +186,16 @@ void FAnimNode_MotionMatching::UpdateAssetPlayer(const FAnimationUpdateContext&
 			SearchContext.AddRole(GetCommonDefaultRole(DatabasesToSearch), &ChooserEvaluationContext, PoseHistory);
 
 			// If we can't advance or enough time has elapsed since the last pose jump then search
-			const bool bCanAdvance = SearchContext.GetContinuingPoseSearchResult().PoseIdx != INDEX_NONE;
-			const bool bSearch = !bCanAdvance || (MotionMatchingState.ElapsedPoseSearchTime >= SearchThrottleTime);
+			bool bSearch = true;
 			FSearchResult SearchResult;
+			if (MotionMatchingState.ElapsedPoseSearchTime < SearchThrottleTime)
+			{
+				SearchResult.InitFrom(MotionMatchingState.SearchResult);
+				SearchContext.UpdateContinuingPoseSearchResult(SearchResult, SearchResult);
+				const bool bCanAdvance = SearchContext.GetContinuingPoseSearchResult().PoseIdx != INDEX_NONE;
+				bSearch = !bCanAdvance;
+			}
+
 			if (bSearch)
 			{
 				MotionMatchingState.ElapsedPoseSearchTime = 0.f;
@@ -210,7 +217,6 @@ void FAnimNode_MotionMatching::UpdateAssetPlayer(const FAnimationUpdateContext&
 				MotionMatchingState.ElapsedPoseSearchTime += DeltaTime;
 
 				// if we didn't perform a search we carry on with the continuing pose search result
-				SearchResult = SearchContext.GetContinuingPoseSearchResult();
 				SearchResult.bIsContinuingPoseSearch = true;
 
 #if UE_POSE_SEARCH_TRACE_ENABLED
diff --git a/Engine/Plugins/Animation/PoseSearch/Source/Runtime/Public/PoseSearch/PoseSearchResult.h b/Engine/Plugins/Animation/PoseSearch/Source/Runtime/Public/PoseSearch/PoseSearchResult.h
--- a/Engine/Plugins/Animation/PoseSearch/Source/Runtime/Public/PoseSearch/PoseSearchResult.h
+++ b/Engine/Plugins/Animation/PoseSearch/Source/Runtime/Public/PoseSearch/PoseSearchResult.h
@@ -115,7 +115,7 @@ struct FSearchResult : FDatabasePoseIdx
 	}
 
 	// Experimental, this feature might be removed without warning, not for production use
-	void InitFrom(const FPoseSearchBlueprintResult& BlueprintResult);
+	UE_API void InitFrom(const FPoseSearchBlueprintResult& BlueprintResult);
 };
 
 // Experimental, this feature might be removed without warning, not for production use
diff --git a/Engine/Plugins/Experimental/UAF/UAFPoseSearch/Source/UAFPoseSearch/Private/MotionMatchingTrait.cpp b/Engine/Plugins/Experimental/UAF/UAFPoseSearch/Source/UAFPoseSearch/Private/MotionMatchingTrait.cpp
--- a/Engine/Plugins/Experimental/UAF/UAFPoseSearch/Source/UAFPoseSearch/Private/MotionMatchingTrait.cpp
+++ b/Engine/Plugins/Experimental/UAF/UAFPoseSearch/Source/UAFPoseSearch/Private/MotionMatchingTrait.cpp
@@ -226,9 +226,16 @@ PRAGMA_ENABLE_DEPRECATION_WARNINGS
 				SearchContext.AddRole(GetCommonDefaultRole(DatabasesToSearch), &ChooserEvaluationContext, PoseHistory);
 
 				// If we can't advance or enough time has elapsed since the last pose jump then search
-				const bool bCanAdvance = SearchContext.GetContinuingPoseSearchResult().PoseIdx != INDEX_NONE;
-				const bool bSearch = !bCanAdvance || (MotionMatchingState.ElapsedPoseSearchTime >= SearchThrottleTime);
+				bool bSearch = true;
 				FSearchResult SearchResult;
+				if (MotionMatchingState.ElapsedPoseSearchTime < SearchThrottleTime)
+				{
+					SearchResult.InitFrom(MotionMatchingState.SearchResult);
+					SearchContext.UpdateContinuingPoseSearchResult(SearchResult, SearchResult);
+					const bool bCanAdvance = SearchContext.GetContinuingPoseSearchResult().PoseIdx != INDEX_NONE;
+					bSearch = !bCanAdvance;
+				}
+
 				if (bSearch)
 				{
 					MotionMatchingState.ElapsedPoseSearchTime = 0.f;
