PoseSearch - fix for FPoseSearchColumn::SetOutputs bForceBlendToValue

#Synced-CL 50317418

diff --git a/Engine/Plugins/Animation/PoseSearch/Source/Runtime/Private/Chooser/PoseSearchChooserColumn.cpp b/Engine/Plugins/Animation/PoseSearch/Source/Runtime/Private/Chooser/PoseSearchChooserColumn.cpp
--- a/Engine/Plugins/Animation/PoseSearch/Source/Runtime/Private/Chooser/PoseSearchChooserColumn.cpp
+++ b/Engine/Plugins/Animation/PoseSearch/Source/Runtime/Private/Chooser/PoseSearchChooserColumn.cpp
@@ -672,7 +672,7 @@ void FPoseSearchColumn::SetOutputs(FChooserEvaluationContext& Context, int RowIn
 				StartTimeValue = BestSearchResult.GetAssetTime();
 				CostValue = BestSearchResult.PoseCost;
 				bMirrorValue = SearchIndexAsset->IsMirrored();
-				bForceBlendToValue = BestSearchResult.bIsContinuingPoseSearch;
+				bForceBlendToValue = !BestSearchResult.bIsContinuingPoseSearch;
 			}
 		}
 	}
@@ -689,7 +689,7 @@ void FPoseSearchColumn::SetOutputs(FChooserEvaluationContext& Context, int RowIn
 					StartTimeValue = SearchResult.GetAssetTime();
 					CostValue = SearchResult.PoseCost;
 					bMirrorValue = SearchIndexAsset->IsMirrored();
-					bForceBlendToValue = SearchResult.bIsContinuingPoseSearch;
+					bForceBlendToValue = !SearchResult.bIsContinuingPoseSearch;
 
 					BestSearchResult = SearchResult;
 					break;
