PoseSearch - fix for crash in FAIPComponentSpacePoseProvider::CalculateComponentSpaceTransform when called with invalid bone indexes

#Synced-CL 50341867

diff --git a/Engine/Plugins/Animation/PoseSearch/Source/Runtime/Private/PoseSearchHistory.cpp b/Engine/Plugins/Animation/PoseSearch/Source/Runtime/Private/PoseSearchHistory.cpp
--- a/Engine/Plugins/Animation/PoseSearch/Source/Runtime/Private/PoseSearchHistory.cpp
+++ b/Engine/Plugins/Animation/PoseSearch/Source/Runtime/Private/PoseSearchHistory.cpp
@@ -319,11 +319,16 @@ FTransform FAIPComponentSpacePoseProvider::CalculateComponentSpaceTransform(cons
 	const USkeleton* Skeleton = BoneContainer.GetSkeletonAsset();
 	const FReferenceSkeleton& ReferenceSkeleton = Skeleton->GetReferenceSkeleton();
 	const int32 ParentIndex = ReferenceSkeleton.GetParentIndex(SkeletonBoneIdx.GetInt());
-	check(ParentIndex >= 0);
 
-	const TArray<FTransform>& RefBonePose = Skeleton->GetReferenceSkeleton().GetRefBonePose();
-	return RefBonePose[SkeletonBoneIdx.GetInt()] * CalculateComponentSpaceTransform(FSkeletonPoseBoneIndex(ParentIndex));
+	const FTransform ParentTransform = (ParentIndex >= 0) ? CalculateComponentSpaceTransform(FSkeletonPoseBoneIndex(ParentIndex)) : FTransform::Identity;
+	const TArray<FTransform>& RefBonePose = ReferenceSkeleton.GetRefBonePose();
+	
+	if (RefBonePose.IsValidIndex(SkeletonBoneIdx.GetInt()))
+	{
+		return RefBonePose[SkeletonBoneIdx.GetInt()] * ParentTransform;
+	}
 
+	return FTransform::Identity;
 }
 
 const USkeleton* FAIPComponentSpacePoseProvider::GetSkeletonAsset() const
