PoseSearch - don't refresh details view on any change which breaks the sliding functionality of interval

#Synced-CL 50341867

diff --git a/Engine/Plugins/Animation/PoseSearch/Source/Editor/Private/PoseSearchCustomization.cpp b/Engine/Plugins/Animation/PoseSearch/Source/Editor/Private/PoseSearchCustomization.cpp
--- a/Engine/Plugins/Animation/PoseSearch/Source/Editor/Private/PoseSearchCustomization.cpp
+++ b/Engine/Plugins/Animation/PoseSearch/Source/Editor/Private/PoseSearchCustomization.cpp
@@ -59,13 +59,6 @@ void FPoseSearchDatabaseAnimationAssetCustomization::CustomizeHeader(TSharedRef<
 		SNew(STextBlock)
 		.Text(SequenceNameText)
 	];
-
-	const FSimpleDelegate OnValueChanged = FSimpleDelegate::CreateLambda([&StructCustomizationUtils]()
-	{
-		StructCustomizationUtils.GetPropertyUtilities()->ForceRefresh();
-	});
-
-	InStructPropertyHandle->SetOnChildPropertyValueChanged(OnValueChanged);
 }
 
 void FPoseSearchDatabaseAnimationAssetCustomization::CustomizeChildren(TSharedRef<IPropertyHandle> InStructPropertyHandle, IDetailChildrenBuilder& StructBuilder, IPropertyTypeCustomizationUtils& StructCustomizationUtils)
