[PoseSearch] Allows tick dependency resolution for Motion Matching Interaction actors that cannot tick. Contribution from Daniel Manero

#Synced-CL 50662661

diff --git a/Engine/Plugins/Animation/PoseSearch/Source/Runtime/Private/PoseSearchInteractionIsland.cpp b/Engine/Plugins/Animation/PoseSearch/Source/Runtime/Private/PoseSearchInteractionIsland.cpp
--- a/Engine/Plugins/Animation/PoseSearch/Source/Runtime/Private/PoseSearchInteractionIsland.cpp
+++ b/Engine/Plugins/Animation/PoseSearch/Source/Runtime/Private/PoseSearchInteractionIsland.cpp
@@ -281,7 +281,7 @@ static void AddPrerequisite(FTickFunction& TickFunction, UObject* TargetObject,
 		TickFunction.AddPrerequisite(TargetObject, TargetTickFunction);
 
 #if ENABLE_ANIM_DEBUG
-		if (!TickFunction.GetPrerequisites().Contains(FTickPrerequisite(TargetObject, TargetTickFunction)))
+		if (TickFunction.bCanEverTick && !TickFunction.GetPrerequisites().Contains(FTickPrerequisite(TargetObject, TargetTickFunction)))
 		{
 			UE_LOG(LogPoseSearch, Error, TEXT("UE::PoseSearch::AddPrerequisite, Failed to add prerequisite from [%s] to [%s, %s]!"), *TickFunction.DiagnosticMessage(), *TargetObject->GetName(), *TargetTickFunction.DiagnosticMessage());
 		}
@@ -1073,7 +1073,7 @@ void FInteractionIsland::AddTickDependencies(UActorComponent* TickActorComponent
 	check(TickActorComponent);
 
 	AActor* TickActor = TickActorComponent->GetOwner();
-	check(TickActor && TickActor->CanEverTick());
+	check(TickActor);
 			
 	if (IInteractionIslandDependency* IslandDependency = FindCustomDependency(TickActorComponent))
 	{
@@ -1112,8 +1112,11 @@ void FInteractionIsland::AddTickDependencies(UActorComponent* TickActorComponent
 			}
 			else
 			{
-				// PostTickFunction should contain only the tick function to the main actor's ones (component and owner)
-				check(PostTickFunction.GetPrerequisites().Num() == 2);
+				if (TickActor->CanEverTick())
+				{
+					// PostTickFunction should contain only the tick function to the main actor's ones (component and owner) if actor ticks
+					check(PostTickFunction.GetPrerequisites().Num() == 2);
+				}
 
 				// adding to PreTickFunction all the tick dependencies TickActorComponent has (excluding the TickActor->PrimaryActorTick),
 				// so it runs after all the tick dependencies of ALL the TickActorComponents in this FInteractionIsland
@@ -1185,8 +1188,11 @@ void FInteractionIsland::AddTickDependencies(UActorComponent* TickActorComponent
 		}
 		else
 		{
-			// PostTickFunction should contain only the tick function to the main actor's ones (component and owner)
-			check(PostTickFunction.GetPrerequisites().Num() == 2);
+			if (TickActor->CanEverTick())
+			{
+				// PostTickFunction should contain only the tick function to the main actor's ones (component and owner) if actor ticks
+				check(PostTickFunction.GetPrerequisites().Num() == 2);
+			}
 
 			// adding to PreTickFunction all the tick dependencies TickActorComponent has (excluding the TickActor->PrimaryActorTick),
 			// so it runs after all the tick dependencies of ALL the TickActorComponents in this FInteractionIsland
@@ -1380,7 +1386,7 @@ void FInteractionIsland::RemoveTickDependencies(bool bValidateTickDependencies)
 			if (UActorComponent* TickActorComponent = TickActorComponentPtr.Get())
 			{
 				AActor* TickActor = TickActorComponent->GetOwner();
-				check(TickActor && TickActor->CanEverTick());
+				check(TickActor);
 
 				if(IInteractionIslandDependency* IslandDependency = FindCustomDependency(TickActorComponent))
 				{
