PoseSearch - DDC key generation fix to support objects with duplicated down dependencies

#Synced-CL 50784554

diff --git a/Engine/Plugins/Animation/PoseSearch/Source/Runtime/Private/PoseSearchDerivedData.cpp b/Engine/Plugins/Animation/PoseSearch/Source/Runtime/Private/PoseSearchDerivedData.cpp
--- a/Engine/Plugins/Animation/PoseSearch/Source/Runtime/Private/PoseSearchDerivedData.cpp
+++ b/Engine/Plugins/Animation/PoseSearch/Source/Runtime/Private/PoseSearchDerivedData.cpp
@@ -104,10 +104,6 @@ static FAutoConsoleVariableRef CVarMotionMatchReindexCancelledDatabases(TEXT("a.
 static bool GVarMotionMatchReindexAllReferencedDatabases = true;
 static FAutoConsoleVariableRef CVarMotionMatchReindexAllReferencedDatabases(TEXT("a.MotionMatch.ReindexAllReferencedDatabases"), GVarMotionMatchReindexAllReferencedDatabases, TEXT("Reindex All Referenced Databases"));
 
-// Experimental, this feature might be removed without warning, not for production use
-static int32 GVarMotionMatchPartialKeyHashesMode = 0;
-static FAutoConsoleVariableRef CVarMotionMatchPartialKeyHashesMode(TEXT("a.MotionMatch.PartialKeyHashesMode"), GVarMotionMatchPartialKeyHashesMode, TEXT("0: use partial key hashes, 1: do not use partial key hashes, 2: do not use and validate partial key hashes"));
-
 static const UE::DerivedData::FValueId Id(UE::DerivedData::FValueId::FromName("Data"));
 static const UE::DerivedData::FCacheBucket Bucket("PoseSearchDatabase");
 
@@ -744,9 +740,8 @@ static void PreprocessSearchIndexKDTree(FSearchIndex& SearchIndex, const UPoseSe
 		// testing kdtree Construct determinism
 		if (AnyTestFlags(EMotionMatchTestFlags::TestKDTreeConstructDeterminism))
 		{
-			const int32 NumIterations = GVarMotionMatchTestNumIterations;
 			TAlignedArray<float> PCAValuesTest;
-			for (int32 Iteration = 0; Iteration < NumIterations; ++Iteration)
+			for (int32 Iteration = 0; Iteration < GVarMotionMatchTestNumIterations; ++Iteration)
 			{
 				// copy PCAValues in a different container to ensure input data has different memory addresses
 				PCAValuesTest = SearchIndex.PCAValues;
@@ -1221,6 +1216,22 @@ static bool IndexDatabase(FSearchIndexBase& SearchIndexBase, const UPoseSearchDa
 	return true;
 }
 
+static FString GetDatabaseAndOuterName(const UPoseSearchDatabase* Database)
+{
+	if (Database)
+	{
+		const UObject* Outer = Database->GetOuter();
+		if (Outer && !Outer->IsA<UPackage>())
+		{
+			return Database->GetName() + " [" + Outer->GetName() + "]";
+		}
+
+		return Database->GetName();
+	}
+	
+	return TEXT("None");
+}
+
 // validating the SearchIndex against the Database to catch any data corruption, that can be caused by not updating or having conflicts over DDC key
 static bool ValidateSearchInfoAgainstDatabase(const FSearchIndex& SearchIndex, const UPoseSearchDatabase* Database, const UE::DerivedData::FCacheKey& FullIndexKey)
 {
@@ -1228,53 +1239,35 @@ static bool ValidateSearchInfoAgainstDatabase(const FSearchIndex& SearchIndex, c
 
 	if (SearchIndex.GetNumDimensions() != Database->Schema->SchemaCardinality)
 	{
-		UE_LOG(LogPoseSearch, Warning, TEXT("%s - %s BuildIndex From Cache Corrupted! SchemaCardinality mismatch %d vs %d"), *LexToString(FullIndexKey.Hash), *Database->GetName(), SearchIndex.GetNumDimensions(), Database->Schema->SchemaCardinality);
+		UE_LOG(LogPoseSearch, Warning, TEXT("%s - %s BuildIndex SearchIndex Corrupted! SchemaCardinality mismatch %d vs %d"), *LexToString(FullIndexKey.Hash), *GetDatabaseAndOuterName(Database), SearchIndex.GetNumDimensions(), Database->Schema->SchemaCardinality);
 		return false;
 	}
 
-	TArray<bool> SourceAssetIdxs;
-	SourceAssetIdxs.SetNum(Database->GetNumAnimationAssets());
+	const int32 NumAnimationAssets = Database->GetNumAnimationAssets();
 
 	for (const FSearchIndexAsset& SearchIndexAsset : SearchIndex.Assets)
 	{
-		if (!SourceAssetIdxs.IsValidIndex(SearchIndexAsset.GetSourceAssetIdx()))
+		const int32 SourceAssetIdx = SearchIndexAsset.GetSourceAssetIdx();
+		if (SourceAssetIdx < 0 || SourceAssetIdx >= NumAnimationAssets)
 		{
-			UE_LOG(LogPoseSearch, Warning, TEXT("%s - %s BuildIndex From Cache Corrupted! SearchIndex.Assets referencing missing asset with index %d"), *LexToString(FullIndexKey.Hash), *Database->GetName(), SearchIndexAsset.GetSourceAssetIdx());
+			UE_LOG(LogPoseSearch, Warning, TEXT("%s - %s BuildIndex SearchIndex Corrupted! SearchIndex.Assets referencing missing asset with index %d"), *LexToString(FullIndexKey.Hash), *GetDatabaseAndOuterName(Database), SourceAssetIdx);
 			return false;
 		}
 
-		SourceAssetIdxs[SearchIndexAsset.GetSourceAssetIdx()] = true;
-	}
-
-	for (int32 AnimationAssetIndex = 0; AnimationAssetIndex < Database->GetNumAnimationAssets(); ++AnimationAssetIndex)
-	{
-		if (const FPoseSearchDatabaseAnimationAssetBase* DatabaseAsset = Database->GetDatabaseAnimationAsset(AnimationAssetIndex))
+		if (const FPoseSearchDatabaseAnimationAssetBase* DatabaseAsset = Database->GetDatabaseAnimationAsset(SourceAssetIdx))
 		{
 			if (const UObject* AnimationAsset = DatabaseAsset->GetAnimationAsset())
 			{
-				if (DatabaseAsset->IsEnabled() && !SourceAssetIdxs[AnimationAssetIndex])
+				if (!DatabaseAsset->IsEnabled())
 				{
-					UE_LOG(LogPoseSearch, Warning, TEXT("%s - %s BuildIndex From Cache Corrupted! Couldn't find references to enabled asset %s in the SearchIndex"), *LexToString(FullIndexKey.Hash), *Database->GetName(), *AnimationAsset->GetName());
-					return false;
-				}
-				else if (!DatabaseAsset->IsEnabled() && SourceAssetIdxs[AnimationAssetIndex])
-				{
-					UE_LOG(LogPoseSearch, Warning, TEXT("%s - %s BuildIndex From Cache Corrupted! Found references to disabled asset %s in the SearchIndex"), *LexToString(FullIndexKey.Hash), *Database->GetName(), *AnimationAsset->GetName());
-					return false;
-				}
-			}
-			else
-			{
-				if (SourceAssetIdxs[AnimationAssetIndex])
-				{
-					UE_LOG(LogPoseSearch, Warning, TEXT("%s - %s BuildIndex From Cache Corrupted! Found references to null asset at index %d in the SearchIndex"), *LexToString(FullIndexKey.Hash), *Database->GetName(), AnimationAssetIndex);
+					UE_LOG(LogPoseSearch, Warning, TEXT("%s - %s BuildIndex SearchIndex Corrupted! Found references to disabled asset %s in the SearchIndex"), *LexToString(FullIndexKey.Hash), *GetDatabaseAndOuterName(Database), *AnimationAsset->GetName());
 					return false;
 				}
 			}
 		}
 		else
 		{
-			UE_LOG(LogPoseSearch, Warning, TEXT("%s - %s BuildIndex From Cache Corrupted! null FPoseSearchDatabaseAnimationAssetBase asset at index %d!?"), *LexToString(FullIndexKey.Hash), *Database->GetName(), AnimationAssetIndex);
+			UE_LOG(LogPoseSearch, Warning, TEXT("%s - %s BuildIndex SearchIndex Corrupted! null FPoseSearchDatabaseAnimationAssetBase asset at index %d!?"), *LexToString(FullIndexKey.Hash), *GetDatabaseAndOuterName(Database), SourceAssetIdx);
 			return false;
 		}
 	}
@@ -1575,22 +1568,8 @@ void FPoseSearchDatabaseAsyncCacheTask::StartNewRequestIfNeeded(bool bPerformCon
 	// Owner.Cancel must be performed before SearchIndex.Reset() in case any task is flying (launched by Owner.LaunchTask)
 	Owner.Cancel();
 
-	FKeyBuilder::EDebugPartialKeyHashesMode DebugPartialKeyHashesMode;
-	switch (GVarMotionMatchPartialKeyHashesMode)
-	{
-	case 0:
-		DebugPartialKeyHashesMode = FKeyBuilder::EDebugPartialKeyHashesMode::Use;
-		break;
-	case 1:
-		DebugPartialKeyHashesMode = FKeyBuilder::EDebugPartialKeyHashesMode::DoNotUse;
-		break;
-	default:
-		DebugPartialKeyHashesMode = FKeyBuilder::EDebugPartialKeyHashesMode::Validate;
-		break;
-	}
-
 	// composing the key
-	const FKeyBuilder KeyBuilder(Database.Get(), true, bPerformConditionalPostLoadIfRequired, &PartialKeyHashes, DebugPartialKeyHashesMode);
+	FKeyBuilder KeyBuilder(Database.Get(), true, bPerformConditionalPostLoadIfRequired, &PartialKeyHashes);
 	if (KeyBuilder.AnyAssetNotFullyLoaded())
 	{
 		DerivedDataKey = FIoHash::Zero;
@@ -1617,7 +1596,7 @@ void FPoseSearchDatabaseAsyncCacheTask::StartNewRequestIfNeeded(bool bPerformCon
 
 				SetState(EState::Prestarted);
 				
-				UE_LOG(LogPoseSearch, Log, TEXT("%s - %s BeginCache"), *LexToString(DerivedDataKey), *Database->GetName());
+				UE_LOG(LogPoseSearch, Log, TEXT("%s - %s BeginCache"), *LexToString(DerivedDataKey), *GetDatabaseAndOuterName(Database.Get()));
 
 				const FCacheKey CacheKey{ Bucket, DerivedDataKey };
 				const FCacheGetRequest CacheRequest = { { Database->GetPathName() }, CacheKey, ECachePolicy::Default };
@@ -1651,7 +1630,7 @@ void FPoseSearchDatabaseAsyncCacheTask::PreCancelIfDependsOn(const UObject* Obje
 		if (DatabaseDependencies.Contains(Object))
 		{
 			// Database can be null if the task was Ended/Failed and Database was already garbage collected, but Tick hasn't been called yet
-			FString DatabaseName = IsValid() ? *Database->GetName() : TEXT("Garbage Collected Database");
+			FString DatabaseName = IsValid() ? *GetDatabaseAndOuterName(Database.Get()) : TEXT("Garbage Collected Database");
 			UE_LOG(LogPoseSearch, Log, TEXT("%s - %s PreCancelled because of %s"), *LexToString(DerivedDataKey), *DatabaseName, *Object->GetName());
 
 			// Owner.Cancel must be performed before SearchIndex.Reset() in case any task is flying (launched by Owner.LaunchTask)
@@ -1667,7 +1646,7 @@ void FPoseSearchDatabaseAsyncCacheTask::Cancel()
 {
 	check(IsInGameThread());
 
-	FString DatabaseName = IsValid() ? *Database->GetName() : TEXT("Garbage Collected Database");
+	FString DatabaseName = IsValid() ? *GetDatabaseAndOuterName(Database.Get()) : TEXT("Garbage Collected Database");
 	UE_LOG(LogPoseSearch, Log, TEXT("%s - %s Cancelled"), *LexToString(DerivedDataKey), *DatabaseName);
 
 	// Owner.Cancel must be performed before SearchIndex.Reset() in case any task is flying (launched by Owner.LaunchTask)
@@ -1774,7 +1753,7 @@ void FPoseSearchDatabaseAsyncCacheTask::OnGetComplete(UE::DerivedData::FCacheGet
 		// so to account for this, we just reindex the database and update the associated DDC 
 		if (ValidateSearchInfoAgainstDatabase(SearchIndex, Database.Get(), FullIndexKey))
 		{
-			UE_LOG(LogPoseSearch, Log, TEXT("%s - %s BuildIndex From Cache"), *LexToString(FullIndexKey.Hash), *Database->GetName());
+			UE_LOG(LogPoseSearch, Log, TEXT("%s - %s BuildIndex From Cache"), *LexToString(FullIndexKey.Hash), *GetDatabaseAndOuterName(Database.Get()));
 		}
 		else
 		{
@@ -1787,7 +1766,7 @@ void FPoseSearchDatabaseAsyncCacheTask::OnGetComplete(UE::DerivedData::FCacheGet
 	if (Response.Status == EStatus::Canceled)
 	{
 		ResetSearchIndex();
-		UE_LOG(LogPoseSearch, Log, TEXT("%s - %s BuildIndex Cancelled"), *LexToString(FullIndexKey.Hash), *Database->GetName());
+		UE_LOG(LogPoseSearch, Log, TEXT("%s - %s BuildIndex Cancelled"), *LexToString(FullIndexKey.Hash), *GetDatabaseAndOuterName(Database.Get()));
 	}
 	
 	bool bForceBuildIndex = false;
@@ -1830,7 +1809,7 @@ void FPoseSearchDatabaseAsyncCacheTask::OnGetComplete(UE::DerivedData::FCacheGet
 					MainDatabase->NormalizationSet->AddUniqueDatabases(IndexBaseDatabases);
 				}
 
-				const FString MainDatabaseName = MainDatabase->GetName();
+				const FString MainDatabaseName = GetDatabaseAndOuterName(MainDatabase);
 				const UPoseSearchSchema* MainDatabaseSchema = MainDatabase->Schema;
 				if (!MainDatabaseSchema || MainDatabaseSchema->SchemaCardinality <= 0)
 				{
@@ -1852,7 +1831,7 @@ void FPoseSearchDatabaseAsyncCacheTask::OnGetComplete(UE::DerivedData::FCacheGet
 					check(DependentDatabase);
 
 					const UPoseSearchSchema* DependentDatabaseSchema = bNormalizeWithCommonSchema ? MainDatabaseSchema : DependentDatabase->Schema.Get();
-					const FString DependentDatabaseName = DependentDatabase->GetName();
+					const FString DependentDatabaseName = GetDatabaseAndOuterName(DependentDatabase);
 
 					FSearchIndexBase& SearchIndexBase = SearchIndexBases[IndexBaseIdx];
 
@@ -2014,8 +1993,7 @@ void FPoseSearchDatabaseAsyncCacheTask::OnGetComplete(UE::DerivedData::FCacheGet
 #if ENABLE_ANIM_DEBUG
 					if (AnyTestFlags(EMotionMatchTestFlags::TestIndexDatabaseDeterminism))
 					{
-						const int32 NumIterations = GVarMotionMatchTestNumIterations;
-						for (int32 Iteration = 0; Iteration < NumIterations; ++Iteration)
+						for (int32 Iteration = 0; Iteration < GVarMotionMatchTestNumIterations; ++Iteration)
 						{
 							FSearchIndexBase TestSearchIndexBase = SearchIndexBase;
 							if (IndexDatabase(TestSearchIndexBase, DependentDatabase, DependentDatabaseSchema, DependentSamplingContext, DependentAdditionalExtrapolationTime, Owner))
@@ -2109,11 +2087,9 @@ void FPoseSearchDatabaseAsyncCacheTask::OnGetComplete(UE::DerivedData::FCacheGet
 				// testing PruneDuplicatePCAValues determinism
 				if (AnyTestFlags(EMotionMatchTestFlags::TestPruneDuplicatePCAValuesDeterminism))
 				{
-					const int32 NumIterations = GVarMotionMatchTestNumIterations;
-
 					FSearchIndex TestSearchIndexA = SearchIndex;
 					TestSearchIndexA.PruneDuplicatePCAValues(MainDatabase->PCAValuesPruningSimilarityThreshold, MainDatabase->GetNumberOfPrincipalComponents());
-					for (int32 Iteration = 0; Iteration < NumIterations; ++Iteration)
+					for (int32 Iteration = 0; Iteration < GVarMotionMatchTestNumIterations; ++Iteration)
 					{
 						FSearchIndex TestSearchIndexB = SearchIndex;
 						TestSearchIndexB.PruneDuplicatePCAValues(MainDatabase->PCAValuesPruningSimilarityThreshold, MainDatabase->GetNumberOfPrincipalComponents());
@@ -2178,6 +2154,13 @@ void FPoseSearchDatabaseAsyncCacheTask::OnGetComplete(UE::DerivedData::FCacheGet
 					CompareSearchIndex(SearchIndexCompare, SearchIndex, MainDatabaseSchema, Message);
 					UE_LOG(LogPoseSearch, Warning, TEXT("%s - %s BuildIndex mismatch with DDC Index\n%s"), *LexToString(FullIndexKey.Hash), *MainDatabaseName, *Message);
 				}
+
+				if (!ValidateSearchInfoAgainstDatabase(SearchIndex, Database.Get(), FullIndexKey))
+				{
+					UE_LOG(LogPoseSearch, Error, TEXT("%s - %s BuildIndex Failed because of corrupted SearchIndex!"), *LexToString(FullIndexKey.Hash), *GetDatabaseAndOuterName(Database.Get()));
+					ResetSearchIndex();
+					return;
+				}
 #endif // ENABLE_ANIM_DEBUG
 
 				// putting SearchIndex to DDC
@@ -2381,7 +2364,7 @@ void FAsyncPoseSearchDatabasesManagement::PreModified(UObject* Object)
 	{
 		FScopeLock Lock(&Mutex);
 
-		PartialKeyHashes.Remove(Object);
+		PartialKeyHashes.InvalidateKeysFor(Object);
 
 		for (TUniquePtr<FPoseSearchDatabaseAsyncCacheTask>& TaskPtr : Tasks)
 		{
@@ -2491,8 +2474,6 @@ void FAsyncPoseSearchDatabasesManagement::Tick(float DeltaTime)
 	const bool bTestAssetSamplerDeterminismFromPreviousExecution = AnyTestFlags(EMotionMatchTestFlags::TestAssetSamplerDeterminismFromPreviousExecution);
 	if (bTestAssetSamplerDeterminism || bTestAssetSamplerDeterminismFromPreviousExecution)
 	{
-		const int32 NumIterations = GVarMotionMatchTestNumIterations;
-
 		struct FTestSample
 		{
 			FBoneContainer BoneContainer;
@@ -2564,6 +2545,7 @@ void FAsyncPoseSearchDatabasesManagement::Tick(float DeltaTime)
 			}
 		}
 
+		int32 NumIterations = GVarMotionMatchTestNumIterations;
 		ParallelFor(TestSamples.Num(), [&TestSamples, NumIterations, bTestAssetSamplerDeterminism, bTestAssetSamplerDeterminismFromPreviousExecution](int32 TestSampleIndex)
 			{
 				FMemMark Mark(FMemStack::Get());
@@ -2682,7 +2664,7 @@ void FAsyncPoseSearchDatabasesManagement::Tick(float DeltaTime)
 				if (Tasks[TaskIndex]->GetState() == FPoseSearchDatabaseAsyncCacheTask::EState::Ended ||
 					Tasks[TaskIndex]->GetState() == FPoseSearchDatabaseAsyncCacheTask::EState::Failed)
 				{
-					UE_LOG(LogPoseSearch, Log, TEXT("%s - %s Removed because of InvalidateCache with WaitForTaskCompletion"), *LexToString(Tasks[TaskIndex]->GetDerivedDataKey()), *Tasks[TaskIndex]->GetDatabase()->GetName());
+					UE_LOG(LogPoseSearch, Log, TEXT("%s - %s Removed because of InvalidateCache with WaitForTaskCompletion"), *LexToString(Tasks[TaskIndex]->GetDerivedDataKey()), *GetDatabaseAndOuterName(Tasks[TaskIndex]->GetDatabase()));
 					Tasks.RemoveAtSwap(TaskIndex, EAllowShrinking::No);
 				}
 			}
@@ -2691,7 +2673,7 @@ void FAsyncPoseSearchDatabasesManagement::Tick(float DeltaTime)
 		{
 			for (TUniquePtr<FPoseSearchDatabaseAsyncCacheTask>& TaskPtr : Tasks)
 			{
-				UE_LOG(LogPoseSearch, Log, TEXT("%s - %s Cancelled because of InvalidateCache"), *LexToString(TaskPtr->GetDerivedDataKey()), *TaskPtr->GetDatabase()->GetName());
+				UE_LOG(LogPoseSearch, Log, TEXT("%s - %s Cancelled because of InvalidateCache"), *LexToString(TaskPtr->GetDerivedDataKey()), *GetDatabaseAndOuterName(TaskPtr->GetDatabase()));
 			}
 			Tasks.Reset();
 		}
@@ -2707,25 +2689,24 @@ void FAsyncPoseSearchDatabasesManagement::Tick(float DeltaTime)
 
 #endif // ENABLE_ANIM_DEBUG
 	
-	const bool bReindexCancelledDatabases = GVarMotionMatchReindexCancelledDatabases;
-
 	// iterating backwards because of the possible RemoveAtSwap 
 	for (int32 TaskIndex = Tasks.Num() - 1; TaskIndex >= 0; --TaskIndex)
 	{
 		if (!Tasks[TaskIndex]->IsValid())
 		{
+			UE_LOG(LogPoseSearch, Log, TEXT("%s - Removed because IsValid"), *LexToString(Tasks[TaskIndex]->GetDerivedDataKey()));
 			Tasks.RemoveAtSwap(TaskIndex, EAllowShrinking::No);
 		}
 		else if (Tasks[TaskIndex]->GetState() == FPoseSearchDatabaseAsyncCacheTask::EState::Cancelled)
 		{
-			if (bReindexCancelledDatabases)
-			{
-				RequestAsyncBuildIndex(Tasks[TaskIndex]->GetDatabase(), ERequestAsyncBuildFlag::NewRequest);
-				Tasks.RemoveAtSwap(TaskIndex, EAllowShrinking::No);
-			}
-			else
+			const UPoseSearchDatabase* TaskDatabase = Tasks[TaskIndex]->GetDatabase();
+
+			UE_LOG(LogPoseSearch, Log, TEXT("%s - %s Removed because is Cancelled"), *LexToString(Tasks[TaskIndex]->GetDerivedDataKey()), *GetDatabaseAndOuterName(TaskDatabase));
+			Tasks.RemoveAtSwap(TaskIndex, EAllowShrinking::No);
+
+			if (GVarMotionMatchReindexCancelledDatabases)
 			{
-				Tasks.RemoveAtSwap(TaskIndex, EAllowShrinking::No);
+				RequestAsyncBuildIndex(TaskDatabase, ERequestAsyncBuildFlag::NewRequest);
 			}
 		}
 		else
@@ -2737,17 +2718,16 @@ void FAsyncPoseSearchDatabasesManagement::Tick(float DeltaTime)
 #if ENABLE_ANIM_DEBUG
 	if (AnyTestFlags(EMotionMatchTestFlags::TestDDCKeyDeterminism))
 	{
-		const int32 NumIterations = GVarMotionMatchTestNumIterations;
 		for (int32 TaskIndex = 0; TaskIndex < Tasks.Num(); ++TaskIndex)
 		{
 			if (const UPoseSearchDatabase* Database = Tasks[TaskIndex]->GetDatabase())
 			{
-				const FKeyBuilder KeyBuilder(Database, false, false);
+				FKeyBuilder KeyBuilder(Database, false, false);
 				const FIoHash IoHash = KeyBuilder.Finalize();
 
-				for (int32 IterationIndex = 0; IterationIndex < NumIterations; ++IterationIndex)
+				for (int32 IterationIndex = 0; IterationIndex < GVarMotionMatchTestNumIterations; ++IterationIndex)
 				{
-					const FKeyBuilder TestKeyBuilder(Database, false, false, &PartialKeyHashes, FKeyBuilder::EDebugPartialKeyHashesMode::Use);
+					FKeyBuilder TestKeyBuilder(Database, false, false, &PartialKeyHashes);
 					const FIoHash TestIoHash = TestKeyBuilder.Finalize();
 
 					if (!KeyBuilder.ValidateAgainst(TestKeyBuilder))
@@ -2885,7 +2865,7 @@ void FAsyncPoseSearchDatabasesManagement::Clear()
 
 	for (TUniquePtr<FPoseSearchDatabaseAsyncCacheTask>& TaskPtr : This.Tasks)
 	{
-		UE_LOG(LogPoseSearch, Log, TEXT("%s - %s Cancelled because of Clear"), *LexToString(TaskPtr->GetDerivedDataKey()), *TaskPtr->GetDatabase()->GetName());
+		UE_LOG(LogPoseSearch, Log, TEXT("%s - %s Cancelled because of Clear"), *LexToString(TaskPtr->GetDerivedDataKey()), *GetDatabaseAndOuterName(TaskPtr->GetDatabase()));
 	}
 	This.Tasks.Reset();
 }
diff --git a/Engine/Plugins/Animation/PoseSearch/Source/Runtime/Private/PoseSearchDerivedDataKey.cpp b/Engine/Plugins/Animation/PoseSearch/Source/Runtime/Private/PoseSearchDerivedDataKey.cpp
--- a/Engine/Plugins/Animation/PoseSearch/Source/Runtime/Private/PoseSearchDerivedDataKey.cpp
+++ b/Engine/Plugins/Animation/PoseSearch/Source/Runtime/Private/PoseSearchDerivedDataKey.cpp
@@ -20,7 +20,11 @@
 
 namespace UE::PoseSearch
 {
-
+// log FKeyBuilder profiling
+#ifndef ENABLE_KEY_PROFILING
+	#define ENABLE_KEY_PROFILING 0
+#endif
+	
 // log properties and UObjects names
 #ifndef UE_POSE_SEARCH_DERIVED_DATA_LOGGING
 	#define UE_POSE_SEARCH_DERIVED_DATA_LOGGING 0
@@ -33,6 +37,8 @@ namespace UE::PoseSearch
 
 FKeyBuilder::FKeyBuilder()
 {
+	check(IsInGameThread());
+
 	ArIgnoreOuterRef = true;
 
 	// Set FDerivedDataKeyBuilder to be a saving archive instead of a reference collector.
@@ -42,53 +48,51 @@ FKeyBuilder::FKeyBuilder()
 	SetIsSaving(true);
 }
 
-FKeyBuilder::FKeyBuilder(const UObject* Object, bool bUseDataVer, bool bPerformConditionalPostLoadIfRequired, FPartialKeyHashes* InPartialKeyHashes, EDebugPartialKeyHashesMode InDebugPartialKeyHashesMode)
+FKeyBuilder::FKeyBuilder(const UObject* Object, bool bUseDataVer, bool bPerformConditionalPostLoadIfRequired, FPartialKeyHashes* InPartialKeyHashes)
 : FKeyBuilder()
 {
+#if ENABLE_KEY_PROFILING
+	const double StartTime(FPlatformTime::Seconds());
+#endif // ENABLE_KEY_PROFILING
+
 	check(Object);
+	check(IsInGameThread());
 
-	KeyOwner = Object;
+	// FKeyBuilder is a saving only archiver, and since it doesn't modify the input Object it's safe to do a const_cast 
+	KeyOwner = const_cast<UObject*>(Object);
 
 	// preallocating a reasonable amount of memory to avoid multiple reallocations
 	ObjectsToSerialize.Reserve(256);
 	ObjectBeingSerializedDependencies.Reserve(256);
-	LocalPartialKeyHashes.Reserve(1024);
 
 	bPerformConditionalPostLoad = bPerformConditionalPostLoadIfRequired;
-	PartialKeyHashes = InPartialKeyHashes;
-	DebugPartialKeyHashesMode = InDebugPartialKeyHashesMode;
+	ExternalPartialKeyHashes = InPartialKeyHashes;
 
-	// FKeyBuilder is a saving only archiver, and since it doesn't modify the input Object it's safe to do a const_cast 
-	UObject* NonConstObject = const_cast<UObject*>(Object);
-	*this << NonConstObject;
+	UObject* NonPackageOuter = GetNonPackageOuter(KeyOwner);
+	
+	if (NonPackageOuter)
+	{
+		// if Object has a non UPackage outer we serialize that as well to collect all the referenced objects to KeyBuilder.GetDependencies() 
+		// to be able to reindex databases owned by other objects. via operator << the Outer will ba added to ObjectsToSerialize as first item
+		// and Object as second item, so Object will be serialized first. This will make unique keys for different Outer(s) with the same Outer
+		// thing that wouldn't happen if Object and Outer gets added in the opposite order
+		*this << NonPackageOuter;
+	}
+	
+	// we still need to serialize the KeyOwner in case it has newly created or duplicated and not yet referenced by Outer
+	*this << KeyOwner;
 
 	while (!ObjectsToSerialize.IsEmpty() && !bAnyAssetNotFullyLoaded)
 	{
 		SerializeObjectInternal(ObjectsToSerialize.Pop(EAllowShrinking::No));
 	}
 
-	if (bUseDataVer && !bAnyAssetNotFullyLoaded)
-	{
-		FLocalPartialKeyHash& LocalCachedHash = LocalPartialKeyHashes.AddDefaulted_GetRef();
-
-		Hasher.Reset();
+#if ENABLE_KEY_PROFILING
+	const double EndTime(FPlatformTime::Seconds());
+	const double DeltaTimeInMilliseconds = (EndTime - StartTime) * 1000.0;
+	UE_LOG(LogPoseSearch, Log, TEXT("Profiling: FKeyBuilder::FKeyBuilder %f ms"), DeltaTimeInMilliseconds);
+#endif // ENABLE_KEY_PROFILING
 
-		// used to invalidate the key without having to change POSESEARCHDB_DERIVEDDATA_VER all the times
-		int32 NonConstDatabaseIndexDerivedDataCacheKeyVersion = DatabaseIndexDerivedDataCacheKeyVersion;
-		FGuid VersionGuid = FDevSystemGuids::GetSystemGuid(FDevSystemGuids::Get().POSESEARCHDB_DERIVEDDATA_VER);
-		FString AnimationCompressionVersionString = UE::Anim::Compression::AnimationCompressionVersionString;
-
-		*this << VersionGuid;
-		*this << AnimationCompressionVersionString;
-		*this << NonConstDatabaseIndexDerivedDataCacheKeyVersion;
-
-		LocalCachedHash.Hash = Hasher.Finalize();
-	}
-}
-
-FKeyBuilder::FKeyBuilder(const UObject* Object, bool bUseDataVer, bool bPerformConditionalPostLoadIfRequired)
-	: FKeyBuilder(Object, bUseDataVer, bPerformConditionalPostLoadIfRequired, nullptr, FKeyBuilder::EDebugPartialKeyHashesMode::DoNotUse)
-{
 }
 
 void FKeyBuilder::Seek(int64 InPos)
@@ -104,7 +108,7 @@ bool FKeyBuilder::ShouldSkipProperty(const FProperty* InProperty) const
 		return false;
 	}
 
-	if (Super::ShouldSkipProperty(InProperty))
+	if (FArchiveUObject::ShouldSkipProperty(InProperty))
 	{
 		#if UE_POSE_SEARCH_DERIVED_DATA_LOGGING
 		UE_LOG(LogPoseSearch, Log, TEXT("  x %s (ShouldSkipProperty)"), *InProperty->GetFullName());
@@ -147,29 +151,52 @@ bool FKeyBuilder::ShouldSkipProperty(const FProperty* InProperty) const
 
 void FKeyBuilder::Serialize(void* Data, int64 Length)
 {
-	const uint8* HasherData = reinterpret_cast<uint8*>(Data);
-
 	#if UE_POSE_SEARCH_DERIVED_DATA_LOGGING_VERBOSE
+	const uint8* HasherData = reinterpret_cast<uint8*>(Data);
 	FString RawBytesString = BytesToString(HasherData, Length);
 	UE_LOG(LogPoseSearch, Log, TEXT("  > %s"), *RawBytesString);
 	#endif
 
-	Hasher.Update(HasherData, Length);
+	Hasher.Update(Data, Length);
+}
+
+FArchive& FKeyBuilder::operator<<(FStringBuilderBase& StringBuilderBase)
+{
+	// this is an encoding only archiver (no decoding will ever happen), so we don't need to serialize the length of the StringBuilderBase
+	Serialize(StringBuilderBase.GetData(), sizeof(FStringBuilderBase::ElementType) * StringBuilderBase.Len());
+	return *this;
 }
 
 FArchive& FKeyBuilder::operator<<(FName& Name)
 {
-	// Don't include the name of the object being serialized, since that isn't technically part of the object's state
-	if (!ObjectBeingSerialized || (Name != ObjectBeingSerialized->GetFName()))
-	{
-		// we cannot use GetTypeHash(Name) since it's bound to be non deterministic between editor restarts, so we convert the name into an FString and let the Serialize(void* Data, int64 Length) deal with it
-		FString NameString = Name.ToString();
-		*this << NameString;
-	}
+	// non FStringBuilderBase version
+	//FString NameString = Name.ToString();
+	//*this << NameString;
+	//return *this;
+	
+	// we cannot use GetTypeHash(Name) since it's bound to be non deterministic between editor restarts, so we convert the 
+	// name into an FString and let the Serialize(void* Data, int64 Length) deal with it
+	
+	// using FStringBuilderBase instead of FAnsiStringBuilderBase to avoid Name TCHAR to ANSICHAR conversion 
+	// (at the cost of having more data to pass into FKeyBuilder::Serialize), but with the optimization of not having to serialize the length of the string
+	StringBuilder.Reset();
+	StringBuilder << Name;
+	*this << StringBuilder;
 	return *this;
 }
 
-FArchive& FKeyBuilder::TryAddDependency(UObject* Object, bool bAddToPartialKeyHashes)
+FArchive& FKeyBuilder::operator<<(FText& Value)
+{
+	// non FStringBuilderBase version
+	//return FArchiveUObject::operator<<(Value);
+
+	StringBuilder.Reset();
+	StringBuilder << Value.ToString();
+	*this << StringBuilder;
+	return *this;
+}
+
+FArchive& FKeyBuilder::TryAddDependency(UObject* Object)
 {
 	// @todo: add RF_NeedPostLoadSubobjects, RF_NeedInitialization, RF_NeedLoad, RF_WillBeLoaded? 
 	if (Object->HasAnyFlags(RF_NeedPostLoad))
@@ -219,7 +246,7 @@ FArchive& FKeyBuilder::operator<<(class UObject*& Object)
 {
 	if (ShouldInclude(Object))
 	{
-		return TryAddDependency(Object, true);
+		return TryAddDependency(Object);
 	}
 
 #if UE_POSE_SEARCH_DERIVED_DATA_LOGGING
@@ -264,91 +291,58 @@ FArchive& FKeyBuilder::operator<<(class UObject*& Object)
 
 void FKeyBuilder::SerializeObjectInternal(UObject* Object)
 {
+	// resetting Hasher, now used to hash all the input Object properties, except referencesto other UObject(s), that will be added to 	
+	// ObjectBeingSerializedDependencies and Dependencies, and processed after this Object processing
 	Hasher.Reset();
 
 	check(!bAnyAssetNotFullyLoaded && Object);
 
-	// test to validate the PartialKeyHashes
-	bool bIsValidTestPartialKeyHash = false;
-	FPartialKeyHashes::FEntry TestEntry;
-
-	// adding the LocalCachedHash here to keep the order consistent with Dependencies
-	FLocalPartialKeyHash& LocalCachedHash = LocalPartialKeyHashes.AddDefaulted_GetRef();
+	FPartialKeyHashes& PartialKeyHashes = GetPartialKeyHashes();
 
-	if (PartialKeyHashes)
+	FObjectKey ObjectKey(Object);
+	if (const FPartialKeyHashes::FEntry* Entry = PartialKeyHashes.FindAndRemoveResolved(ObjectKey))
 	{
-		if (DebugPartialKeyHashesMode == EDebugPartialKeyHashesMode::Validate)
+		for (const FObjectKey& DependencyPtr : Entry->Dependencies)
 		{
-			if (const FPartialKeyHashes::FEntry* Entry = PartialKeyHashes->Find(Object))
+			if (UObject* Dependency = DependencyPtr.ResolveObjectPtr())
 			{
-				bIsValidTestPartialKeyHash = true;
-				TestEntry = *Entry;
+				TryAddDependency(Dependency);
 			}
 		}
-		else if (DebugPartialKeyHashesMode == EDebugPartialKeyHashesMode::Use)
-		{
-			if (const FPartialKeyHashes::FEntry* Entry = PartialKeyHashes->Find(Object))
-			{
-				for (const TWeakObjectPtr<>& DependencyPtr : Entry->Dependencies)
-				{
-					if (UObject* Dependency = DependencyPtr.Get())
-					{
-						TryAddDependency(Dependency, false);
-					}
-				}
-
-				LocalCachedHash.Object = Object;
-				LocalCachedHash.Hash = Entry->Hash;
-				return;
-			}
-		}
-	}
-
-	// making sure we don't call Object->Serialize recursively!
-	check(ObjectBeingSerialized == nullptr);
-	ObjectBeingSerialized = Object;
-
-	ObjectBeingSerializedDependencies.Reset();
-
-	if (IsAddNameOnlyType(Object))
-	{
-		// for specific types we only add their names to the hash
-#if UE_POSE_SEARCH_DERIVED_DATA_LOGGING
-		UE_LOG(LogPoseSearch, Log, TEXT("AddingNameOnly '%s' (%s)"), *Object->GetName(), *Object->GetClass()->GetName());
-#endif
-		FString ObjectName = GetFullNameSafe(Object);
-		*this << ObjectName;
 	}
 	else
 	{
+		ObjectBeingSerializedDependencies.Reset();
+
+		if (IsAddNameOnlyType(Object))
+		{
+			// for specific types we only add their names to the hash
 #if UE_POSE_SEARCH_DERIVED_DATA_LOGGING
-		UE_LOG(LogPoseSearch, Log, TEXT("Begin '%s' (%s)"), *Object->GetName(), *Object->GetClass()->GetName());
+			UE_LOG(LogPoseSearch, Log, TEXT("AddingNameOnly '%s' (%s)"), *Object->GetName(), *Object->GetClass()->GetName());
 #endif
-		Object->Serialize(*this);
+			// non FStringBuilderBase version
+			//FString FullName = Object->GetFullName();
+			//*this << FullName;
 
+			StringBuilder.Reset();
+			Object->GetFullName(StringBuilder);
+			*this << StringBuilder;
+		}
+		else
+		{
 #if UE_POSE_SEARCH_DERIVED_DATA_LOGGING
-		UE_LOG(LogPoseSearch, Log, TEXT("End '%s' (%s)"), *Object->GetName(), *Object->GetClass()->GetName());
+			UE_LOG(LogPoseSearch, Log, TEXT("Begin '%s' (%s)"), *Object->GetName(), *Object->GetClass()->GetName());
 #endif
-	}
-
-	// making sure we don't call Object->Serialize recursively!
-	check(ObjectBeingSerialized == Object);
-	ObjectBeingSerialized = nullptr;
+			Object->Serialize(*this);
 
-	if (!bAnyAssetNotFullyLoaded)
-	{
-		LocalCachedHash.Object = Object;
-		LocalCachedHash.Hash = Hasher.Finalize();
+#if UE_POSE_SEARCH_DERIVED_DATA_LOGGING
+			UE_LOG(LogPoseSearch, Log, TEXT("End '%s' (%s)"), *Object->GetName(), *Object->GetClass()->GetName());
+#endif
+		}
 
-		if (DebugPartialKeyHashesMode != EDebugPartialKeyHashesMode::DoNotUse && PartialKeyHashes)
+		if (!bAnyAssetNotFullyLoaded)
 		{
-			PartialKeyHashes->Add(LocalCachedHash.Object, LocalCachedHash.Hash, ObjectBeingSerializedDependencies);
-
-			if (bIsValidTestPartialKeyHash)
-			{
-				check(TestEntry.CheckDependencies(ObjectBeingSerializedDependencies));
-				check(LocalCachedHash.Hash == TestEntry.Hash);
-			}
+			PartialKeyHashes.Add(ObjectKey, Hasher.Finalize(), ObjectBeingSerializedDependencies);
 		}
 	}
 }
@@ -401,20 +395,81 @@ bool FKeyBuilder::AnyAssetNotReady() const
 	return false;
 }
 
-FIoHash FKeyBuilder::Finalize() const
+FIoHash FKeyBuilder::Finalize()
 {
+#if ENABLE_KEY_PROFILING
+	const double StartTime(FPlatformTime::Seconds());
+#endif // ENABLE_KEY_PROFILING
+
 	check(!bAnyAssetNotFullyLoaded); // otherwise key can be non deterministic
+	check(KeyOwner);
+
+	FPartialKeyHashes& PartialKeyHashes = GetPartialKeyHashes();
+
+	Hasher.Reset();
+
+	TSet<FObjectKey> VisitedObjectKeys;
+	VisitedObjectKeys.Reserve(8192);
 
-	HashBuilderType FinalizeHasher;
+	TArray<FObjectKey> ObjectKeysToVisit;
+	ObjectKeysToVisit.Reserve(1024);
 
-	for (const FLocalPartialKeyHash& LocalCachedHash : LocalPartialKeyHashes)
+	if (UObject* NonPackageOuter = GetNonPackageOuter(KeyOwner))
 	{
-		const HashDigestType::ByteArray& LocalHashData = LocalCachedHash.Hash.GetBytes();
-		FinalizeHasher.Update(LocalHashData, sizeof(HashDigestType::ByteArray));
+		ObjectKeysToVisit.Add(NonPackageOuter);
+		
+		// non FStringBuilderBase version
+		//FString FullName = KeyOwner->GetFullName();
+		//*this << FullName;
+
+		// since Object is contained in Outer we need to serialize its name (otherwise skipped)
+		// we cannot use GetTypeHash(Name) since it's bound to be non deterministic between editor restarts, so we convert the name into an FString and let the Serialize(void* Data, int64 Length) deal with it
+		StringBuilder.Reset(); 
+		KeyOwner->GetFullName(StringBuilder);
+		*this << StringBuilder;
+	}
+
+	ObjectKeysToVisit.Add(KeyOwner);
+
+	// used to invalidate the key without having to change POSESEARCHDB_DERIVEDDATA_VER all the times
+	int32 NonConstDatabaseIndexDerivedDataCacheKeyVersion = DatabaseIndexDerivedDataCacheKeyVersion;
+	FGuid VersionGuid = FDevSystemGuids::GetSystemGuid(FDevSystemGuids::Get().POSESEARCHDB_DERIVEDDATA_VER);
+	FString AnimationCompressionVersionString = UE::Anim::Compression::AnimationCompressionVersionString;
+
+	*this << VersionGuid;
+	*this << AnimationCompressionVersionString;
+	*this << NonConstDatabaseIndexDerivedDataCacheKeyVersion;
+
+	// iterating over all ObjectKeysToVisit to finalize the composition the final hash while continuing adding their unvisited dependencies
+	while (!ObjectKeysToVisit.IsEmpty())
+	{
+		const FObjectKey CurrentObjectKey = ObjectKeysToVisit.Pop(EAllowShrinking::No);
+
+		const FPartialKeyHashes::FEntry* Entry = PartialKeyHashes.Find(CurrentObjectKey);
+		if (ensure(Entry))
+		{
+			const HashDigestType::ByteArray& LocalHashData = Entry->Hash.GetBytes();
+			Hasher.Update(LocalHashData, sizeof(HashDigestType::ByteArray));
+	
+			bool bAlreadyVisited = false;
+			VisitedObjectKeys.Add(CurrentObjectKey, &bAlreadyVisited);
+			if (!bAlreadyVisited)
+			{
+				ObjectKeysToVisit.Append(Entry->Dependencies);
+			}
+		}
 	}
 
 	// Stores a BLAKE3-160 hash, taken from the first 20 bytes of a BLAKE3-256 hash
-	return FIoHash(FinalizeHasher.Finalize());
+	const FIoHash IoHash(Hasher.Finalize());
+
+#if ENABLE_KEY_PROFILING
+	const double EndTime(FPlatformTime::Seconds());
+	const double DeltaTimeInMilliseconds = (EndTime - StartTime) * 1000.0;
+	UE_LOG(LogPoseSearch, Log, TEXT("Profiling: FKeyBuilder::Finalize %f ms"), DeltaTimeInMilliseconds);
+#endif // ENABLE_KEY_PROFILING
+
+	return IoHash;
 }
 
 const TSet<const UObject*>& FKeyBuilder::GetDependencies() const
@@ -519,24 +574,6 @@ bool FKeyBuilder::ValidateAgainst(const FKeyBuilder& Other) const
 		}
 	}
 
-	if (LocalPartialKeyHashes.Num() != Other.LocalPartialKeyHashes.Num())
-	{
-		return false;
-	}
-	
-	for (int32 Index = 0; Index < LocalPartialKeyHashes.Num(); ++Index)
-	{
-		if (LocalPartialKeyHashes[Index].Hash != Other.LocalPartialKeyHashes[Index].Hash)
-		{
-			return false;
-		}
-
-		if (LocalPartialKeyHashes[Index].Object != Other.LocalPartialKeyHashes[Index].Object)
-		{
-			return false;
-		}
-	}
-
 	return true;
 }
 
diff --git a/Engine/Plugins/Animation/PoseSearch/Source/Runtime/Public/PoseSearch/PoseSearchDerivedDataKey.h b/Engine/Plugins/Animation/PoseSearch/Source/Runtime/Public/PoseSearch/PoseSearchDerivedDataKey.h
--- a/Engine/Plugins/Animation/PoseSearch/Source/Runtime/Public/PoseSearch/PoseSearchDerivedDataKey.h
+++ b/Engine/Plugins/Animation/PoseSearch/Source/Runtime/Public/PoseSearch/PoseSearchDerivedDataKey.h
@@ -7,123 +7,37 @@
 #include "Hash/Blake3.h"
 #include "IO/IoHash.h"
 #include "Serialization/ArchiveUObject.h"
+#include "UObject/ObjectKey.h"
+#include "UObject/Package.h"
 #include "UObject/UnrealType.h"
 
 namespace UE::PoseSearch
 {
 
-class FPartialKeyHashes;
+using HashDigestType = FBlake3Hash;
+using HashBuilderType = FBlake3;
 
-class FKeyBuilder : public FArchiveUObject
+static UObject* GetNonPackageOuter(UObject* Object)
 {
-public:
-	using Super = FArchiveUObject;
-	using HashDigestType = FBlake3Hash;
-	using HashBuilderType = FBlake3;
-
-	inline static const FName ExcludeFromHashName = FName(TEXT("ExcludeFromHash"));
-	inline static const FName NeverInHashName = FName(TEXT("NeverInHash"));
-	inline static const FName IgnoreForMemberInitializationTestName = FName(TEXT("IgnoreForMemberInitializationTest"));
-	
-	POSESEARCH_API FKeyBuilder();
-	POSESEARCH_API FKeyBuilder(const UObject* Object, bool bUseDataVer, bool bPerformConditionalPostLoadIfRequired);
-	
-	// Experimental, this feature might be removed without warning, not for production use
-	enum EDebugPartialKeyHashesMode
-	{
-		Use,
-		DoNotUse,
-		Validate
-	};
-
-	// Experimental, this feature might be removed without warning, not for production use
-	POSESEARCH_API FKeyBuilder(const UObject* Object, bool bUseDataVer, bool bPerformConditionalPostLoadIfRequired, FPartialKeyHashes* InPartialKeyHashes, EDebugPartialKeyHashesMode InDebugPartialKeyHashesMode);
-
-	// Experimental, this feature might be removed without warning, not for production use
-	POSESEARCH_API bool ValidateAgainst(const FKeyBuilder& Other) const;
-
-	using Super::IsSaving;
-	using Super::operator<<;
-
-	// Begin FArchive Interface
-	POSESEARCH_API virtual void Seek(int64 InPos) override;
-	POSESEARCH_API virtual bool ShouldSkipProperty(const FProperty* InProperty) const override;
-	POSESEARCH_API virtual void Serialize(void* Data, int64 Length) override;
-	POSESEARCH_API virtual FArchive& operator<<(FName& Name) override;
-	POSESEARCH_API virtual FArchive& operator<<(class UObject*& Object) override;
-	POSESEARCH_API virtual FString GetArchiveName() const override;
-	// End FArchive Interface
-	
-	POSESEARCH_API bool AnyAssetNotFullyLoaded() const;
-	POSESEARCH_API bool AnyAssetNotReady() const;
-	POSESEARCH_API FIoHash Finalize() const;
-	POSESEARCH_API const TSet<const UObject*>& GetDependencies() const;
-
-	// Experimental, this feature might be removed without warning, not for production use
-	static POSESEARCH_API bool ShouldInclude(const UObject* Object);
-
-protected:
-	// to keep the key generation lightweight, we don't hash these types
-	static POSESEARCH_API bool IsExcludedType(const UObject* Object);
-
-	// to keep the key generation lightweight, we hash only the full names for these types. Object(s) will be added to Dependencies
-	static POSESEARCH_API bool IsAddNameOnlyType(const UObject* Object);
-
-	HashBuilderType Hasher;
-
-	// UPoseSearchDatabase instance "owner" of the key generation
-	const UObject* KeyOwner = nullptr;
-
-	// Set of objects that have already been serialized
-	TSet<const UObject*> Dependencies;
-
-	// Object currently being serialized
-	UObject* ObjectBeingSerialized = nullptr;
-
-	// true if some dependent assets are not fully loaded
-	bool bAnyAssetNotFullyLoaded = false;
-
-	// if true ConditionalPostLoad will be performed on the dependant assets requiring it
-	bool bPerformConditionalPostLoad = false;
-
-private:
-	// Experimental, this feature might be removed without warning, not for production use
-	POSESEARCH_API void SerializeObjectInternal(UObject* Object);
-
-	// Experimental, this feature might be removed without warning, not for production use
-	POSESEARCH_API FArchive& TryAddDependency(UObject* Object, bool bAddToPartialKeyHashes);
-
-	// Experimental, this feature might be removed without warning, not for production use
-	TArray<UObject*> ObjectsToSerialize;
-	
-	// Experimental, this feature might be removed without warning, not for production use
-	TArray<UObject*> ObjectBeingSerializedDependencies;
-
-	// Experimental, this feature might be removed without warning, not for production use
-	struct FLocalPartialKeyHash
+	if (UObject* Outer = Object->GetOuter())
 	{
-		UObject* Object = nullptr;
-		HashDigestType Hash;
-	};
-	
-	// Experimental, this feature might be removed without warning, not for production use
-	TArray<FLocalPartialKeyHash> LocalPartialKeyHashes;
-	
-	// Experimental, this feature might be removed without warning, not for production use
-	FPartialKeyHashes* PartialKeyHashes = nullptr;
-
-	// Experimental, this feature might be removed without warning, not for production use
-	EDebugPartialKeyHashesMode DebugPartialKeyHashesMode = EDebugPartialKeyHashesMode::Use;
-};
+		if (!Outer->IsA<UPackage>())
+		{
+			return Outer;
+		}
+	}
+	return nullptr;
+}
 
 // Experimental, this feature might be removed without warning, not for production use
-class FPartialKeyHashes
+// FPartialKeyHashes contains the combined hash of all the properties, excluding the references to other UObject(s),
+// for each UObject encountered during FKeyBuilder creation (serialization via FArchiveUObject)
+struct FPartialKeyHashes
 {
-public:
 	struct FEntry
 	{
-		FKeyBuilder::HashDigestType Hash;
-		TArray<TWeakObjectPtr<>> Dependencies;
+		HashDigestType Hash;
+		TArray<FObjectKey> Dependencies;
 
 		bool CheckDependencies(TArrayView<UObject*> OtherDependencies) const
 		{
@@ -135,73 +49,66 @@ public:
 			
 			for (int32 DependencyIndex = 0; DependencyIndex < Dependencies.Num(); ++DependencyIndex)
 			{
-				if (!OtherDependencies[DependencyIndex])
+				if (Dependencies[DependencyIndex] != OtherDependencies[DependencyIndex])
 				{
 					return false;
 				}
-
-				// we could have lost a weak pointer here...
-				if (Dependencies[DependencyIndex].IsValid())
-				{
-					if (Dependencies[DependencyIndex].Get() != OtherDependencies[DependencyIndex])
-					{
-						return false;
-					}
-				}
 			}
 			#endif // DO_CHECK
 
 			return true;
 		}
-
 	};
 
+	typedef TPair<FObjectKey, FEntry> FPair;
+
 	void Reset()
 	{
 		Entries.Reset();
 	}
 
-	void Remove(UObject* Object)
+	void InvalidateKeysFor(UObject* Object)
 	{
+		if (UObject* NonPackageOuter = GetNonPackageOuter(Object))
+		{
+			Entries.Remove(NonPackageOuter);
+		}
 		Entries.Remove(Object);
 	}
 
-	void Add(UObject* Object, const FKeyBuilder::HashDigestType& Hash, TArrayView<UObject*> Dependencies)
+	const FEntry* Add(FObjectKey& ObjectKey, const HashDigestType& Hash, TArrayView<UObject*> Dependencies)
 	{
-		check(Object);
-
 		check(!Hash.IsZero());
 
-		if (FEntry* OldEntry = Entries.Find(Object))
+		if (FEntry* OldEntry = Entries.Find(ObjectKey))
 		{
 			check(OldEntry->Hash == Hash);
-			check(OldEntry->CheckDependencies(Dependencies))
-		}
-		else
-		{
-			FEntry& NewEntry = Entries.Add(Object);
-			NewEntry.Hash = Hash;
-			NewEntry.Dependencies = Dependencies;
+			check(OldEntry->CheckDependencies(Dependencies));
+			return OldEntry;
 		}
+		
+		FEntry& NewEntry = Entries.Add(ObjectKey);
+		NewEntry.Hash = Hash;
+		NewEntry.Dependencies = Dependencies;
+		return &NewEntry;
 	}
 
-	const FEntry* Find(UObject* Object)
+	FEntry* FindAndRemoveResolved(FObjectKey& ObjectKey)
 	{
-		check(Object);
-		if (const TPair<TWeakObjectPtr<>, FEntry>* Pair = Entries.FindPair(Object))
+		if (FPair* Pair = Entries.FindPair(ObjectKey))
 		{
-			// making sure all the TWeakObjectPtr are still valid
-			if (!Pair->Key.IsValid())
+			// making sure all the FObjectKey(s) are still valid
+			if (!Pair->Key.ResolveObjectPtr())
 			{
-				Entries.Remove(Object);
+				Entries.Remove(ObjectKey);
 				return nullptr;
 			}
 			
-			for (const TWeakObjectPtr<> Dependency : Pair->Value.Dependencies)
+			for (FObjectKey& Dependency : Pair->Value.Dependencies)
 			{
-				if (!Dependency.IsValid())
+				if (!Dependency.ResolveObjectPtr())
 				{
-					Entries.Remove(Object);
+					Entries.Remove(ObjectKey);
 					return nullptr;
 				}
 			}
@@ -212,18 +119,98 @@ public:
 		return nullptr;
 	}
 
+	const FEntry* Find(const FObjectKey& ObjectKey) const
+	{
+		return Entries.Find(ObjectKey);
+	}
+		
 private:
-	struct FEntries : public TMap<TWeakObjectPtr<>, FEntry>
+	struct FEntries : public TMap<FObjectKey, FEntry>
 	{
-		const TPair<TWeakObjectPtr<>, FEntry>* FindPair(TWeakObjectPtr<> Key) const
+		FPair* FindPair(const FObjectKey& ObjectKey)
 		{
-			return Pairs.Find(Key);
+			return Pairs.Find(ObjectKey);
 		}
 	};
 
 	FEntries Entries;
 };
 
+class FKeyBuilder : public FArchiveUObject
+{
+public:
+	inline static const FName ExcludeFromHashName = FName(TEXT("ExcludeFromHash"));
+	inline static const FName NeverInHashName = FName(TEXT("NeverInHash"));
+	inline static const FName IgnoreForMemberInitializationTestName = FName(TEXT("IgnoreForMemberInitializationTest"));
+	
+	UE_NONCOPYABLE(FKeyBuilder);
+
+	POSESEARCH_API FKeyBuilder();
+	
+	// Experimental, this feature might be removed without warning, not for production use
+	POSESEARCH_API FKeyBuilder(const UObject* Object, bool bUseDataVer, bool bPerformConditionalPostLoadIfRequired, FPartialKeyHashes* PartialKeyHashes = nullptr);
+
+	// Experimental, this feature might be removed without warning, not for production use
+	POSESEARCH_API bool ValidateAgainst(const FKeyBuilder& Other) const;
+
+	using FArchiveUObject::IsSaving;
+	using FArchiveUObject::operator<<;
+
+	// Begin FArchive Interface
+	POSESEARCH_API virtual void Seek(int64 InPos) override;
+	POSESEARCH_API virtual bool ShouldSkipProperty(const FProperty* InProperty) const override;
+	POSESEARCH_API virtual void Serialize(void* Data, int64 Length) override;
+	POSESEARCH_API virtual FArchive& operator<<(FName& Name) override;
+	POSESEARCH_API virtual FArchive& operator<<(FText& Value) override;
+	POSESEARCH_API virtual FArchive& operator<<(class UObject*& Object) override;
+	POSESEARCH_API virtual FString GetArchiveName() const override;
+	// End FArchive Interface
+	
+	POSESEARCH_API bool AnyAssetNotFullyLoaded() const;
+	POSESEARCH_API bool AnyAssetNotReady() const;
+	POSESEARCH_API FIoHash Finalize();
+	POSESEARCH_API const TSet<const UObject*>& GetDependencies() const;
+
+	// Experimental, this feature might be removed without warning, not for production use
+	static POSESEARCH_API bool ShouldInclude(const UObject* Object);
+
+protected:
+	// to keep the key generation lightweight, we don't hash these types
+	static POSESEARCH_API bool IsExcludedType(const UObject* Object);
+
+	// to keep the key generation lightweight, we hash only the full names for these types. Object(s) will be added to Dependencies
+	static POSESEARCH_API bool IsAddNameOnlyType(const UObject* Object);
+
+	HashBuilderType Hasher;
+
+	// UPoseSearchDatabase instance "owner" of the key generation
+	UObject* KeyOwner = nullptr;
+
+	// Set of objects that have already been serialized
+	TSet<const UObject*> Dependencies;
+
+	// true if some dependent assets are not fully loaded
+	bool bAnyAssetNotFullyLoaded = false;
+
+	// if true ConditionalPostLoad will be performed on the dependant assets requiring it
+	bool bPerformConditionalPostLoad = false;
+
+private:
+	void SerializeObjectInternal(UObject* Object);
+	FArchive& TryAddDependency(UObject* Object);
+
+	FPartialKeyHashes& GetPartialKeyHashes() { return ExternalPartialKeyHashes ? *ExternalPartialKeyHashes : InternalPartialKeyHashes; }
+	FArchive& operator<<(FStringBuilderBase& StringBuilderBase);
+
+	TArray<UObject*> ObjectsToSerialize;
+	TArray<UObject*> ObjectBeingSerializedDependencies;
+
+	FPartialKeyHashes* ExternalPartialKeyHashes = nullptr;
+	FPartialKeyHashes InternalPartialKeyHashes;
+
+	FStringBuilderBase StringBuilder;
+};
+
 } // namespace UE::PoseSearch
 
 #endif // WITH_EDITOR
