PoseSearch - PoseSearchColumn support for unique PSD per row

#Synced-CL 50900617

diff --git a/Engine/Plugins/Animation/PoseSearch/Source/Editor/Private/Chooser/PoseSearchColumnEditor.cpp b/Engine/Plugins/Animation/PoseSearch/Source/Editor/Private/Chooser/PoseSearchColumnEditor.cpp
--- a/Engine/Plugins/Animation/PoseSearch/Source/Editor/Private/Chooser/PoseSearchColumnEditor.cpp
+++ b/Engine/Plugins/Animation/PoseSearch/Source/Editor/Private/Chooser/PoseSearchColumnEditor.cpp
@@ -4,6 +4,7 @@
 #include "Animation/AnimationAsset.h"
 #include "ChooserColumnHeader.h"
 #include "ContentBrowserModule.h"
+#include "DetailLayoutBuilder.h"
 #include "Editor.h"
 #include "IContentBrowserSingleton.h"
 #include "ObjectChooserWidgetFactories.h"
@@ -12,7 +13,9 @@
 #include "PoseSearch/PoseSearchSchema.h"
 #include "ScopedTransaction.h"
 #include "SPropertyAccessChainWidget.h"
+#include "Subsystems/AssetEditorSubsystem.h"
 #include "Widgets/Input/SNumericEntryBox.h"
+#include "Widgets/Input/STextEntryPopup.h"
 #include "Widgets/SBoxPanel.h"
 
 #define LOCTEXT_NAMESPACE "PoseSearchColumnEditor"
@@ -22,77 +25,209 @@ namespace UE::PoseSearchEditor
 
 TSharedRef<SWidget> CreatePoseSearchColumnWidget(UChooserTable* Chooser, FChooserColumnBase* Column, int RowIndex)
 {
+	using namespace UE::ChooserEditor;
+
 	FPoseSearchColumn* PoseSearchColumn = static_cast<FPoseSearchColumn*>(Column);
-	
+
 	if (RowIndex == UE::ChooserEditor::ColumnWidget_SpecialIndex_Header)
 	{
-		// @todo: avoid reopening the asset picker after tyhe first user choice, or prevent the FPoseSearchColumn creation if no schemas has been selected!
-		if (!PoseSearchColumn->GetDatabaseSchema())
-		{
-			TSharedPtr<SWindow> PickerWindow;
-
-			// asset picker config to select Schema
-			FAssetPickerConfig AssetPickerConfig;
-			AssetPickerConfig.Filter.ClassPaths.Add(UPoseSearchSchema::StaticClass()->GetClassPathName());
-			AssetPickerConfig.Filter.bRecursiveClasses = true;
-			AssetPickerConfig.InitialAssetViewType = EAssetViewType::List;
-			AssetPickerConfig.OnAssetSelected = FOnAssetSelected::CreateLambda([PoseSearchColumn, &PickerWindow](const FAssetData& SelectedAsset)
-				{
-					if (UPoseSearchSchema* Schema = Cast<UPoseSearchSchema>(SelectedAsset.GetAsset()))
-					{
-						FScopedTransaction ScopedTransaction(LOCTEXT("AssignPoseSearchColumnDatabaseSchema", "Assign Pose Search Column Database Schema"));
-						PoseSearchColumn->SetDatabaseSchema(Schema);
-					}
-					PickerWindow->RequestDestroyWindow();
-				});
-
-			// Load the content browser module to display an asset picker
-			FContentBrowserModule& ContentBrowserModule = FModuleManager::LoadModuleChecked<FContentBrowserModule>("ContentBrowser");
-			PickerWindow = SNew(SWindow)
-				.Title(LOCTEXT("SelectPoseSearchColumnDatabaseSchema", "Pick Schema for owned Pose Search Database"))
-				.ClientSize(FVector2D(500, 600))
-				.SupportsMinimize(false)
-				.SupportsMaximize(false)
-				[
-					SNew(SBorder)
-						.BorderImage(FAppStyle::GetBrush("Menu.Background"))
-						[
-							ContentBrowserModule.Get().CreateAssetPicker(AssetPickerConfig)
-						]
-				];
-		
-			GEditor->EditorAddModalWindow(PickerWindow.ToSharedRef());
-			PickerWindow.Reset();
-		}
-
 		// create column header widget
 		const FSlateBrush* ColumnIcon = FAppStyle::Get().GetBrush("Icons.Search");
 		const FText ColumnTooltip = LOCTEXT("Pose Match Tooltip", "Pose Match: Selects results based on the animations with the best matching poses scored via this column \"Internal Database\", and outputs the StartTime of the best matching animation. AutoPopulate will fill in Column data in case it mismatch the Chooser Table \"Result\" column.");
-		const FText ColumnName = LOCTEXT("Pose Match","Pose Match");
-        		
+		const FText ColumnName = LOCTEXT("Pose Match", "Pose Match");
+
 		TSharedPtr<SWidget> DebugWidget = nullptr;
 
 		return UE::ChooserEditor::MakeColumnHeaderWidget(Chooser, Column, ColumnName, ColumnTooltip, ColumnIcon, DebugWidget);
 	}
 
 	// create cell widget
-	return
-		SNew(STextBlock)
-			.Text_Lambda([RowIndex, PoseSearchColumn]()
+	TSharedRef<SComboButton> Button = SNew(SComboButton)
+		.ContentPadding(0)
+		.ButtonContent()
+		[
+			SNew(STextBlock)
+				.Font(IDetailLayoutBuilder::GetDetailFont())
+				.Text_Lambda([PoseSearchColumn, RowIndex]()
+					{
+						if (UPoseSearchDatabase* Database = PoseSearchColumn->GetDatabase(RowIndex))
+						{
+							return FText::FromString(Database->GetName());
+						}
+						return LOCTEXT("None", "None");
+					})
+		];
+
+	Button->SetOnGetMenuContent(
+		FOnGetContent::CreateLambda(
+			[Button, Chooser, PoseSearchColumn, RowIndex]()
 			{
-				if (const UObject* Asset = PoseSearchColumn->GetDatabaseAsset(RowIndex))
+				FMenuBuilder MenuBuilder(true, nullptr);
+
+				MenuBuilder.BeginSection(NAME_None, LOCTEXT("DatabaseSelection", "Database"));
 				{
-					return FText::FromString(Asset->GetName());
+					MenuBuilder.AddMenuEntry(
+						LOCTEXT("None", "None"),
+						LOCTEXT("NoDatabase_Tooltip", "Disable this row to be indexed in any Database"),
+						FSlateIcon(),
+						FUIAction(FExecuteAction::CreateLambda([Button, Chooser, PoseSearchColumn, RowIndex]()
+							{
+								const FScopedTransaction Transaction(LOCTEXT("Clear Database From Row", "Clear Database"));
+								Chooser->Modify();
+								Button->SetIsOpen(false);
+								PoseSearchColumn->SwitchDatabase(nullptr, RowIndex);
+								Chooser->PostEditChange();
+							})));
+
+					MenuBuilder.AddMenuEntry(
+						LOCTEXT("New Database", "New Database"),
+						LOCTEXT("NewDatabase_Tooltip", "Create a new database used to index this row asset"),
+						FSlateIcon(),
+						FUIAction(FExecuteAction::CreateLambda([Button, Chooser, PoseSearchColumn, RowIndex]()
+							{
+								Button->SetIsOpen(false);
+
+								// picking a schema for the newly created database
+												
+								// Load the content browser module to display an asset picker
+								FContentBrowserModule& ContentBrowserModule = FModuleManager::LoadModuleChecked<FContentBrowserModule>("ContentBrowser");
+
+								FAssetPickerConfig AssetPickerConfig;
+								TSharedPtr<SWindow> PickerWindow;
+								UPoseSearchSchema* TargetSchema = nullptr;
+
+								// The asset picker will only show UPoseSearchSchema(s)
+								AssetPickerConfig.Filter.ClassPaths.Add(UPoseSearchSchema::StaticClass()->GetClassPathName());
+								AssetPickerConfig.Filter.bRecursiveClasses = true;
+
+								// The delegate that fires when an asset was selected
+								AssetPickerConfig.OnAssetSelected = FOnAssetSelected::CreateLambda([&PickerWindow, &TargetSchema](const FAssetData& SelectedAsset)
+									{
+										//FSlateApplication::Get().DismissAllMenus();
+										TargetSchema = Cast<UPoseSearchSchema>(SelectedAsset.GetAsset());
+										PickerWindow->RequestDestroyWindow();
+									});
+													
+								// The default view mode should be a list view
+								AssetPickerConfig.InitialAssetViewType = EAssetViewType::List;
+
+								PickerWindow = SNew(SWindow)
+									.Title(LOCTEXT("CreateDatabaseOptions", "Pick Database Schema"))
+									.ClientSize(FVector2D(500, 600))
+									.SupportsMinimize(false)
+									.SupportsMaximize(false)
+									[
+										SNew(SBorder)
+											.BorderImage(FAppStyle::GetBrush("Menu.Background"))
+											[
+												ContentBrowserModule.Get().CreateAssetPicker(AssetPickerConfig)
+											]
+									];
+
+								GEditor->EditorAddModalWindow(PickerWindow.ToSharedRef());
+								PickerWindow.Reset();
+
+								TSharedRef<STextEntryPopup> TextEntry =
+									SNew(STextEntryPopup)
+									.Label(LOCTEXT("NewDatabaseLabel", "New Database Name"))
+									.OnTextCommitted_Lambda([Chooser, PoseSearchColumn, RowIndex, TargetSchema](FText InText, ETextCommit::Type InCommitType)
+										{
+											if (InCommitType == ETextCommit::OnEnter && TargetSchema)
+											{
+												FSlateApplication::Get().DismissAllMenus();
+												FName NewDatabaseName = FName(InText.ToString());
+
+												UChooserTable* RootChooser = Chooser->GetRootChooser();
+												if (NewDatabaseName == NAME_None || FindObject<UObject>(Chooser, NewDatabaseName.ToString()))
+												{
+													NewDatabaseName = MakeUniqueObjectName(Chooser, UChooserTable::StaticClass(), NewDatabaseName);
+												}
+
+												UPoseSearchDatabase* NewDatabase = NewObject<UPoseSearchDatabase>(RootChooser, UPoseSearchDatabase::StaticClass(), NewDatabaseName, RF_Transactional);
+												NewDatabase->Schema = TargetSchema;
+
+												const FScopedTransaction Transaction(LOCTEXT("Assign New Database", "Assign New Database"));
+												RootChooser->Modify();
+												RootChooser->AddNestedObject(NewDatabase);
+
+												PoseSearchColumn->SwitchDatabase(NewDatabase, RowIndex);
+												RootChooser->PostEditChange();
+											}
+										});
+
+								FSlateApplication& SlateApp = FSlateApplication::Get();
+								SlateApp.PushMenu(
+									SlateApp.GetInteractiveTopLevelWindows()[0],
+									FWidgetPath(),
+									TextEntry,
+									SlateApp.GetCursorPos(),
+									FPopupTransitionEffect::TypeInPopup
+								);
+							})));
+
+					MenuBuilder.AddSubMenu(LOCTEXT("Select Existing", "Select Existing"),
+						LOCTEXT("SelectExisting_Tooltip", "Select Existing Embedded UPoseSearchDatabase from this asset"),
+						FNewMenuDelegate::CreateLambda([Button, Chooser, PoseSearchColumn, RowIndex](FMenuBuilder& SubMenuBuilder)
+							{
+								SubMenuBuilder.BeginSection("Existing", LOCTEXT("Existing Databases", "Existing Databases"));
+
+								if (UChooserTable* RootChooser = Chooser->GetRootChooser())
+								{
+									for (UObject* Object : RootChooser->NestedObjects)
+									{
+										if (UPoseSearchDatabase* Database = Cast<UPoseSearchDatabase>(Object))
+										{
+											SubMenuBuilder.AddMenuEntry(FText::FromString(Database->GetName()), LOCTEXT("AddExistingObjectTooltip", "Add a reference to this existing Database."), FSlateIcon(),
+												FUIAction(FExecuteAction::CreateLambda([Button, Chooser, PoseSearchColumn, Database, RowIndex]()
+													{
+														const FScopedTransaction Transaction(LOCTEXT("Set Database", "Set Database"));
+														Chooser->Modify();
+														Button->SetIsOpen(false);
+														PoseSearchColumn->SwitchDatabase(Database, RowIndex);
+														Chooser->PostEditChange();
+													})));
+										}
+									}
+								}
+							}));
 				}
-				return FText();
-			});
+
+				return MenuBuilder.MakeWidget();
+
+			})
+	);
+
+	return SNew(SHorizontalBox)
+		+ SHorizontalBox::Slot().FillWidth(1.0)
+		[
+			Button
+		]
+		+ SHorizontalBox::Slot().AutoWidth()
+		[
+			SNew(SButton)
+				.Text(LOCTEXT("Edit", "Edit"))
+				.OnClicked_Lambda([PoseSearchColumn, RowIndex]()
+					{
+						// Open editor
+						if (UPoseSearchDatabase* Database = PoseSearchColumn->GetDatabase(RowIndex))
+						{
+							if (UAssetEditorSubsystem* AssetEditorSS = GEditor->GetEditorSubsystem<UAssetEditorSubsystem>())
+							{
+								AssetEditorSS->OpenEditorForAsset(Database);
+							}
+						}
+
+						return FReply::Handled();
+					})
+		];
 }
 
+
 void RegisterPoseSearchChooserWidgets()
 {
 	UE::ChooserEditor::FObjectChooserWidgetFactories::RegisterColumnWidgetCreator(FPoseSearchColumn::StaticStruct(), CreatePoseSearchColumnWidget);
 }
-	
-}
+
+} // UE::PoseSearchEditor
 
 #undef LOCTEXT_NAMESPACE
diff --git a/Engine/Plugins/Animation/PoseSearch/Source/Editor/Private/PoseSearchDatabaseAssetListItem.cpp b/Engine/Plugins/Animation/PoseSearch/Source/Editor/Private/PoseSearchDatabaseAssetListItem.cpp
--- a/Engine/Plugins/Animation/PoseSearch/Source/Editor/Private/PoseSearchDatabaseAssetListItem.cpp
+++ b/Engine/Plugins/Animation/PoseSearch/Source/Editor/Private/PoseSearchDatabaseAssetListItem.cpp
@@ -847,6 +847,8 @@ namespace UE::PoseSearch
 				PoseSearchDatabase->Modify();
 
 				ViewModelPtr->SetDisableReselection(AssetTreeNode->SourceAssetIdx, NewCheckboxState == ECheckBoxState::Checked ? true : false);
+			
+				PoseSearchDatabase->PostEditChange();
 
 				SkeletonView.Pin()->RefreshTreeView(false, true);
 			}
diff --git a/Engine/Plugins/Animation/PoseSearch/Source/Editor/Private/PoseSearchDatabaseAssetTree.cpp b/Engine/Plugins/Animation/PoseSearch/Source/Editor/Private/PoseSearchDatabaseAssetTree.cpp
--- a/Engine/Plugins/Animation/PoseSearch/Source/Editor/Private/PoseSearchDatabaseAssetTree.cpp
+++ b/Engine/Plugins/Animation/PoseSearch/Source/Editor/Private/PoseSearchDatabaseAssetTree.cpp
@@ -32,10 +32,6 @@
 
 namespace UE::PoseSearch
 {
-	SDatabaseAssetTree::~SDatabaseAssetTree()
-	{
-	}
-
 	void SDatabaseAssetTree::Construct(
 		const FArguments& InArgs, 
 		TSharedRef<FDatabaseViewModel> InEditorViewModel)
@@ -226,7 +222,7 @@ namespace UE::PoseSearch
 
 		// Empty node data.
 		RootNodes.Reset();
-		AllNodes.Reset();
+		//AllNodes.Reset();
 
 		const UPoseSearchDatabase* Database = ViewModel->GetPoseSearchDatabase();
 		if (!IsValid(Database))
@@ -293,7 +289,7 @@ namespace UE::PoseSearch
 
 						// Keep track of node
 						RootNodes.Add(SequenceGroupNode);
-						AllNodes.Add(SequenceGroupNode);
+						//AllNodes.Add(SequenceGroupNode);
 					}
 				}
 			}
@@ -421,6 +417,8 @@ namespace UE::PoseSearch
 				++AddedAssets;
 			}
 			
+			PoseSearchDatabase->PostEditChange();
+
 			GWarn->EndSlowTask();
 		}
 
@@ -689,7 +687,7 @@ namespace UE::PoseSearch
 	{
 		TArray<TSharedPtr<FDatabaseAssetTreeNode>> NewSelectedNodes;
 
-		for (const TSharedPtr<FDatabaseAssetTreeNode>& Node : AllNodes)
+		for (const TSharedPtr<FDatabaseAssetTreeNode>& Node : RootNodes)
 		{
 			const bool bFoundNode = PreviouslySelectedNodes.ContainsByPredicate([Node](const TSharedPtr<FDatabaseAssetTreeNode>& PrevSelectedNode) { return PrevSelectedNode->SourceAssetIdx == Node->SourceAssetIdx; });
 			if (bFoundNode)
@@ -903,6 +901,8 @@ namespace UE::PoseSearch
 					ViewModel->SetIsEnabled(SelectedNode->SourceAssetIdx, bIsEnabled);
 				}
 
+				PoseSearchDatabase->PostEditChange();
+
 				FinalizeTreeChanges();
 			}
 		}
@@ -925,6 +925,8 @@ namespace UE::PoseSearch
 					ViewModel->SetMirrorOption(SelectedNode->SourceAssetIdx, InMirrorOption);
 				}
 
+				PoseSearchDatabase->PostEditChange();
+
 				FinalizeTreeChanges();
 			}
 		}
@@ -947,6 +949,8 @@ namespace UE::PoseSearch
 					ViewModel->SetDisableReselection(SelectedNode->SourceAssetIdx, !bIsEnabled);
 				}
 
+				PoseSearchDatabase->PostEditChange();
+
 				RefreshTreeView(false, true);
 			}
 		}
@@ -966,17 +970,19 @@ namespace UE::PoseSearch
 				
 				for (const TSharedPtr<FDatabaseAssetTreeNode>& SelectedNode : SelectedNodes)
 				{
-					if (FPoseSearchDatabaseAnimationAssetBase* DatabaseAnimationAssetBase = PoseSearchDatabase->GetMutableDatabaseAnimationAsset(SelectedNode->SourceAssetIdx))
+					if (const FPoseSearchDatabaseAnimationAsset* DatabaseAnimationAsset = PoseSearchDatabase->GetDatabaseAnimationAsset(SelectedNode->SourceAssetIdx))
 					{
-						if (UAnimSequenceBase* AnimSequenceBase = Cast<UAnimSequenceBase>(DatabaseAnimationAssetBase->GetAnimationAsset()))
+						if (UAnimSequenceBase* AnimSequenceBase = Cast<UAnimSequenceBase>(DatabaseAnimationAsset->GetAnimationAsset()))
 						{
 							if (!bModified)
 							{
-								AnimSequenceBase->Modify();
+								PoseSearchDatabase->Modify();
 								bModified = true;
 							}
 							
-							const FFloatInterval SamplingRange = DatabaseAnimationAssetBase->GetEffectiveSamplingRange(FVector::ZeroVector);
+							AnimSequenceBase->Modify();
+
+							const FFloatInterval SamplingRange = DatabaseAnimationAsset->GetEffectiveSamplingRange(FVector::ZeroVector);
 							const float StartTime = SamplingRange.Min;
 							const float Duration = SamplingRange.Max - SamplingRange.Min;
 							const FName TrackName = "PoseSearch";
@@ -988,11 +994,22 @@ namespace UE::PoseSearch
 
 							UAnimNotifyState_PoseSearchBranchIn* PoseSearchBranchIn = CastChecked<UAnimNotifyState_PoseSearchBranchIn>(UAnimationBlueprintLibrary::AddAnimationNotifyStateEvent(AnimSequenceBase, TrackName, StartTime, Duration, UAnimNotifyState_PoseSearchBranchIn::StaticClass()));
 							PoseSearchBranchIn->Database = PoseSearchDatabase;
-							DatabaseAnimationAssetBase->BranchInId = PoseSearchBranchIn->GetBranchInId();
+						
+							AnimSequenceBase->PostEditChange();
+
+							FPoseSearchDatabaseAnimationAsset NewDatabaseAnimationAsset = *DatabaseAnimationAsset;
+							NewDatabaseAnimationAsset.BranchInId = PoseSearchBranchIn->GetBranchInId();
+
+							PoseSearchDatabase->SetAnimationAssetAt(NewDatabaseAnimationAsset, SelectedNode->SourceAssetIdx);
 						}
 					}
 				}
 
+				if (bModified)
+				{
+					PoseSearchDatabase->PostEditChange();
+				}
+
 				PoseSearchDatabase->SynchronizeWithExternalDependencies();
 
 				FinalizeTreeChanges();
@@ -1019,7 +1036,7 @@ namespace UE::PoseSearch
 
 		if (SourceAssetIdx >= 0)
 		{
-			for (TSharedPtr<FDatabaseAssetTreeNode>& Node : AllNodes)
+			for (TSharedPtr<FDatabaseAssetTreeNode>& Node : RootNodes)
 			{
 				if (Node->SourceAssetIdx == SourceAssetIdx)
 				{
diff --git a/Engine/Plugins/Animation/PoseSearch/Source/Editor/Private/PoseSearchDatabaseEditorClipboard.cpp b/Engine/Plugins/Animation/PoseSearch/Source/Editor/Private/PoseSearchDatabaseEditorClipboard.cpp
--- a/Engine/Plugins/Animation/PoseSearch/Source/Editor/Private/PoseSearchDatabaseEditorClipboard.cpp
+++ b/Engine/Plugins/Animation/PoseSearch/Source/Editor/Private/PoseSearchDatabaseEditorClipboard.cpp
@@ -57,6 +57,8 @@ void UPoseSearchDatabaseEditorClipboardContent::PasteToDatabase(UPoseSearchDatab
 		FPoseSearchDatabaseAnimationAsset::StaticStruct()->ImportText(*Item->Content, &NewAsset, nullptr, PPF_None, GLog, FPoseSearchDatabaseAnimationAsset::StaticStruct()->GetName());
 		InTargetDatabase->AddAnimationAsset(NewAsset);
 	}
+
+	InTargetDatabase->PostEditChange();
 }
 
 class FPoseSearchDatabaseEditorClipboardContentTextFactory : public FCustomizableTextObjectFactory
diff --git a/Engine/Plugins/Animation/PoseSearch/Source/Editor/Private/PoseSearchDatabaseEditorReflection.cpp b/Engine/Plugins/Animation/PoseSearch/Source/Editor/Private/PoseSearchDatabaseEditorReflection.cpp
--- a/Engine/Plugins/Animation/PoseSearch/Source/Editor/Private/PoseSearchDatabaseEditorReflection.cpp
+++ b/Engine/Plugins/Animation/PoseSearch/Source/Editor/Private/PoseSearchDatabaseEditorReflection.cpp
@@ -21,36 +21,6 @@ void UPoseSearchDatabaseReflectionBase::SetSourceLink(
 	AssetTreeWidget = InAssetTreeWidget;
 }
 
-void UPoseSearchDatabaseReflectionBase::PostEditChangeProperty(FPropertyChangedEvent& PropertyChangedEvent)
-{
-	Super::PostEditChangeProperty(PropertyChangedEvent);
-
-	// Skip changes during EPropertyChangeType::Interactive since they are always followed by a PostEditChangeProperty() call
-	// with EPropertyChangeType::ValueSet holding the final values.
-	if (PropertyChangedEvent.ChangeType == EPropertyChangeType::Interactive)
-	{
-		return;
-	}
-	
-	if (ApplyChanges())
-	{
-		const bool bShouldRefreshView = (PropertyChangedEvent.Property == nullptr) || !PropertyChangedEvent.Property->IsA(FFloatProperty::StaticClass());
-		AssetTreeWidget->FinalizeTreeChanges(true, bShouldRefreshView);
-
-		if (const TSharedPtr<UE::PoseSearch::FDatabaseAssetTreeNode> AssetTreeNode = WeakAssetTreeNode.Pin())
-		{
-			if (const TSharedPtr<UE::PoseSearch::FDatabaseViewModel> ViewModel = AssetTreeNode->EditorViewModel.Pin())
-			{
-				UPoseSearchDatabase* Database = ViewModel->GetPoseSearchDatabase();
-				if (IsValid(Database))
-				{
-					Database->Modify();
-				}
-			}
-		}
-	}
-}
-
 bool UPoseSearchDatabaseAnimationAssetReflection::ApplyChanges() 
 {
 	if (UObject* AnimationAsset = DatabaseAnimationAsset.GetAnimationAsset())
@@ -73,11 +43,11 @@ bool UPoseSearchDatabaseAnimationAssetReflection::ApplyChanges()
 			UPoseSearchDatabase* Database = ViewModel->GetPoseSearchDatabase();
 			if (IsValid(Database))
 			{
-				if (FPoseSearchDatabaseAnimationAsset* DatabaseAnimationAssetPtr = Database->GetMutableDatabaseAnimationAsset(AssetTreeNode->SourceAssetIdx))
+				if (const FPoseSearchDatabaseAnimationAsset* DatabaseAnimationAssetPtr = Database->GetDatabaseAnimationAsset(AssetTreeNode->SourceAssetIdx))
 				{
 					if (*DatabaseAnimationAssetPtr != DatabaseAnimationAsset)
 					{
-						*DatabaseAnimationAssetPtr = DatabaseAnimationAsset;
+						Database->SetAnimationAssetAt(DatabaseAnimationAsset, AssetTreeNode->SourceAssetIdx);
 						return true;
 					}
 				}
diff --git a/Engine/Plugins/Animation/PoseSearch/Source/Editor/Private/PoseSearchDatabaseViewModel.cpp b/Engine/Plugins/Animation/PoseSearch/Source/Editor/Private/PoseSearchDatabaseViewModel.cpp
--- a/Engine/Plugins/Animation/PoseSearch/Source/Editor/Private/PoseSearchDatabaseViewModel.cpp
+++ b/Engine/Plugins/Animation/PoseSearch/Source/Editor/Private/PoseSearchDatabaseViewModel.cpp
@@ -647,6 +647,8 @@ void FDatabaseViewModel::AddAnimationAssetToDatabase(UObject* AnimationAsset)
 		FPoseSearchDatabaseAnimationAsset NewAsset;
 		NewAsset.AnimAsset = AnimationAsset;
 		Database->AddAnimationAsset(NewAsset);
+
+		Database->PostEditChange();
 	}
 }
 
@@ -697,6 +699,8 @@ bool FDatabaseViewModel::DeleteFromDatabase(int32 AnimationAssetIndex)
 
 			Database->RemoveAnimationAssetAt(AnimationAssetIndex);
 
+			Database->PostEditChange();
+
 			return true;
 		}
 	}
@@ -708,11 +712,15 @@ void FDatabaseViewModel::SetDisableReselection(int32 AnimationAssetIndex, bool b
 {
 	if (UPoseSearchDatabase* Database = GetPoseSearchDatabase())
 	{
-		if (FPoseSearchDatabaseAnimationAssetBase* DatabaseAnimationAsset = Database->GetMutableDatabaseAnimationAsset(AnimationAssetIndex))
+		if (const FPoseSearchDatabaseAnimationAsset* DatabaseAnimationAsset = Database->GetDatabaseAnimationAsset(AnimationAssetIndex))
 		{
 			Database->Modify();
 			
-			DatabaseAnimationAsset->SetDisableReselection(bEnabled);
+			FPoseSearchDatabaseAnimationAsset NewDatabaseAnimationAsset = *DatabaseAnimationAsset;
+			NewDatabaseAnimationAsset.SetDisableReselection(bEnabled);
+			Database->SetAnimationAssetAt(NewDatabaseAnimationAsset, AnimationAssetIndex);
+
+			Database->PostEditChange();
 		}
 	}
 }
@@ -734,11 +742,15 @@ void FDatabaseViewModel::SetIsEnabled(int32 AnimationAssetIndex, bool bEnabled)
 {
 	if (UPoseSearchDatabase* Database = GetPoseSearchDatabase())
 	{
-		if (FPoseSearchDatabaseAnimationAssetBase* DatabaseAnimationAsset = Database->GetMutableDatabaseAnimationAsset(AnimationAssetIndex))
+		if (const FPoseSearchDatabaseAnimationAsset* DatabaseAnimationAsset = Database->GetDatabaseAnimationAsset(AnimationAssetIndex))
 		{
 			Database->Modify();
 
-			DatabaseAnimationAsset->SetIsEnabled(bEnabled);
+			FPoseSearchDatabaseAnimationAsset NewDatabaseAnimationAsset = *DatabaseAnimationAsset;
+			NewDatabaseAnimationAsset.SetIsEnabled(bEnabled);
+			Database->SetAnimationAssetAt(NewDatabaseAnimationAsset, AnimationAssetIndex);
+
+			Database->PostEditChange();
 		}
 	}
 }
@@ -747,7 +759,7 @@ bool FDatabaseViewModel::IsEnabled(int32 AnimationAssetIndex) const
 {
 	if (const UPoseSearchDatabase* Database = GetPoseSearchDatabase())
 	{
-		if (const FPoseSearchDatabaseAnimationAssetBase* DatabaseAnimationAsset = Database->GetDatabaseAnimationAsset(AnimationAssetIndex))
+		if (const FPoseSearchDatabaseAnimationAsset* DatabaseAnimationAsset = Database->GetDatabaseAnimationAsset(AnimationAssetIndex))
 		{
 			return DatabaseAnimationAsset->IsEnabled();
 		}
@@ -762,10 +774,15 @@ bool FDatabaseViewModel::SetAnimationAsset(int32 AnimationAssetIndex, UObject* A
 	{
 		if (UPoseSearchDatabase* Database = GetPoseSearchDatabase())
 		{
-			if (FPoseSearchDatabaseAnimationAsset* DatabaseAnimationAsset = Database->GetMutableDatabaseAnimationAsset(AnimationAssetIndex))
+			if (const FPoseSearchDatabaseAnimationAsset* DatabaseAnimationAsset = Database->GetDatabaseAnimationAsset(AnimationAssetIndex))
 			{
 				Database->Modify();
-				DatabaseAnimationAsset->AnimAsset = AnimAsset;
+
+				FPoseSearchDatabaseAnimationAsset NewDatabaseAnimationAsset = *DatabaseAnimationAsset;
+				NewDatabaseAnimationAsset.AnimAsset = AnimAsset;
+				Database->SetAnimationAssetAt(NewDatabaseAnimationAsset, AnimationAssetIndex);
+
+				Database->PostEditChange();
 				return true;
 			}
 		}
@@ -778,11 +795,15 @@ void FDatabaseViewModel::SetMirrorOption(int32 AnimationAssetIndex, EPoseSearchM
 {
 	if (UPoseSearchDatabase* Database = GetPoseSearchDatabase())
 	{
-		if (FPoseSearchDatabaseAnimationAssetBase* DatabaseAnimationAsset = Database->GetMutableDatabaseAnimationAsset(AnimationAssetIndex))
+		if (const FPoseSearchDatabaseAnimationAsset* DatabaseAnimationAsset = Database->GetDatabaseAnimationAsset(AnimationAssetIndex))
 		{
 			Database->Modify();
 
-			DatabaseAnimationAsset->MirrorOption = InMirrorOption;
+			FPoseSearchDatabaseAnimationAsset NewDatabaseAnimationAsset = *DatabaseAnimationAsset;
+			NewDatabaseAnimationAsset.MirrorOption = InMirrorOption;
+			Database->SetAnimationAssetAt(NewDatabaseAnimationAsset, AnimationAssetIndex);
+
+			Database->PostEditChange();
 		}
 	}
 }
@@ -791,7 +812,7 @@ EPoseSearchMirrorOption FDatabaseViewModel::GetMirrorOption(int32 AnimationAsset
 {
 	if (const UPoseSearchDatabase* Database = GetPoseSearchDatabase())
 	{
-		if (const FPoseSearchDatabaseAnimationAssetBase* DatabaseAnimationAsset = Database->GetDatabaseAnimationAsset(AnimationAssetIndex))
+		if (const FPoseSearchDatabaseAnimationAsset* DatabaseAnimationAsset = Database->GetDatabaseAnimationAsset(AnimationAssetIndex))
 		{
 			return DatabaseAnimationAsset->MirrorOption;
 		}
@@ -1028,7 +1049,7 @@ bool FDatabaseViewModel::GetAnimationTime(int32 SourceAssetIdx, float& CurrentPl
 					CurrentPlayTime = PlayTime + IndexAsset.GetFirstSampleTime(Database->Schema->SampleRate);
 					BlendParameters = IndexAsset.GetBlendParameters();
 
-					if (const UObject* AnimationAsset = Database->GetDatabaseAnimationAsset(IndexAsset)->GetAnimationAsset())
+					if (const UObject* AnimationAsset = Database->GetDatabaseAnimationAsset(IndexAsset.GetSourceAssetIdx())->GetAnimationAsset())
 					{
 						if (AnimationAsset->IsA<UBlendSpace>() && !FMath::IsNearlyEqual(MaxPreviewPlayLength, MinPreviewPlayLength))
 						{
diff --git a/Engine/Plugins/Animation/PoseSearch/Source/Editor/Private/PoseSearchDebuggerDatabaseView.cpp b/Engine/Plugins/Animation/PoseSearch/Source/Editor/Private/PoseSearchDebuggerDatabaseView.cpp
--- a/Engine/Plugins/Animation/PoseSearch/Source/Editor/Private/PoseSearchDebuggerDatabaseView.cpp
+++ b/Engine/Plugins/Animation/PoseSearch/Source/Editor/Private/PoseSearchDebuggerDatabaseView.cpp
@@ -231,7 +231,7 @@ static void AddUnfilteredDatabaseRow(const UPoseSearchDatabase* Database, TConst
 			}
 		}
 
-		if (const FPoseSearchDatabaseAnimationAssetBase* DatabaseAsset = Database->GetDatabaseAnimationAsset(*SearchIndexAsset))
+		if (const FPoseSearchDatabaseAnimationAssetBase* DatabaseAsset = Database->GetDatabaseAnimationAsset(SearchIndexAsset->GetSourceAssetIdx()))
 		{
 			const FVector BlendParameters = SearchIndexAsset->GetBlendParameters();
 			const float PlayLength = DatabaseAsset->GetPlayLength(BlendParameters);
diff --git a/Engine/Plugins/Animation/PoseSearch/Source/Editor/Public/PoseSearchDatabaseAssetTree.h b/Engine/Plugins/Animation/PoseSearch/Source/Editor/Public/PoseSearchDatabaseAssetTree.h
--- a/Engine/Plugins/Animation/PoseSearch/Source/Editor/Public/PoseSearchDatabaseAssetTree.h
+++ b/Engine/Plugins/Animation/PoseSearch/Source/Editor/Public/PoseSearchDatabaseAssetTree.h
@@ -20,8 +20,6 @@ namespace UE::PoseSearch
 		SLATE_BEGIN_ARGS(SDatabaseAssetTree) {}
 		SLATE_END_ARGS()
 
-		virtual ~SDatabaseAssetTree();
-
 		void Construct(const FArguments& InArgs, TSharedRef<FDatabaseViewModel> InEditorViewModel);
 
 		// SWidget interface
@@ -51,7 +49,6 @@ namespace UE::PoseSearch
 		TSharedPtr<STreeView<TSharedPtr<FDatabaseAssetTreeNode>>> TreeView;
 		TSharedPtr<SWidget> TreeViewDragAndDropSuggestion;
 		TArray<TSharedPtr<FDatabaseAssetTreeNode>> RootNodes;
-		TArray<TSharedPtr<FDatabaseAssetTreeNode>> AllNodes;
 
 		TSharedRef<ITableRow> MakeTableRowWidget(
 			TSharedPtr<FDatabaseAssetTreeNode> InItem,
diff --git a/Engine/Plugins/Animation/PoseSearch/Source/Editor/Public/PoseSearchDatabaseEditorReflection.h b/Engine/Plugins/Animation/PoseSearch/Source/Editor/Public/PoseSearchDatabaseEditorReflection.h
--- a/Engine/Plugins/Animation/PoseSearch/Source/Editor/Public/PoseSearchDatabaseEditorReflection.h
+++ b/Engine/Plugins/Animation/PoseSearch/Source/Editor/Public/PoseSearchDatabaseEditorReflection.h
@@ -25,7 +25,6 @@ public:
 		const TSharedPtr<UE::PoseSearch::SDatabaseAssetTree>& InAssetTreeWidget);
 	
 	virtual bool ApplyChanges() { return false; };
-	virtual void PostEditChangeProperty(FPropertyChangedEvent& PropertyChangedEvent) override;
 	
 protected:
 	TWeakPtr<UE::PoseSearch::FDatabaseAssetTreeNode> WeakAssetTreeNode;
diff --git a/Engine/Plugins/Animation/PoseSearch/Source/Runtime/Private/Chooser/PoseSearchChooserColumn.cpp b/Engine/Plugins/Animation/PoseSearch/Source/Runtime/Private/Chooser/PoseSearchChooserColumn.cpp
--- a/Engine/Plugins/Animation/PoseSearch/Source/Runtime/Private/Chooser/PoseSearchChooserColumn.cpp
+++ b/Engine/Plugins/Animation/PoseSearch/Source/Runtime/Private/Chooser/PoseSearchChooserColumn.cpp
@@ -14,6 +14,7 @@
 #include "PoseSearch/PoseSearchDatabase.h"
 #include "PoseSearch/PoseSearchLibrary.h"
 #include "PoseSearch/PoseSearchSchema.h"
+#include "Templates/Greater.h"
 #if WITH_EDITOR
 #include "StructUtils/PropertyBag.h"
 #endif
@@ -22,50 +23,60 @@
 
 #define LOCTEXT_NAMESPACE "FPoseSearchColumn"
 
+//#define ENSURE_MAPPING_VALIDATION(InExpression) ensure(InExpression)
+#define ENSURE_MAPPING_VALIDATION(InExpression) 
+//#define DO_ENSURE_CONSISTENCY(InExpression) ensure(InExpression)
+#define DO_ENSURE_CONSISTENCY(InExpression) 
+
 namespace UE::PoseSearch
 {
-	int32 GetDatabaseAssetIndex(int32 RowIndex)
+	// searching for FChooserPlayerSettings in the Context.Params
+	static const FChooserPlayerSettings* GetChooserPlayerSettings(const FChooserEvaluationContext& Context)
 	{
-		if (RowIndex == ChooserColumn_SpecialIndex_Fallback)
+		for (const FStructView& Param : Context.Params)
 		{
-			return 0;
+			if (const FChooserPlayerSettings* ChooserPlayerSettings = Param.GetPtr<FChooserPlayerSettings>())
+			{
+				return ChooserPlayerSettings;
+			}
 		}
-
-		check(RowIndex >= 0);
-
-		return RowIndex + 1;
+		return nullptr;
 	}
 
-	int32 GetRowIndex(int32 DatabaseAssetIndex)
+#if WITH_EDITOR
+	static const FInstancedStruct& GetResult(const UChooserTable* Chooser, int32 RowIndex)
 	{
-		check(DatabaseAssetIndex >= 0);
+		check(Chooser);
+		if (RowIndex == ChooserColumn_SpecialIndex_Fallback)
+		{
+			return Chooser->FallbackResult;
+		}
 
-		if (DatabaseAssetIndex == 0)
+#if WITH_EDITORONLY_DATA
+		if (!Chooser->IsCookedData())
 		{
-			return ChooserColumn_SpecialIndex_Fallback;
+			return Chooser->ResultsStructs[RowIndex];
 		}
-		
-		return DatabaseAssetIndex - 1;
+#endif
+		return Chooser->CookedResults[RowIndex];
 	}
 
-	// searching for FChooserPlayerSettings in the Context.Params
-	static const FChooserPlayerSettings* GetChooserPlayerSettings(const FChooserEvaluationContext& Context)
+	static UObject* GetReferencedObject(const UChooserTable* Chooser, int32 RowIndex)
 	{
-		for (const FStructView& Param : Context.Params)
+		const FInstancedStruct& Result = GetResult(Chooser, RowIndex);
+		if (Result.IsValid())
 		{
-			if (const FChooserPlayerSettings* ChooserPlayerSettings = Param.GetPtr<FChooserPlayerSettings>())
-			{
-				return ChooserPlayerSettings;
-			}
+			return Result.Get<FObjectChooserBase>().GetReferencedObject();
 		}
 		return nullptr;
 	}
 
-	// Experimental, this feature might be removed without warning, not for production use
+#endif // WITH_EDITOR
+
 	struct FPoseSearchColumnScratchArea
 	{
-		TArray<UE::PoseSearch::FSearchResult> SearchResults;
-		const UE::PoseSearch::IPoseHistory* PoseHistory = nullptr;
+		TArray<FSearchResult> SearchResults;
+		const IPoseHistory* PoseHistory = nullptr;
 		const FChooserPlayerSettings* ChooserPlayerSettings = nullptr;
 
 #if DO_CHECK
@@ -79,359 +90,569 @@ namespace UE::PoseSearch
 		Ar << ActiveColumnCost.RowCost;
 		return Ar;
 	}
-} // namespace UE::PoseSearch
 
-bool FPoseHistoryContextProperty::GetValue(FChooserEvaluationContext& Context, FPoseHistoryReference& OutResult) const
-{
-	return Binding.GetStructValue(Context, OutResult);
-}
+	typedef TArray<FPoseSearchColumn*, TInlineAllocator<16>> FPoseSearchColumnPtrs;
+	typedef TSet<const UChooserTable*, DefaultKeyFuncs<const UChooserTable*>, TInlineSetAllocator<32>> FVisitedChoosers;
+	typedef TFunction<void(const FInstancedStruct& ResultsStruct, const UChooserTable* Chooser, int32 RowIndex, const FPoseSearchColumnPtrs& PoseSearchColumnPtrs)> FProcessRow;
 
-FPoseSearchColumn::FPoseSearchColumn()
-{
-}
+	static void IterateChooserRows(UChooserTable* Chooser, FVisitedChoosers& VisitedChoosers, const FProcessRow& ProcessRow);
 
-void FPoseSearchColumn::CheckConsistency() const
-{
-#if DO_CHECK && WITH_EDITORONLY_DATA
+	static void ProcessResultsStruct(const FInstancedStruct& ResultsStruct, const UChooserTable* Chooser, int32 RowIndex, FVisitedChoosers& VisitedChoosers, const FPoseSearchColumnPtrs& PoseSearchColumnPtrs, const FProcessRow& ProcessRow)
+	{
+		if (ResultsStruct.IsValid())
+		{
+			if (const FNestedChooser* NestedChooser = ResultsStruct.GetPtr<FNestedChooser>())
+			{
+				if (NestedChooser->Chooser)
+				{
+					IterateChooserRows(NestedChooser->Chooser, VisitedChoosers, ProcessRow);
+				}
+			}
+			else
+			{
+				ProcessRow(ResultsStruct, Chooser, RowIndex, PoseSearchColumnPtrs);
+			}
+		}
+	}
 
-	using namespace UE::PoseSearch;
+	static void IterateChooserRows(UChooserTable* Chooser, FVisitedChoosers& VisitedChoosers, const FProcessRow& ProcessRow)
+	{
+		ensure(Chooser);
 
-	check(InternalDatabase);
-	UChooserTable* OuterChooser = Cast<UChooserTable>(InternalDatabase->GetOuter());
-	check(OuterChooser);
+		bool bAlreadyVisited = false;
+		VisitedChoosers.Add(Chooser, &bAlreadyVisited);
+		if (bAlreadyVisited)
+		{
+			return;
+		}
 
-	const int32 NumAnimationAssets = InternalDatabase->GetNumAnimationAssets();
-	const int32 NumRows = GetRowIndex(NumAnimationAssets);
-	if (NumRows != OuterChooser->ResultsStructs.Num())
-	{
-		UE_LOG(LogPoseSearch, Error, TEXT("FPoseSearchColumn::CheckConsistency for Chooser Table '%s' failed because the number of rows (%d) differs from the internal database number of assets %d. Click 'AutoPopulate All' button in the Chooser Table Editor to try resynchronize the data"), *OuterChooser->GetName(), OuterChooser->ResultsStructs.Num(), NumRows);
+		int32 NumRows = 0;
+#if WITH_EDITORONLY_DATA
+		if (!Chooser->IsCookedData())
+		{
+			NumRows = Chooser->ResultsStructs.Num();
+		}
+		else
+#endif
+		{
+			NumRows = Chooser->CookedResults.Num();
+		}
+
+		FPoseSearchColumnPtrs PoseSearchColumnPtrs;
+		for (FInstancedStruct& ColumnData : Chooser->ColumnsStructs)
+		{
+			if (FPoseSearchColumn* PoseSearchColumn = ColumnData.GetMutablePtr<FPoseSearchColumn>())
+			{
+				// adding only columns that has already been upgraded / correcly set up
+				if (NumRows == PoseSearchColumn->GetNumRows())
+				{
+#if WITH_EDITOR
+					// making sure Initialize has been called on PoseSearchColumn
+					ensure(PoseSearchColumn->GetChooser());
+#endif // WITH_EDITOR
+					PoseSearchColumnPtrs.Add(PoseSearchColumn);
+				}
+			}
+		}
+
+#if WITH_EDITORONLY_DATA
+		if (!Chooser->IsCookedData())
+		{
+			for (int32 ResultIndex = 0; ResultIndex < Chooser->ResultsStructs.Num(); ++ResultIndex)
+			{
+				if (!Chooser->IsRowDisabled(ResultIndex))
+				{
+					ProcessResultsStruct(Chooser->ResultsStructs[ResultIndex], Chooser, ResultIndex, VisitedChoosers, PoseSearchColumnPtrs, ProcessRow);
+				}
+			}
+		}
+		else
+#endif
+		{
+			for (int32 ResultIndex = 0; ResultIndex < Chooser->CookedResults.Num(); ++ResultIndex)
+			{
+				ProcessResultsStruct(Chooser->CookedResults[ResultIndex], Chooser, ResultIndex, VisitedChoosers, PoseSearchColumnPtrs, ProcessRow);
+			}
+		}
+
+		if (Chooser->FallbackResult.IsValid())
+		{
+			ProcessResultsStruct(Chooser->FallbackResult, Chooser, ChooserColumn_SpecialIndex_Fallback, VisitedChoosers, PoseSearchColumnPtrs, ProcessRow);
+		}
 	}
 
-	for (int32 AnimationAssetIndex = 0; AnimationAssetIndex < NumAnimationAssets; ++AnimationAssetIndex)
+	static bool CheckConsistency(const UChooserTable* Chooser)
 	{
-		const FPoseSearchDatabaseAnimationAsset* DatabaseAnimationAsset = InternalDatabase->GetDatabaseAnimationAsset(AnimationAssetIndex);
-		check(DatabaseAnimationAsset);
-		
-		const UObject* AnimationAsset = DatabaseAnimationAsset->GetAnimationAsset();
-		const int32 RowIndex = GetRowIndex(AnimationAssetIndex);
-		const UObject* ReferencedObject = GetReferencedObject(RowIndex);
-		if (AnimationAsset != ReferencedObject)
+		bool bSuccess = true;
+#if WITH_EDITOR
+		if (!Chooser)
+		{
+			bSuccess = false;
+		}
+		else if (!Chooser->GetRootChooser())
+		{
+			bSuccess = false;
+		}
+		else
 		{
-			UE_LOG(LogPoseSearch, Error, TEXT("FPoseSearchColumn::CheckConsistency for Chooser Table '%s' failed for row %d: chooser table asset '%s' differs from internal database asset ' % s'. Click 'AutoPopulate All' button in the Chooser Table Editor to try resynchronize the data"), *OuterChooser->GetName(), RowIndex, *GetNameSafe(ReferencedObject), *GetNameSafe(AnimationAsset));
+			FVisitedChoosers VisitedChoosers;
+			IterateChooserRows(const_cast<UChooserTable*>(Chooser->GetRootChooser()), VisitedChoosers,
+				[&bSuccess](const FInstancedStruct& ResultsStruct, const UChooserTable* Chooser, int32 RowIndex, const FPoseSearchColumnPtrs& PoseSearchColumnPtrs)
+				{
+					for (FPoseSearchColumn* PoseSearchColumnPtr : PoseSearchColumnPtrs)
+					{
+						if (UPoseSearchDatabase* Database = PoseSearchColumnPtr->GetRowValue(RowIndex).Database)
+						{
+							if (Database->GetOuter() != Chooser->GetRootChooser())
+							{
+								bSuccess = false;
+							}
+						}
+					}
+				});
 		}
+#endif // WITH_EDITOR
+		return bSuccess;
 	}
 
-#endif // DO_CHECK && WITH_EDITORONLY_DATA 
-}
+	FMapping::FMapping(UChooserTable* Chooser)
+	{
+		Init(Chooser);
+	}
 
-#if WITH_EDITOR
+	void FMapping::Init(UChooserTable* Chooser)
+	{
+		Map.Reset();
 
-UObject* FPoseSearchColumn::GetReferencedObject(int32 RowIndex) const
-{
-	using namespace UE::PoseSearch;
+		if (Chooser)
+		{
+			check(Chooser == Chooser->GetRootChooser());
+			FVisitedChoosers VisitedChoosers;
+			IterateChooserRows(Chooser, VisitedChoosers,
+				[this](const FInstancedStruct& ResultsStruct, const UChooserTable* Chooser, int32 RowIndex, const FPoseSearchColumnPtrs& PoseSearchColumnPtrs)
+				{
+					for (FPoseSearchColumn* PoseSearchColumnPtr : PoseSearchColumnPtrs)
+					{
+						UPoseSearchDatabase* Database = PoseSearchColumnPtr->GetRowValue(RowIndex).Database;
+						ensure(!Database || Database->GetOuter() == Chooser->GetRootChooser());
+						Map.FindOrAdd(Database).Add({ PoseSearchColumnPtr, RowIndex });
+					}
+				});
+		}
+	}
 
-	UChooserTable* OuterChooser = Cast<UChooserTable>(InternalDatabase->GetOuter());
-	check(OuterChooser);
+	bool FMapping::Equals(const FMapping& Other) const
+	{
+		return Map.OrderIndependentCompareEqual(Other.Map);
+	}
 
-	UObject* ReferencedObject = nullptr;
+	int32 FMapping::GetNumIndexesInDatabase(const UPoseSearchDatabase* Database) const
+	{
+		if (const FTableCells* TableCells = Map.Find(Database))
+		{
+			return TableCells->Num();
+		}
+		return 0;
+	}
 
-	if (RowIndex != ChooserColumn_SpecialIndex_Fallback)
+	int32 FMapping::FindIndexInDatabase(int32 RowIndex, const FPoseSearchColumn* Column, const UPoseSearchDatabase* Database) const
 	{
-		check(RowIndex >= 0);
-		const FInstancedStruct& Result = OuterChooser->ResultsStructs[RowIndex];
-		if (Result.IsValid())
+		if (const FTableCells* TableCells = Map.Find(Database))
 		{
-			ReferencedObject = Result.Get<FObjectChooserBase>().GetReferencedObject();
+			for (int32 IndexInDatabase = 0; IndexInDatabase < TableCells->Num(); ++IndexInDatabase)
+			{
+				const FTableCell& TableCell = (*TableCells)[IndexInDatabase];
+				if (TableCell.Column == Column && TableCell.RowIndex == RowIndex)
+				{
+					return IndexInDatabase;
+				}
+			}
 		}
+		return INDEX_NONE;
 	}
-	else
+
+	const FPoseSearchDatabaseAnimationAsset* FMapping::FindDatabaseAnimationAsset(int32 IndexInDatabase, const UPoseSearchDatabase* Database) const
 	{
-		const FInstancedStruct& FallbackResult = OuterChooser->FallbackResult;
-		if (FallbackResult.IsValid())
+		const FPoseSearchDatabaseAnimationAsset* DatabaseAnimationAsset = nullptr;
+		ensure(IndexInDatabase >= 0 && Database);
+		if (const FTableCells* TableCells = Map.Find(Database))
 		{
-			ReferencedObject = FallbackResult.Get<FObjectChooserBase>().GetReferencedObject();
+			if (TableCells->IsValidIndex(IndexInDatabase))
+			{
+				const FTableCell& TableCell = (*TableCells)[IndexInDatabase];
+				DatabaseAnimationAsset = &TableCell.Column->GetRowValue(TableCell.RowIndex).Data;
+#if WITH_EDITOR
+				ensure(GetReferencedObject(TableCell.Column->GetChooser(), TableCell.RowIndex) == DatabaseAnimationAsset->AnimAsset);
+#endif // WITH_EDITOR
+			}
 		}
+		return DatabaseAnimationAsset;
 	}
 
-	return ReferencedObject;
-}
+#if WITH_EDITOR
+	void FMapping::SetAnimationAssetAt(const FPoseSearchDatabaseAnimationAsset& AnimationAsset, int32 IndexInDatabase, const UPoseSearchDatabase* Database)
+	{
+		ensure(IndexInDatabase >= 0 && Database);
+		if (FTableCells* TableCells = Map.Find(Database))
+		{
+			if (TableCells->IsValidIndex(IndexInDatabase))
+			{
+				FTableCell& TableCell = (*TableCells)[IndexInDatabase];
 
-FName FPoseSearchColumn::RowValuesPropertyName()
-{
-	return "RowValues";
-}
+				DO_ENSURE_CONSISTENCY(CheckConsistency(TableCell.Column->GetChooser()));
 
-void FPoseSearchColumn::SetNumRows(int32 NumRows)
-{
-	using namespace UE::PoseSearch;
+				FPoseSearchDatabaseAnimationAsset& DatabaseAnimationAsset = TableCell.Column->EditRowValue(TableCell.RowIndex).Data;
+				DatabaseAnimationAsset = AnimationAsset;
 
-	check(InternalDatabase && NumRows >= 0);
+				// enforcing RowValue.Data.AnimAsset to be equals to ReferencedObject
+				UObject* ReferencedObject = GetReferencedObject(TableCell.Column->GetChooser(), TableCell.RowIndex);
+				if (ReferencedObject != DatabaseAnimationAsset.AnimAsset)
+				{
+					DatabaseAnimationAsset.AnimAsset = ReferencedObject;
+				}
+
+				DO_ENSURE_CONSISTENCY(CheckConsistency(TableCell.Column->GetChooser()));
+			}
+		}
+	}
+#endif // WITH_EDITOR
 
-	const int32 NumDatabaseEntries = GetDatabaseAssetIndex(NumRows);
-	if (NumDatabaseEntries == InternalDatabase->GetNumAnimationAssets())
+	int32 FMapping::GetRowIndex(int32 IndexInDatabase, const UPoseSearchDatabase* Database) const
 	{
-		// nothing to do
+		ensure(IndexInDatabase >= 0 && Database);
+		return Map[Database][IndexInDatabase].RowIndex;
 	}
-	else if (NumDatabaseEntries < InternalDatabase->GetNumAnimationAssets())
+
+	void FMapping::ApplyMappingToAllColumns()
 	{
-		InternalDatabase->Modify();
-		while (InternalDatabase->GetNumAnimationAssets() > 0 && InternalDatabase->GetNumAnimationAssets() > NumDatabaseEntries)
+		for (TPair<UPoseSearchDatabase*, FTableCells>& Pair : Map)
 		{
-			InternalDatabase->RemoveAnimationAssetAt(InternalDatabase->GetNumAnimationAssets() - 1);
+			for (FTableCell& TableCell : Pair.Value)
+			{
+				TableCell.Column->UpdateMapping(*this);
+			}
 		}
 	}
-	else
+
+	FMapping* FMapping::FindMapping(UChooserTable* Chooser)
 	{
-		InternalDatabase->Modify();
-		while (InternalDatabase->GetNumAnimationAssets() < NumDatabaseEntries)
+		if (Chooser)
 		{
-			FPoseSearchDatabaseAnimationAsset NewAsset;
-			const int32 RowIndex = GetRowIndex(InternalDatabase->GetNumAnimationAssets());
-			NewAsset.AnimAsset = GetReferencedObject(RowIndex);
-			InternalDatabase->AddAnimationAsset(NewAsset);
+			for (FInstancedStruct& ColumnData : Chooser->ColumnsStructs)
+			{
+				if (FPoseSearchColumn* PoseSearchColumn = ColumnData.GetMutablePtr<FPoseSearchColumn>())
+				{
+					return &PoseSearchColumn->GetMapping();
+				}
+			}
+
+#if WITH_EDITORONLY_DATA
+			if (!Chooser->IsCookedData())
+			{
+				for (int32 ResultIndex = 0; ResultIndex < Chooser->ResultsStructs.Num(); ++ResultIndex)
+				{
+					if (!Chooser->IsRowDisabled(ResultIndex))
+					{
+						if (const FNestedChooser* NestedChooser = Chooser->ResultsStructs[ResultIndex].GetPtr<FNestedChooser>())
+						{
+							if (FMapping* Mapping = FindMapping(NestedChooser->Chooser))
+							{
+								return Mapping;
+							}
+						}
+					}
+				}
+			}
+			else
+#endif
+			{
+				for (int32 ResultIndex = 0; ResultIndex < Chooser->CookedResults.Num(); ++ResultIndex)
+				{
+					if (const FNestedChooser* NestedChooser = Chooser->CookedResults[ResultIndex].GetPtr<FNestedChooser>())
+					{
+						if (FMapping* Mapping = FindMapping(NestedChooser->Chooser))
+						{
+							return Mapping;
+						}
+					}
+				}
+			}
+
+			if (Chooser->FallbackResult.IsValid())
+			{
+				if (const FNestedChooser* NestedChooser = Chooser->FallbackResult.GetPtr<FNestedChooser>())
+				{
+					if (FMapping* Mapping = FindMapping(NestedChooser->Chooser))
+					{
+						return Mapping;
+					}
+				}
+			}
 		}
+
+		return nullptr;
 	}
+} // namespace UE::PoseSearch
 
-	InternalDatabase->NotifySynchronizeWithExternalDependencies();
+FPoseSearchColumn::FPoseSearchColumn(const FPoseSearchColumn& Other)
+	: FChooserColumnBase(Other)
+	, RowValues(Other.RowValues)
+	, FallbackValue(Other.FallbackValue)
+	, Chooser(Other.Chooser)
+	// no need to copy Mapping since it has to be recomputed, because it holds references to Other.Chooser!
+	, Mapping()
+	, PoseReselectHistory(Other.PoseReselectHistory)
+	, MaxNumberOfResults(Other.MaxNumberOfResults)
+	, InterruptMode(Other.InterruptMode)
+	, InputValue(Other.InputValue)
+	, OutputStartTime(Other.OutputStartTime)
+	, OutputMirror(Other.OutputMirror)
+	, OutputForceBlendTo(Other.OutputForceBlendTo)
+	, OutputCost(Other.OutputCost)
+{
+	// has Other column been properly initialized?
+	ensure(Other.Chooser);
 
-	// NumRows == 0 has a special meaning of resetting the column and usually make the column inconsistent with the chooser asset, so we don't CheckConsistency
-	if (NumRows != 0)
+	// resetting the references of RowValues and FallbackValue Database(s) to null, to avoid referencing Other.Chooser nested objects,
+	// and expecting CopyRow and CopyFallback APIs to set them accordingly
+	for (FPoseSearchColumnRow& RowValue : RowValues)
 	{
-		CheckConsistency();
+		RowValue.Database = nullptr;
 	}
+	FallbackValue.Database = nullptr;
 }
 
-void FPoseSearchColumn::InsertRows(int Index, int Count)
+FPoseSearchColumn& FPoseSearchColumn::operator=(const FPoseSearchColumn& Other)
 {
-	using namespace UE::PoseSearch;
-
-	check(InternalDatabase);
-	if (Count > 0)
+	if (this != &Other)
 	{
-		InternalDatabase->Modify();
-		for (int i = 0; i < Count; i++)
-		{
-			FPoseSearchDatabaseAnimationAsset NewAsset;
-			const int32 RowIndex = Index + i;
-			const int32 DatabaseAssetIndex = GetDatabaseAssetIndex(RowIndex);
-			NewAsset.AnimAsset = GetReferencedObject(RowIndex);
-			InternalDatabase->InsertAnimationAssetAt(NewAsset, DatabaseAssetIndex);
-		}
+		this->~FPoseSearchColumn();
+		new(this) FPoseSearchColumn(Other);
 	}
-
-	InternalDatabase->NotifySynchronizeWithExternalDependencies();
-	CheckConsistency();
+	return *this;
 }
 
-void FPoseSearchColumn::DeleteRows(TArrayView<int> RowIndices)
+const FPoseSearchColumnRow& FPoseSearchColumn::GetRowValue(int32 RowIndex) const
 {
 	using namespace UE::PoseSearch;
-
-	if (RowIndices.IsEmpty())
+	if (RowIndex == ChooserColumn_SpecialIndex_Fallback)
 	{
-		return;
+		return FallbackValue;
 	}
-
-	check(InternalDatabase);
-	InternalDatabase->Modify();
-	
-	for (int RowIndex : RowIndices)
+	if (ensure(RowValues.IsValidIndex(RowIndex)))
 	{
-		const int32 DatabaseAssetIndex = GetDatabaseAssetIndex(RowIndex);
-		InternalDatabase->RemoveAnimationAssetAt(DatabaseAssetIndex);
+		return RowValues[RowIndex];
 	}
+	return FallbackValue;
+}
 
-	// resynchronize the fallback row asset
-	if (FPoseSearchDatabaseAnimationAsset* FallbackDatabaseAnimationAsset = InternalDatabase->GetMutableDatabaseAnimationAsset(GetDatabaseAssetIndex(ChooserColumn_SpecialIndex_Fallback)))
+FPoseSearchColumnRow& FPoseSearchColumn::EditRowValue(int32 RowIndex)
+{
+	using namespace UE::PoseSearch;
+	if (RowIndex == ChooserColumn_SpecialIndex_Fallback)
 	{
-		FallbackDatabaseAnimationAsset->AnimAsset = GetReferencedObject(ChooserColumn_SpecialIndex_Fallback);
+		return FallbackValue;
 	}
-
-	InternalDatabase->PostEditChange();
-
-	InternalDatabase->NotifySynchronizeWithExternalDependencies();
-	CheckConsistency();
+	if (ensure(RowValues.IsValidIndex(RowIndex)))
+	{
+		return RowValues[RowIndex];
+	}
+	return FallbackValue;
 }
 
-void FPoseSearchColumn::MoveRow(int SourceRowIndex, int TargetRowIndex)
+void FPoseSearchColumn::PostLoad()
 {
 	using namespace UE::PoseSearch;
 
-	check(InternalDatabase);
-
-	const int32 SourceDatabaseAssetIndex = GetDatabaseAssetIndex(SourceRowIndex);
-	int32 TargetDatabaseAssetIndex = GetDatabaseAssetIndex(TargetRowIndex);
-
-	if (const FPoseSearchDatabaseAnimationAsset* SourceAsset = InternalDatabase->GetDatabaseAnimationAsset(SourceDatabaseAssetIndex))
+#if WITH_EDITORONLY_DATA
+	
+	// upgrading InternalDatabase_DEPRECATED into their corresponding Databases
+	if (InternalDatabase_DEPRECATED)
 	{
-		InternalDatabase->Modify();
+		// InternalDatabase_DEPRECATED was mapped with this indices conventions:
+		// RowIndex -> DatabaseAssetIndex
+		// 0 -> ChooserColumn_SpecialIndex_Fallback
+		// RowIndex -> RowIndex + 1
 
-		FPoseSearchDatabaseAnimationAsset CopySourceAsset = *SourceAsset;
-		InternalDatabase->RemoveAnimationAssetAt(SourceDatabaseAssetIndex);
+		ensure(!Chooser);
+		Chooser = Cast<UChooserTable>(InternalDatabase_DEPRECATED->GetOuter());
+		ensure(Chooser);
 
-		if (SourceDatabaseAssetIndex < TargetDatabaseAssetIndex)
+		InternalDatabase_DEPRECATED->Modify();
+		const int32 NumRows = InternalDatabase_DEPRECATED->GetNumAnimationAssets() - 1;
+		ensure(NumRows >= 0);
+
+		RowValues.SetNum(NumRows);
+		for (int32 RowIndex = 0; RowIndex < NumRows; ++RowIndex)
 		{
-			TargetDatabaseAssetIndex--;
+			if (const FPoseSearchDatabaseAnimationAsset* DatabaseAnimationAsset = InternalDatabase_DEPRECATED->GetDatabaseAnimationAsset(RowIndex + 1))
+			{
+				if (DatabaseAnimationAsset->GetAnimationAsset())
+				{
+					RowValues[RowIndex].Data = *DatabaseAnimationAsset;
+					RowValues[RowIndex].Database = InternalDatabase_DEPRECATED;
+				}
+			}
 		}
-		InternalDatabase->InsertAnimationAssetAt(CopySourceAsset, TargetDatabaseAssetIndex);
 
-		InternalDatabase->NotifySynchronizeWithExternalDependencies();
-	}
-
-	CheckConsistency();
-}
+		if (const FPoseSearchDatabaseAnimationAsset* DatabaseAnimationAsset = InternalDatabase_DEPRECATED->GetDatabaseAnimationAsset(0))
+		{
+			if (DatabaseAnimationAsset->GetAnimationAsset())
+			{
+				FallbackValue.Data = *DatabaseAnimationAsset;
+				FallbackValue.Database = InternalDatabase_DEPRECATED;
+			}
+		}
 
-void FPoseSearchColumn::CopyRow(FChooserColumnBase& SourceColumn, int SourceRowIndex, int TargetRowIndex)
-{
-	using namespace UE::PoseSearch;
+		// removing all the AnimationAssets from InternalDatabase_DEPRECATED
+		const int32 NumAnimationAssetsToRemove = InternalDatabase_DEPRECATED->GetNumAnimationAssets();
+		for (int32 AnimationAssetIndex = NumAnimationAssetsToRemove - 1; AnimationAssetIndex >= 0; --AnimationAssetIndex)
+		{
+			InternalDatabase_DEPRECATED->RemoveAnimationAssetAt(AnimationAssetIndex);
+		}
 
-	check(InternalDatabase);
+		// changing InternalDatabase_DEPRECATED outer to RootChooser
+		UChooserTable* RootChooser = Chooser->GetRootChooser();
+		const FName DatabaseName = *("PSD_" + Chooser->GetName());
+		const FName DatabaseUniqueName = MakeUniqueObjectName(RootChooser, UPoseSearchDatabase::StaticClass(), DatabaseName);
+		InternalDatabase_DEPRECATED->Rename(*DatabaseUniqueName.ToString(), RootChooser);
+		RootChooser->AddNestedObject(InternalDatabase_DEPRECATED);
 
-	const FPoseSearchColumn& SourcePoseSearchColumn = static_cast<FPoseSearchColumn&>(SourceColumn);
+		InternalDatabase_DEPRECATED->PostEditChange();
+		
+		InternalDatabase_DEPRECATED = nullptr;
+	}
+#endif // WITH_EDITORONLY_DATA
 
-	const int32 SourceDatabaseAssetIndex = GetDatabaseAssetIndex(SourceRowIndex);
-	const int32 TargetDatabaseAssetIndex = GetDatabaseAssetIndex(TargetRowIndex);
-	if (const FPoseSearchDatabaseAnimationAsset* SourceAsset = SourcePoseSearchColumn.InternalDatabase->GetDatabaseAnimationAsset(SourceDatabaseAssetIndex))
+	if (Chooser)
 	{
-		if (FPoseSearchDatabaseAnimationAsset* TargetAsset = InternalDatabase->GetMutableDatabaseAnimationAsset(TargetDatabaseAssetIndex))
-		{
-			InternalDatabase->Modify();
-			*TargetAsset = *SourceAsset;
-
-			InternalDatabase->NotifySynchronizeWithExternalDependencies();
-		}
+		FMapping(Chooser->GetRootChooser()).ApplyMappingToAllColumns();
 	}
 
-	CheckConsistency();
-}
-
-UScriptStruct* FPoseSearchColumn::GetInputBaseType() const
-{
-	return FChooserParameterPoseHistoryBase::StaticStruct();
-}
-
-const UScriptStruct* FPoseSearchColumn::GetInputType() const
-{
-	return InputValue.IsValid() ? InputValue.GetScriptStruct() : nullptr;
+	DO_ENSURE_CONSISTENCY(CheckConsistency(Chooser));
 }
 
-void FPoseSearchColumn::SetInputType(const UScriptStruct* Type)
+UPoseSearchDatabase* FPoseSearchColumn::GetDatabase(int32 RowIndex) const
 {
-	InputValue.InitializeAs(Type);
+	UPoseSearchDatabase* Database = GetRowValue(RowIndex).Database;
+	ensure(!Database || Database->GetOuter() == Chooser->GetRootChooser());
+	return Database;
 }
 
-void FPoseSearchColumn::AutoPopulate(int32 RowIndex, UObject* OutputObject)
+#if WITH_EDITOR
+void FPoseSearchColumn::SwitchDatabase(UPoseSearchDatabase* NewDatabase, int32 RowIndex)
 {
 	using namespace UE::PoseSearch;
+	
+	DO_ENSURE_CONSISTENCY(CheckConsistency(Chooser));
 
-	const int32 DatabaseAssetIndex = GetDatabaseAssetIndex(RowIndex);
-	if (FPoseSearchDatabaseAnimationAsset* Asset = InternalDatabase->GetMutableDatabaseAnimationAsset(DatabaseAssetIndex))
+	UChooserTable* RootChooser = Chooser->GetRootChooser();
+	TObjectPtr<UPoseSearchDatabase>& Database = EditRowValue(RowIndex).Database;
+	ensure(!Database || Database->GetOuter() == RootChooser);
+	ensure(!NewDatabase || NewDatabase->GetOuter() == RootChooser);
+	if (Database != NewDatabase)
 	{
-		if (Asset->AnimAsset != OutputObject)
+		if (Database)
 		{
-			InternalDatabase->Modify();
-			Asset->AnimAsset = OutputObject;
+			Database->Modify();
 		}
-	}
-	
-	CheckConsistency();
-}
-
-bool FPoseSearchColumn::EditorTestFilter(int32 RowIndex) const
-{
-	using namespace UE::PoseSearch;
-
-	for (FActiveColumnCost& ActiveColumnCost : ActiveColumnCosts)
-	{
-		if (ActiveColumnCost.RowIndex == RowIndex)
+		if (NewDatabase)
 		{
-			return true;
+			NewDatabase->Modify();
 		}
-	}
-	return false;
-}
 
-float FPoseSearchColumn::EditorTestCost(int32 RowIndex) const
-{
-	using namespace UE::PoseSearch;
+		Database = NewDatabase;
+		FMapping(RootChooser).ApplyMappingToAllColumns();
 
-	for (FActiveColumnCost& ActiveColumnCost : ActiveColumnCosts)
-	{
-		if (ActiveColumnCost.RowIndex == RowIndex)
+		if (Database)
 		{
-			return ActiveColumnCost.RowCost;
+			Database->PostEditChange();
+		}
+		if (NewDatabase)
+		{
+			NewDatabase->PostEditChange();
 		}
 	}
-	return 0.f;
+
+	DO_ENSURE_CONSISTENCY(CheckConsistency(Chooser));
 }
 
-void FPoseSearchColumn::SetTestValue(TArrayView<const uint8> Value)
+void FPoseSearchColumn::InsertAnimationAssetAt(const FPoseSearchDatabaseAnimationAsset& AnimationAsset, int32 AnimationAssetIndex, UChooserTable* Chooser, const UPoseSearchDatabase* Database)
 {
-    FMemoryReaderView Reader(Value);
-    Reader << ActiveColumnCosts;
-}
+	// @todo: implement me! Without this method, we're not able to add a chooser row by adding a row in the related UPoseSearchDatabase.
 
-#endif // WITH_EDITOR
+	using namespace UE::PoseSearch;
+	//FMapping(Chooser->GetRootChooser()).ApplyMappingToAllColumns();
+	DO_ENSURE_CONSISTENCY(CheckConsistency(Chooser));
+}
 
-const FChooserParameterBase* FPoseSearchColumn::GetInputValue() const
+void FPoseSearchColumn::RemoveAnimationAssetAt(int32 AnimationAssetIndex, UChooserTable* Chooser, const UPoseSearchDatabase* Database)
 {
-	if (const FChooserParameterBase* ChooserParameterBase = InputValue.GetPtr<FChooserParameterBase>())
-	{
-		return ChooserParameterBase;
-	}
-	return &NullInputValue;
+	// @todo: implement me! Without this method, we're not able to delete a chooser row by deleting a row in the related UPoseSearchDatabase.
+
+	using namespace UE::PoseSearch;
+	//FMapping(Chooser->GetRootChooser()).ApplyMappingToAllColumns();
+	DO_ENSURE_CONSISTENCY(CheckConsistency(Chooser));
 }
 
-FChooserParameterBase* FPoseSearchColumn::GetInputValue()
+void FPoseSearchColumn::SetAnimationAssetAt(const FPoseSearchDatabaseAnimationAsset& AnimationAsset, int32 AnimationAssetIndex, UChooserTable* Chooser, const UPoseSearchDatabase* Database)
 {
-	if (FChooserParameterBase* ChooserParameterBase = InputValue.GetMutablePtr<FChooserParameterBase>())
+	using namespace UE::PoseSearch;
+	DO_ENSURE_CONSISTENCY(CheckConsistency(Chooser));
+	if (FMapping* Mapping = FMapping::FindMapping(Chooser))
 	{
-		return ChooserParameterBase;
+		ENSURE_MAPPING_VALIDATION(Mapping->Equals(FMapping(Chooser->GetRootChooser())));
+		Mapping->SetAnimationAssetAt(AnimationAsset, AnimationAssetIndex, Database);
 	}
-	return &NullInputValue;
+	DO_ENSURE_CONSISTENCY(CheckConsistency(Chooser));
 }
+#endif // WITH_EDITOR
 
-const UObject* FPoseSearchColumn::GetDatabaseAsset(int32 RowIndex) const
+const FPoseSearchDatabaseAnimationAsset* FPoseSearchColumn::GetDatabaseAnimationAsset(int32 AnimationAssetIndex, UChooserTable* Chooser, const UPoseSearchDatabase* Database)
 {
 	using namespace UE::PoseSearch;
-	const int32 DatabaseAssetIndex = GetDatabaseAssetIndex(RowIndex);
-
-	if (const FPoseSearchDatabaseAnimationAsset* Asset = InternalDatabase->GetDatabaseAnimationAsset(DatabaseAssetIndex))
+	DO_ENSURE_CONSISTENCY(CheckConsistency(Chooser));
+	if (FMapping* Mapping = FMapping::FindMapping(Chooser))
 	{
-		return Asset->GetAnimationAsset();
+		ENSURE_MAPPING_VALIDATION(Mapping->Equals(FMapping(Chooser->GetRootChooser())));
+		return Mapping->FindDatabaseAnimationAsset(AnimationAssetIndex, Database);
 	}
 	return nullptr;
 }
 
-const UPoseSearchSchema* FPoseSearchColumn::GetDatabaseSchema() const
+int32 FPoseSearchColumn::GetNumAnimationAssets(UChooserTable* Chooser, const UPoseSearchDatabase* Database)
 {
-	check(InternalDatabase);
-	return InternalDatabase->Schema;
-}
-	
-void FPoseSearchColumn::SetDatabaseSchema(const UPoseSearchSchema* Schema)
-{
-	check(InternalDatabase);
-	InternalDatabase->Modify();
-	InternalDatabase->Schema = Schema;
+	using namespace UE::PoseSearch;
+	DO_ENSURE_CONSISTENCY(CheckConsistency(Chooser));
+	if (FMapping* Mapping = FMapping::FindMapping(Chooser))
+	{
+		ENSURE_MAPPING_VALIDATION(Mapping->Equals(FMapping(Chooser->GetRootChooser())));
+		return Mapping->GetNumIndexesInDatabase(Database);
+	}
+	return 0;
 }
 
-bool FPoseSearchColumn::HasFilters() const
+void FPoseSearchColumn::UpdateMapping(UE::PoseSearch::FMapping& NewMapping)
 {
-	return true;
+	using namespace UE::PoseSearch;
+	DO_ENSURE_CONSISTENCY(CheckConsistency(Chooser));
+	Mapping = NewMapping;
+	DO_ENSURE_CONSISTENCY(CheckConsistency(Chooser));
 }
 
 void FPoseSearchColumn::Filter(FChooserEvaluationContext& Context, const FChooserIndexArray& IndexListIn, FChooserIndexArray& IndexListOut, TArrayView<uint8> ScratchArea) const
 {
 	using namespace UE::PoseSearch;
 
-	CheckConsistency();
+	DO_ENSURE_CONSISTENCY(CheckConsistency(Chooser));
+
+	ENSURE_MAPPING_VALIDATION(Mapping.Equals(FMapping(Chooser->GetRootChooser())));
 
-	check(ScratchArea.Num() == GetScratchAreaSize());
+	ensure(ScratchArea.Num() == GetScratchAreaSize());
 	FPoseSearchColumnScratchArea* PoseSearchColumnScratchArea = reinterpret_cast<FPoseSearchColumnScratchArea*>(ScratchArea.begin());
 #if DO_CHECK
-	check(PoseSearchColumnScratchArea->DebugOwner == this);
-	check(PoseSearchColumnScratchArea->SearchResults.IsEmpty());
+	ensure(PoseSearchColumnScratchArea->DebugOwner == this);
+	ensure(PoseSearchColumnScratchArea->SearchResults.IsEmpty());
 #endif // DO_CHECK
 
 	const UE::PoseSearch::IPoseHistory* PoseHistory = nullptr;
@@ -465,29 +686,54 @@ void FPoseSearchColumn::Filter(FChooserEvaluationContext& Context, const FChoose
 	{
 		PoseSearchColumnScratchArea->PoseHistory = PoseHistory;
 
+		ensure(Chooser);
+		
 		FStackDatabaseToAssetIndexes AssetIndexesToConsiderPerDatabase;
-		FStackAssetIndexes& AssetIndexesToConsider = AssetIndexesToConsiderPerDatabase.Add(InternalDatabase.Get());
-
 		for (const FChooserIndexArray::FIndexData& IndexData : IndexListIn)
 		{
 			const int32 RowIndex = IndexData.Index;
-			const int32 DatabaseAssetIndex = GetDatabaseAssetIndex(RowIndex);
 
-			// pruning out eventual nested chooser or invalid assets not supported by the database
-			if (InternalDatabase->GetAnimationAsset(DatabaseAssetIndex))
+			if (UPoseSearchDatabase* Database = GetDatabase(RowIndex))
 			{
-				check(DatabaseAssetIndex >= 0 && DatabaseAssetIndex < InternalDatabase->GetNumAnimationAssets());
-				AssetIndexesToConsider.Add(DatabaseAssetIndex);
+				const int32 IndexInDatabase = Mapping.FindIndexInDatabase(RowIndex, this, Database);
+				if (IndexInDatabase != INDEX_NONE)
+				{
+					AssetIndexesToConsiderPerDatabase.FindOrAdd(Database).Add(IndexInDatabase);
+				}
 			}
 		}
 
-		if (!AssetIndexesToConsider.IsEmpty())
+		if (!AssetIndexesToConsiderPerDatabase.IsEmpty())
 		{
+			// collecting and sorting DatabasesToSearch to have a deterministic search
+			bool bRoleSet = false;
+			FRole Role = DefaultRole;
+
+			TArray<const UObject*> DatabasesToSearch;
+			for (TPair<const UPoseSearchDatabase*, FStackAssetIndexes>& Pair : AssetIndexesToConsiderPerDatabase)
+			{
+				if (Pair.Key->Schema)
+				{
+					if (!bRoleSet)
+					{
+						Role = Pair.Key->Schema->GetDefaultRole();
+						bRoleSet = true;
+					}
+#if !NO_LOGGING
+					else if (Role != Pair.Key->Schema->GetDefaultRole())
+					{
+						// @todo: improve logging
+						UE_LOG(LogPoseSearch, Error, TEXT("FPoseSearchColumn::Filter inconsistent Roles"));
+					}
+#endif // !NO_LOGGING
+					DatabasesToSearch.Add(Pair.Key);
+				}
+			}
+			DatabasesToSearch.Sort();
+
 			const FFloatInterval PoseJumpThresholdTime(0.f, 0.f);
-			const UObject* DatabasesToSearch[] = { InternalDatabase.Get() };
 			FSearchContext SearchContext(0.f, PoseJumpThresholdTime, FPoseSearchEvent());
 
-			const FRole Role = InternalDatabase->Schema ? InternalDatabase->Schema->GetDefaultRole() : DefaultRole;
 			SearchContext.AddRole(Role, &Context, PoseHistory);
 			SearchContext.SetAssetIndexesToConsiderPerDatabase(&AssetIndexesToConsiderPerDatabase);
 
@@ -512,11 +758,12 @@ void FPoseSearchColumn::Filter(FChooserEvaluationContext& Context, const FChoose
 				FSearchResults_Single SearchResults;
 				UPoseSearchLibrary::MotionMatch(SearchContext, DatabasesToSearch, ContinuingProperties, SearchResults);
 #if DO_CHECK
-				SearchResults.IterateOverSearchResults([&AssetIndexesToConsider](FSearchResult& SearchResult)
+				SearchResults.IterateOverSearchResults([&AssetIndexesToConsiderPerDatabase](FSearchResult& SearchResult)
 					{
 						if (const FSearchIndexAsset* SearchIndexAsset = SearchResult.GetSearchIndexAsset())
 						{
-							ensure(AssetIndexesToConsider.Contains(SearchIndexAsset->GetSourceAssetIdx()));
+							const FStackAssetIndexes* AssetIndexesToConsider = AssetIndexesToConsiderPerDatabase.Find(SearchResult.Database.Get());
+							ensure(AssetIndexesToConsider && AssetIndexesToConsider->Contains(SearchIndexAsset->GetSourceAssetIdx()));
 						}
 						return false;
 					});
@@ -526,14 +773,8 @@ void FPoseSearchColumn::Filter(FChooserEvaluationContext& Context, const FChoose
 
 				if (const FSearchIndexAsset* SearchIndexAsset = BestSearchResult.GetSearchIndexAsset())
 				{
-					check(BestSearchResult.Database == InternalDatabase);
-
-					// @todo: handle duplicate entries, by providing to the UPoseSearchLibrary::MotionMatch a SourceAssetIdxsToConsider rather than AssetsToConsider
-					const int32 SearchResultSourceAssetIdx = SearchIndexAsset->GetSourceAssetIdx();
-					check(SearchResultSourceAssetIdx >= 0);
-
-					const int32 RowIndex = GetRowIndex(SearchResultSourceAssetIdx);
-					check(RowIndex >= 0);
+					const int32 RowIndex = Mapping.GetRowIndex(SearchIndexAsset->GetSourceAssetIdx(), BestSearchResult.Database.Get());
+					ensure(RowIndex >= 0);
 					PoseSearchColumnScratchArea->SearchResults.Add(BestSearchResult);
 					IndexListOut.Push({ static_cast<uint32>(RowIndex), BestSearchResult.PoseCost });
 #if WITH_EDITOR
@@ -546,11 +787,12 @@ void FPoseSearchColumn::Filter(FChooserEvaluationContext& Context, const FChoose
 				FSearchResults_AssetBests SearchResults;
 				UPoseSearchLibrary::MotionMatch(SearchContext, DatabasesToSearch, ContinuingProperties, SearchResults);
 #if DO_CHECK
-				SearchResults.IterateOverSearchResults([&AssetIndexesToConsider](FSearchResult& SearchResult)
+				SearchResults.IterateOverSearchResults([&AssetIndexesToConsiderPerDatabase](FSearchResult& SearchResult)
 					{
 						if (const FSearchIndexAsset* SearchIndexAsset = SearchResult.GetSearchIndexAsset())
 						{
-							ensure(AssetIndexesToConsider.Contains(SearchIndexAsset->GetSourceAssetIdx()));
+							const FStackAssetIndexes* AssetIndexesToConsider = AssetIndexesToConsiderPerDatabase.Find(SearchResult.Database.Get());
+							ensure(AssetIndexesToConsider && AssetIndexesToConsider->Contains(SearchIndexAsset->GetSourceAssetIdx()));
 						}
 						return false;
 					});
@@ -566,14 +808,17 @@ void FPoseSearchColumn::Filter(FChooserEvaluationContext& Context, const FChoose
 				for (const FChooserIndexArray::FIndexData& IndexData : IndexListIn)
 				{
 					const uint32 RowIndex = IndexData.Index;
-					const int32 DatabaseAssetIndex = GetDatabaseAssetIndex(RowIndex);
-					if (const FSearchResult* FoundSearchResult = SearchResults.FindSearchResultFor(InternalDatabase.Get(), DatabaseAssetIndex))
+					if (UPoseSearchDatabase* Database = GetDatabase(RowIndex))
 					{
-						PoseSearchColumnScratchArea->SearchResults.Add(*FoundSearchResult);
-						IndexListOut.Push({ RowIndex, FoundSearchResult->PoseCost });
+						const int32 DatabaseAssetIndex = Mapping.FindIndexInDatabase(RowIndex, this, Database);
+						if (const FSearchResult* FoundSearchResult = SearchResults.FindSearchResultFor(Database, DatabaseAssetIndex))
+						{
+							PoseSearchColumnScratchArea->SearchResults.Add(*FoundSearchResult);
+							IndexListOut.Push({ RowIndex, FoundSearchResult->PoseCost });
 #if WITH_EDITOR
-						TracedActiveColumnCosts.Add({ static_cast<int32>(RowIndex), FoundSearchResult->PoseCost });
+							TracedActiveColumnCosts.Add({ static_cast<int32>(RowIndex), FoundSearchResult->PoseCost });
 #endif // WITH_EDITOR
+						}
 					}
 				}
 			}
@@ -581,7 +826,7 @@ void FPoseSearchColumn::Filter(FChooserEvaluationContext& Context, const FChoose
 	}
 
 #if WITH_EDITOR
-	TRACE_CHOOSER_VALUE(Context, ToCStr(GetInputValue()->GetDebugName()), TracedActiveColumnCosts);
+	TRACE_CHOOSER_VALUE(Context, ToCStr(InputValue.GetPtr<FChooserParameterBase>()->GetDebugName()), TracedActiveColumnCosts);
 	if (Context.DebuggingInfo.bCurrentDebugTarget)
 	{
 		ActiveColumnCosts = TracedActiveColumnCosts;
@@ -589,90 +834,88 @@ void FPoseSearchColumn::Filter(FChooserEvaluationContext& Context, const FChoose
 #endif // WITH_EDITOR
 }
 
+
 void FPoseSearchColumn::SetOutputs(FChooserEvaluationContext& Context, int RowIndex, TArrayView<uint8> ScratchArea) const
 {
 	using namespace UE::PoseSearch;
 
+	DO_ENSURE_CONSISTENCY(CheckConsistency(Chooser));
+
 	// set an arbitrary "high" cost value so that cost threshold implementations can still wait for the pose search to find a good match
 	float CostValue = 100.f;
 	float StartTimeValue = 0.f;
 	bool bMirrorValue = false;
 	bool bForceBlendToValue = false;
 
-	check(ScratchArea.Num() == GetScratchAreaSize());
+	ensure(ScratchArea.Num() == GetScratchAreaSize());
 	FPoseSearchColumnScratchArea* PoseSearchColumnScratchArea = reinterpret_cast<FPoseSearchColumnScratchArea*>(ScratchArea.begin());
 #if DO_CHECK
-	check(PoseSearchColumnScratchArea->DebugOwner == this);
+	ensure(PoseSearchColumnScratchArea->DebugOwner == this);
 #endif // DO_CHECK
 
 	FSearchResult BestSearchResult;
-
-	if (RowIndex == ChooserColumn_SpecialIndex_Fallback && PoseSearchColumnScratchArea->PoseHistory)
+	if (RowIndex == ChooserColumn_SpecialIndex_Fallback)
 	{
-		// making sure that if we search again with the fallback column we didn't find ANY previous result!
-		check(PoseSearchColumnScratchArea->SearchResults.IsEmpty());
-
-		// do the search again only with the fallback row..
-		FStackDatabaseToAssetIndexes AssetIndexesToConsiderPerDatabase;
-		FStackAssetIndexes& AssetIndexesToConsider = AssetIndexesToConsiderPerDatabase.Add(InternalDatabase.Get());
-
-		const int32 DatabaseAssetIndex = GetDatabaseAssetIndex(RowIndex);
-		
-		// pruning out eventual nested chooser or invalid assets not supported by the database
-		if (InternalDatabase->GetAnimationAsset(DatabaseAssetIndex))
+		if (PoseSearchColumnScratchArea->PoseHistory)
 		{
-			check(DatabaseAssetIndex >= 0 && DatabaseAssetIndex < InternalDatabase->GetNumAnimationAssets());
-			AssetIndexesToConsider.Add(DatabaseAssetIndex);
+			// making sure that if we search again with the fallback column we didn't find ANY previous result!
+			ensure(PoseSearchColumnScratchArea->SearchResults.IsEmpty());
 
-			const FFloatInterval PoseJumpThresholdTime(0.f, 0.f);
-			const UObject* DatabasesToSearch[] = { InternalDatabase.Get() };
-			FSearchContext SearchContext(0.f, PoseJumpThresholdTime, FPoseSearchEvent());
+			// do the search again only with the fallback row..
+			FStackDatabaseToAssetIndexes AssetIndexesToConsiderPerDatabase;
+			if (UPoseSearchDatabase* Database = GetDatabase(RowIndex))
+			{
+				const int32 IndexInDatabase = Mapping.FindIndexInDatabase(RowIndex, this, Database);
+				// also pruning out eventual nested chooser or invalid assets not supported by the database
+				if (IndexInDatabase != INDEX_NONE && Database->GetAnimationAsset(IndexInDatabase))
+				{
+					AssetIndexesToConsiderPerDatabase.Add(Database).Add(IndexInDatabase);
 
-			const FRole Role = InternalDatabase->Schema ? InternalDatabase->Schema->GetDefaultRole() : DefaultRole;
-			SearchContext.AddRole(Role, &Context, PoseSearchColumnScratchArea->PoseHistory);
-			SearchContext.SetAssetIndexesToConsiderPerDatabase(&AssetIndexesToConsiderPerDatabase);
+					const FFloatInterval PoseJumpThresholdTime(0.f, 0.f);
+					const UObject* DatabasesToSearch[] = { Database };
+					FSearchContext SearchContext(0.f, PoseJumpThresholdTime, FPoseSearchEvent());
 
-			FPoseSearchContinuingProperties ContinuingProperties;
-			if (PoseSearchColumnScratchArea->ChooserPlayerSettings)
-			{
-				ContinuingProperties.PlayingAsset = PoseSearchColumnScratchArea->ChooserPlayerSettings->PlayingAsset;
-				ContinuingProperties.PlayingAssetAccumulatedTime = PoseSearchColumnScratchArea->ChooserPlayerSettings->PlayingAssetAccumulatedTime;
-				ContinuingProperties.bIsPlayingAssetMirrored = PoseSearchColumnScratchArea->ChooserPlayerSettings->bIsPlayingAssetMirrored;
-				ContinuingProperties.PlayingAssetBlendParameters = PoseSearchColumnScratchArea->ChooserPlayerSettings->PlayingAssetBlendParameters;
+					const FRole Role = Database->Schema ? Database->Schema->GetDefaultRole() : DefaultRole;
+					SearchContext.AddRole(Role, &Context, PoseSearchColumnScratchArea->PoseHistory);
+					SearchContext.SetAssetIndexesToConsiderPerDatabase(&AssetIndexesToConsiderPerDatabase);
 
-				uint8 InterruptModeValue = 0;
-				if (InterruptMode.IsValid() && InterruptMode.Get<FChooserParameterEnumBase>().GetValue(Context, InterruptModeValue))
-				{
-					ContinuingProperties.InterruptMode = (EPoseSearchInterruptMode)InterruptModeValue;
-				}
-			}
+					FPoseSearchContinuingProperties ContinuingProperties;
+					if (PoseSearchColumnScratchArea->ChooserPlayerSettings)
+					{
+						ContinuingProperties.PlayingAsset = PoseSearchColumnScratchArea->ChooserPlayerSettings->PlayingAsset;
+						ContinuingProperties.PlayingAssetAccumulatedTime = PoseSearchColumnScratchArea->ChooserPlayerSettings->PlayingAssetAccumulatedTime;
+						ContinuingProperties.bIsPlayingAssetMirrored = PoseSearchColumnScratchArea->ChooserPlayerSettings->bIsPlayingAssetMirrored;
+						ContinuingProperties.PlayingAssetBlendParameters = PoseSearchColumnScratchArea->ChooserPlayerSettings->PlayingAssetBlendParameters;
 
-			FSearchResults_Single SearchResults;
-			UPoseSearchLibrary::MotionMatch(SearchContext, DatabasesToSearch, ContinuingProperties, SearchResults);
+						uint8 InterruptModeValue = 0;
+						if (InterruptMode.IsValid() && InterruptMode.Get<FChooserParameterEnumBase>().GetValue(Context, InterruptModeValue))
+						{
+							ContinuingProperties.InterruptMode = (EPoseSearchInterruptMode)InterruptModeValue;
+						}
+					}
 
+					FSearchResults_Single SearchResults;
+					UPoseSearchLibrary::MotionMatch(SearchContext, DatabasesToSearch, ContinuingProperties, SearchResults);
 #if DO_CHECK
-			SearchResults.IterateOverSearchResults([&AssetIndexesToConsider](FSearchResult& SearchResult)
-				{
-					if (const FSearchIndexAsset* SearchIndexAsset = SearchResult.GetSearchIndexAsset())
+					SearchResults.IterateOverSearchResults([&AssetIndexesToConsiderPerDatabase](FSearchResult& SearchResult)
+						{
+							if (const FSearchIndexAsset* SearchIndexAsset = SearchResult.GetSearchIndexAsset())
+							{
+								const FStackAssetIndexes* AssetIndexesToConsider = AssetIndexesToConsiderPerDatabase.Find(SearchResult.Database.Get());
+								ensure(AssetIndexesToConsider && AssetIndexesToConsider->Contains(SearchIndexAsset->GetSourceAssetIdx()));
+							}
+							return false;
+						});
+#endif // DO_CHECK
+					BestSearchResult = SearchResults.GetBestResult();
+					if (const FSearchIndexAsset* SearchIndexAsset = BestSearchResult.GetSearchIndexAsset())
 					{
-						ensure(AssetIndexesToConsider.Contains(SearchIndexAsset->GetSourceAssetIdx()));
+						StartTimeValue = BestSearchResult.GetAssetTime();
+						CostValue = BestSearchResult.PoseCost;
+						bMirrorValue = SearchIndexAsset->IsMirrored();
+						bForceBlendToValue = !BestSearchResult.bIsContinuingPoseSearch;
 					}
-					return false;
-				});
-#endif // DO_CHECK
-
-			BestSearchResult = SearchResults.GetBestResult();
-			if (const FSearchIndexAsset* SearchIndexAsset = BestSearchResult.GetSearchIndexAsset())
-			{
-				check(BestSearchResult.Database == InternalDatabase);
-
-				const int32 SearchResultSourceAssetIdx = SearchIndexAsset->GetSourceAssetIdx();
-				ensure(SearchResultSourceAssetIdx == 0 && RowIndex == GetRowIndex(SearchResultSourceAssetIdx));
-
-				StartTimeValue = BestSearchResult.GetAssetTime();
-				CostValue = BestSearchResult.PoseCost;
-				bMirrorValue = SearchIndexAsset->IsMirrored();
-				bForceBlendToValue = !BestSearchResult.bIsContinuingPoseSearch;
+				}
 			}
 		}
 	}
@@ -682,8 +925,9 @@ void FPoseSearchColumn::SetOutputs(FChooserEvaluationContext& Context, int RowIn
 		{
 			if (const FSearchIndexAsset* SearchIndexAsset = SearchResult.GetSearchIndexAsset())
 			{
-				const int32 SearchResultSourceAssetIdx = SearchIndexAsset->GetSourceAssetIdx();
-				const int32 SearchResultRowIndex = GetRowIndex(SearchResultSourceAssetIdx);
+				const int32 SearchResultRowIndex = Mapping.GetRowIndex(SearchIndexAsset->GetSourceAssetIdx(), SearchResult.Database.Get());
+				ensure(SearchResultRowIndex >= 0);
+
 				if (SearchResultRowIndex == RowIndex)
 				{
 					StartTimeValue = SearchResult.GetAssetTime();
@@ -743,9 +987,9 @@ int32 FPoseSearchColumn::GetScratchAreaSize() const
 void FPoseSearchColumn::InitializeScratchArea(TArrayView<uint8> ScratchArea) const
 {
 	using namespace UE::PoseSearch;
-	check(ScratchArea.Num() == GetScratchAreaSize());
-	FPoseSearchColumnScratchArea* PoseSearchColumnScratchArea = new(ScratchArea.begin()) FPoseSearchColumnScratchArea();
 
+	ensure(ScratchArea.Num() == GetScratchAreaSize());
+	FPoseSearchColumnScratchArea* PoseSearchColumnScratchArea = new(ScratchArea.begin()) FPoseSearchColumnScratchArea();
 #if DO_CHECK
 	PoseSearchColumnScratchArea->DebugOwner = this;
 #endif // DO_CHECK
@@ -754,10 +998,11 @@ void FPoseSearchColumn::InitializeScratchArea(TArrayView<uint8> ScratchArea) con
 void FPoseSearchColumn::DeinitializeScratchArea(TArrayView<uint8> ScratchArea) const
 {
 	using namespace UE::PoseSearch;
-	check(ScratchArea.Num() == GetScratchAreaSize());
+
+	ensure(ScratchArea.Num() == GetScratchAreaSize());
 	FPoseSearchColumnScratchArea* PoseSearchColumnScratchArea = reinterpret_cast<FPoseSearchColumnScratchArea*>(ScratchArea.begin());
 #if DO_CHECK
-	check(PoseSearchColumnScratchArea->DebugOwner == this);
+	ensure(PoseSearchColumnScratchArea->DebugOwner == this);
 #endif // DO_CHECK
 	PoseSearchColumnScratchArea->~FPoseSearchColumnScratchArea();
 }
@@ -765,53 +1010,426 @@ void FPoseSearchColumn::DeinitializeScratchArea(TArrayView<uint8> ScratchArea) c
 #if WITH_EDITOR
 void FPoseSearchColumn::Initialize(UChooserTable* OuterChooser)
 {
+	using namespace UE::PoseSearch;
+	
+	ensure(OuterChooser);
 	FChooserColumnBase::Initialize(OuterChooser);
-	FName BaseName = FName(FString::Format(TEXT("{0}_PSD"), { OuterChooser->GetName() }));
-	FName DatabaseName = MakeUniqueObjectName(OuterChooser, UPoseSearchDatabase::StaticClass(), BaseName);
-	InternalDatabase = NewObject<UPoseSearchDatabase>(OuterChooser, DatabaseName, RF_Transactional);
+
+	RowValues.Reset();
+	FallbackValue = FPoseSearchColumnRow();
+
+	Chooser = OuterChooser;
+	FMapping(Chooser->GetRootChooser()).ApplyMappingToAllColumns();
+}
+
+bool FPoseSearchColumn::EditorTestFilter(int32 RowIndex) const
+{
+	using namespace UE::PoseSearch;
+
+	for (FActiveColumnCost& ActiveColumnCost : ActiveColumnCosts)
+	{
+		if (ActiveColumnCost.RowIndex == RowIndex)
+		{
+			return true;
+		}
+	}
+	return false;
+}
+
+float FPoseSearchColumn::EditorTestCost(int32 RowIndex) const
+{
+	using namespace UE::PoseSearch;
+
+	for (FActiveColumnCost& ActiveColumnCost : ActiveColumnCosts)
+	{
+		if (ActiveColumnCost.RowIndex == RowIndex)
+		{
+			return ActiveColumnCost.RowCost;
+		}
+	}
+	return 0.f;
+}
+
+void FPoseSearchColumn::SetTestValue(TArrayView<const uint8> Value)
+{
+    FMemoryReaderView Reader(Value);
+    Reader << ActiveColumnCosts;
 }
 
 void FPoseSearchColumn::AddToDetails(FInstancedPropertyBag& PropertyBag, int32 ColumnIndex, int32 RowIndex)
 {
 	using namespace UE::PoseSearch;
 
-	check(InternalDatabase);
+	DO_ENSURE_CONSISTENCY(CheckConsistency(Chooser));
 
+	ensure(Chooser);
 	FText DisplayName = NSLOCTEXT("PoseSearchColumn","Pose Search", "Pose Search");
 	FName PropertyName("RowData", ColumnIndex);
-	FPropertyBagPropertyDesc PropertyDesc(PropertyName,  EPropertyBagPropertyType::Struct, FPoseSearchDatabaseAnimationAsset::StaticStruct());
+	FPropertyBagPropertyDesc PropertyDesc(PropertyName,  EPropertyBagPropertyType::Struct, FPoseSearchColumnRowReflection::StaticStruct());
 	PropertyDesc.MetaData.Add(FPropertyBagPropertyDescMetaData("DisplayName", DisplayName.ToString()));
 	PropertyBag.AddProperties({PropertyDesc});
-		
-	const int32 DatabaseAssetIndex = GetDatabaseAssetIndex(RowIndex);
-	if (const FPoseSearchDatabaseAnimationAsset* Asset = InternalDatabase->GetDatabaseAnimationAsset(DatabaseAssetIndex))
-	{
-		PropertyBag.SetValueStruct(PropertyName, *Asset);
-	}
+
+	const FPoseSearchColumnRow& RowValue = GetRowValue(RowIndex);
+	ensure(GetReferencedObject(GetChooser(), RowIndex) == RowValue.Data.AnimAsset);
+
+	FPoseSearchColumnRowReflection RowReflection;
+	RowReflection.bDisableReselection = RowValue.Data.bDisableReselection;
+	RowReflection.MirrorOption = RowValue.Data.MirrorOption;
+	RowReflection.bUseSingleSample = RowValue.Data.bUseSingleSample;
+	RowReflection.bUseGridForSampling = RowValue.Data.bUseGridForSampling;
+	RowReflection.NumberOfHorizontalSamples = RowValue.Data.NumberOfHorizontalSamples;
+	RowReflection.NumberOfVerticalSamples = RowValue.Data.NumberOfVerticalSamples;
+	RowReflection.BlendParamX = RowValue.Data.BlendParamX;
+	RowReflection.BlendParamY = RowValue.Data.BlendParamY;
+	RowReflection.Database = RowValue.Database;
+
+	ensure(!RowReflection.Database || RowReflection.Database->GetOuter() == Chooser->GetRootChooser());
+
+	PropertyBag.SetValueStruct(PropertyName, RowReflection);
+
+	DO_ENSURE_CONSISTENCY(CheckConsistency(Chooser));
 }
 
 void FPoseSearchColumn::SetFromDetails(FInstancedPropertyBag& PropertyBag, int32 ColumnIndex, int32 RowIndex)
 {
 	using namespace UE::PoseSearch;
 
-	check(InternalDatabase);
+	DO_ENSURE_CONSISTENCY(CheckConsistency(Chooser));
 
+	ensure(Chooser);
 	FName PropertyName("RowData", ColumnIndex);
-		
-   	TValueOrError<FStructView, EPropertyBagResult> Result = PropertyBag.GetValueStruct(PropertyName, FPoseSearchDatabaseAnimationAsset::StaticStruct());
+   	TValueOrError<FStructView, EPropertyBagResult> Result = PropertyBag.GetValueStruct(PropertyName, FPoseSearchColumnRowReflection::StaticStruct());
 	if (FStructView* StructView = Result.TryGetValue())
 	{
-		const int32 DatabaseAssetIndex = GetDatabaseAssetIndex(RowIndex);
-		if (FPoseSearchDatabaseAnimationAsset* Asset = InternalDatabase->GetMutableDatabaseAnimationAsset(DatabaseAssetIndex))
+		const FPoseSearchColumnRowReflection& RowReflection = StructView->Get<FPoseSearchColumnRowReflection>();
+		FPoseSearchColumnRow& RowValue = EditRowValue(RowIndex);
+
+		ensure(!RowReflection.Database || RowReflection.Database->GetOuter() == Chooser->GetRootChooser());
+
+		UObject* AnimAsset = GetReferencedObject(GetChooser(), RowIndex);
+		if (RowValue.Data.AnimAsset != AnimAsset ||
+			RowValue.Data.bDisableReselection != RowReflection.bDisableReselection ||
+			RowValue.Data.MirrorOption != RowReflection.MirrorOption ||
+			RowValue.Data.bUseSingleSample != RowReflection.bUseSingleSample ||
+			RowValue.Data.bUseGridForSampling != RowReflection.bUseGridForSampling ||
+			RowValue.Data.NumberOfHorizontalSamples != RowReflection.NumberOfHorizontalSamples ||
+			RowValue.Data.NumberOfVerticalSamples != RowReflection.NumberOfVerticalSamples ||
+			RowValue.Data.BlendParamX != RowReflection.BlendParamX ||
+			RowValue.Data.BlendParamY != RowReflection.BlendParamY ||
+			RowValue.Database != RowReflection.Database)
 		{
-			InternalDatabase->Modify();
-			*Asset = StructView->Get<FPoseSearchDatabaseAnimationAsset>();
-			InternalDatabase->NotifySynchronizeWithExternalDependencies();
+			Chooser->Modify();
+
+			RowValue.Data.AnimAsset = GetReferencedObject(GetChooser(), RowIndex);
+			RowValue.Data.bDisableReselection = RowReflection.bDisableReselection;
+			RowValue.Data.MirrorOption = RowReflection.MirrorOption;
+			RowValue.Data.bUseSingleSample = RowReflection.bUseSingleSample;
+			RowValue.Data.bUseGridForSampling = RowReflection.bUseGridForSampling;
+			RowValue.Data.NumberOfHorizontalSamples = RowReflection.NumberOfHorizontalSamples;
+			RowValue.Data.NumberOfVerticalSamples = RowReflection.NumberOfVerticalSamples;
+			RowValue.Data.BlendParamX = RowReflection.BlendParamX;
+			RowValue.Data.BlendParamY = RowReflection.BlendParamY;
+			RowValue.Database = RowReflection.Database;
+
+			FMapping(Chooser->GetRootChooser()).ApplyMappingToAllColumns();
+			Chooser->PostEditChange();
 		}
 	}
 
-	CheckConsistency();
+	DO_ENSURE_CONSISTENCY(CheckConsistency(Chooser));
+}
+
+void FPoseSearchColumn::ReplaceReferences(UObject* ReferenceToReplace, UObject* ReplaceWith)
+{
+	using namespace UE::PoseSearch;
+	ensure(Chooser);
+
+	DO_ENSURE_CONSISTENCY(CheckConsistency(Chooser));
+
+	TArray<UChooserTable*, TInlineAllocator<2>> ModifiedChoosers;
+
+	if (ReferenceToReplace && ReferenceToReplace != ReplaceWith)
+	{
+		if (UPoseSearchDatabase* DatabaseToReplace = Cast<UPoseSearchDatabase>(ReferenceToReplace))
+		{
+			UPoseSearchDatabase* DatabaseReplaceWith = Cast<UPoseSearchDatabase>(ReplaceWith);
+			if (ReplaceWith && !DatabaseReplaceWith)
+			{
+				UE_LOG(LogPoseSearch, Error, TEXT("FPoseSearchColumn::ReplaceReferences, mismatch between ReferenceToReplace and ReplaceWith"));
+			}
+			else
+			{
+				ensure(!DatabaseReplaceWith || DatabaseReplaceWith->GetOuter() == Chooser->GetRootChooser());
+
+				for (FPoseSearchColumnRow& RowValue : RowValues)
+				{
+					if (RowValue.Database == DatabaseToReplace)
+					{
+						if (!ModifiedChoosers.Contains(Chooser))
+						{
+							ModifiedChoosers.Add(Chooser);
+							Chooser->Modify();
+						}
+						RowValue.Database = DatabaseReplaceWith;
+					}
+				}
+
+				if (FallbackValue.Database == DatabaseToReplace)
+				{
+					if (!ModifiedChoosers.Contains(Chooser))
+					{
+						ModifiedChoosers.Add(Chooser);
+						Chooser->Modify();
+					}
+					FallbackValue.Database = DatabaseReplaceWith;
+				}
+			}
+		}
+		else if (UChooserTable* ChooserToReplace = Cast<UChooserTable>(ReferenceToReplace))
+		{
+			UChooserTable* ChooserToReplaceWith = Cast<UChooserTable>(ReplaceWith);
+			if (ReplaceWith && !ChooserToReplaceWith)
+			{
+				UE_LOG(LogPoseSearch, Error, TEXT("FPoseSearchColumn::ReplaceReferences, mismatch between ReferenceToReplace and ReplaceWith"));
+			}
+			else if (ChooserToReplace == Chooser)
+			{
+				if (!ModifiedChoosers.Contains(Chooser))
+				{
+					ModifiedChoosers.Add(Chooser);
+					Chooser->Modify();
+				}
+				if (!ModifiedChoosers.Contains(ChooserToReplaceWith))
+				{
+					ModifiedChoosers.Add(ChooserToReplaceWith);
+					ChooserToReplaceWith->Modify();
+				}
+				Chooser = ChooserToReplaceWith;
+			}
+		}
+	}
+
+	for (UChooserTable* ModifiedChooser : ModifiedChoosers)
+	{
+		FMapping(ModifiedChooser->GetRootChooser()).ApplyMappingToAllColumns();
+		ModifiedChooser->PostEditChange();
+	}
+
+	DO_ENSURE_CONSISTENCY(CheckConsistency(Chooser));
+}
+
+FName FPoseSearchColumn::RowValuesPropertyName()
+{
+	return "RowValues";
+}
+
+void FPoseSearchColumn::SetNumRows(int32 NumRows)
+{
+	using namespace UE::PoseSearch;
+	DO_ENSURE_CONSISTENCY(CheckConsistency(Chooser));
+
+	if (NumRows <= RowValues.Num())
+	{
+		RowValues.SetNum(NumRows);
+	}
+	else
+	{
+		while (RowValues.Num() < NumRows)
+		{
+			RowValues.Add(FPoseSearchColumnRow());
+		}
+	}
+	FMapping(Chooser->GetRootChooser()).ApplyMappingToAllColumns();
+
+	DO_ENSURE_CONSISTENCY(CheckConsistency(Chooser));
+}
+
+void FPoseSearchColumn::InsertRows(int Index, int Count)
+{
+	using namespace UE::PoseSearch;
+	DO_ENSURE_CONSISTENCY(CheckConsistency(Chooser));
+
+	for (int i=0;i<Count;i++)
+	{
+		RowValues.Insert(FPoseSearchColumnRow(), Index);
+	}
+	FMapping(Chooser->GetRootChooser()).ApplyMappingToAllColumns();
+
+	DO_ENSURE_CONSISTENCY(CheckConsistency(Chooser));
+}
+
+void FPoseSearchColumn::DeleteRows(TArrayView<int> RowIndices)
+{
+	using namespace UE::PoseSearch;
+	DO_ENSURE_CONSISTENCY(CheckConsistency(Chooser));
+
+	for(int32 Index : RowIndices)
+	{
+		if (RowValues.IsValidIndex(Index))
+		{
+			RowValues.RemoveAt(Index);
+		}
+	}
+	FMapping(Chooser->GetRootChooser()).ApplyMappingToAllColumns();
+
+	DO_ENSURE_CONSISTENCY(CheckConsistency(Chooser));
+}
+
+void FPoseSearchColumn::MoveRow(int SourceRowIndex, int TargetRowIndex)
+{
+	using namespace UE::PoseSearch;
+	DO_ENSURE_CONSISTENCY(CheckConsistency(Chooser));
+
+	FPoseSearchColumnRow RowData = RowValues[SourceRowIndex];
+    RowValues.RemoveAt(SourceRowIndex);
+    if (SourceRowIndex < TargetRowIndex)
+	{
+		TargetRowIndex--;
+	}
+    RowValues.Insert(RowData, TargetRowIndex);
+	FMapping(Chooser->GetRootChooser()).ApplyMappingToAllColumns();
+
+	DO_ENSURE_CONSISTENCY(CheckConsistency(Chooser));
+}
+
+void FPoseSearchColumn::CopyRow(FChooserColumnBase& SourceColumn, int SourceRowIndex, int TargetRowIndex)
+{
+	using namespace UE::PoseSearch;
+
+	const FPoseSearchColumn& SourcePoseSearchColumn = static_cast<FPoseSearchColumn&>(SourceColumn);
+	ensure(Chooser && SourcePoseSearchColumn.Chooser);
+	UChooserTable* RootChooser = Chooser->GetRootChooser();
+	ensure(RootChooser);
+
+	if (&SourcePoseSearchColumn == this && SourceRowIndex == TargetRowIndex)
+	{
+		DO_ENSURE_CONSISTENCY(CheckConsistency(Chooser));
+		return;
+	}
+	
+	DO_ENSURE_CONSISTENCY(CheckConsistency(Chooser));
+	DO_ENSURE_CONSISTENCY(CheckConsistency(SourcePoseSearchColumn.GetChooser()));
+
+	const bool bSameRootChooser = RootChooser == SourcePoseSearchColumn.Chooser->GetRootChooser();
+	const FPoseSearchColumnRow& SourceRow = SourcePoseSearchColumn.GetRowValue(SourceRowIndex);
+	if (SourceRow.Database)
+	{
+		ensure(SourceRow.Database->GetOuter() == SourcePoseSearchColumn.Chooser->GetRootChooser());
+		SourceRow.Database->Modify();
+	}
+
+	EditRowValue(TargetRowIndex) = SourceRow;
+
+	bool bPerformDatabasePostLoad = false;
+	bool bPerformPostEditChange = false;
+	UPoseSearchDatabase* NewTargetDatabase = nullptr;
+	if (!bSameRootChooser && SourceRow.Database)
+	{
+		FString SourceDatabasePlainName = SourceRow.Database->GetFName().GetPlainNameString();
+		for (UObject* NestedObject : RootChooser->NestedObjects)
+		{
+			if (UPoseSearchDatabase* Database = Cast<UPoseSearchDatabase>(NestedObject))
+			{
+				if (SourceRow.Database->Schema == Database->Schema)
+				{
+					FString DatabasePlainName = Database->GetFName().GetPlainNameString();
+					if (SourceDatabasePlainName == DatabasePlainName)
+					{
+						// we found a good candidate for NewTargetDatabase. Let's use it instead of creating a new  UPoseSearchDatabase
+						NewTargetDatabase = Database;
+						bPerformPostEditChange = true;
+						NewTargetDatabase->Modify();
+						break;
+					}
+				}
+			}
+		}
+
+		if (!NewTargetDatabase)
+		{
+			const FName NewTargetDatabaseUniqueName = MakeUniqueObjectName(RootChooser, UPoseSearchDatabase::StaticClass(), SourceRow.Database->GetFName());
+			FObjectDuplicationParameters Parameters = InitStaticDuplicateObjectParams(SourceRow.Database, RootChooser, NewTargetDatabaseUniqueName);
+			
+			// delaying UPoseSearchDatabase PostLoad (starting database indexing) until we set all the Chooser properties and ApplyMappingToAllColumns
+			Parameters.bSkipPostLoad = true;
+			bPerformDatabasePostLoad = true;
+			NewTargetDatabase = static_cast<UPoseSearchDatabase*>(StaticDuplicateObjectEx(Parameters));
+
+			RootChooser->Modify();
+			RootChooser->AddNestedObject(NewTargetDatabase);
+			RootChooser->PostEditChange();
+		}
+
+		RowValues[TargetRowIndex].Database = NewTargetDatabase;
+	}
+
+	ensure(!GetRowValue(TargetRowIndex).Database || GetRowValue(TargetRowIndex).Database->GetOuter() == Chooser->GetRootChooser());
+
+	FMapping(Chooser->GetRootChooser()).ApplyMappingToAllColumns();
+
+	if (bPerformDatabasePostLoad)
+	{
+		check(NewTargetDatabase);
+		NewTargetDatabase->ConditionalPostLoad();
+	}
+
+	if (bPerformPostEditChange)
+	{
+		check(NewTargetDatabase);
+		NewTargetDatabase->PostEditChange();
+	}
+
+	if (SourceRow.Database)
+	{
+		SourceRow.Database->PostEditChange();
+	}
+
+	DO_ENSURE_CONSISTENCY(CheckConsistency(Chooser));
+	DO_ENSURE_CONSISTENCY(CheckConsistency(SourcePoseSearchColumn.GetChooser()));
+}
+
+void FPoseSearchColumn::CopyFallback(FChooserColumnBase& SourceColumn)
+{
+	CopyRow(SourceColumn, ChooserColumn_SpecialIndex_Fallback, ChooserColumn_SpecialIndex_Fallback);
+}
+
+void FPoseSearchColumn::AutoPopulate(int32 RowIndex, UObject* OutputObject)
+{
+	using namespace UE::PoseSearch;
+	DO_ENSURE_CONSISTENCY(CheckConsistency(Chooser));
+	FPoseSearchColumnRow& RowValue = EditRowValue(RowIndex);
+	if (OutputObject != RowValue.Data.AnimAsset)
+	{
+		RowValue.Data.AnimAsset = OutputObject;
+	}
+	DO_ENSURE_CONSISTENCY(CheckConsistency(Chooser));
+}
+
+UScriptStruct* FPoseSearchColumn::GetInputBaseType() const
+{
+	return FPoseHistoryContextProperty::StaticStruct();
+}
+
+const UScriptStruct* FPoseSearchColumn::GetInputType() const
+{
+	return InputValue.IsValid() ? InputValue.GetScriptStruct() : nullptr;
+}
+
+FInstancedStruct* FPoseSearchColumn::GetInputValuePtr()
+{
+	return &InputValue;
+}
+
+void FPoseSearchColumn::SetInputType(const UScriptStruct* Type)
+{
+	InputValue.InitializeAs(Type);
+}
+
+#endif // WITH_EDITOR
+
+FChooserParameterBase* FPoseSearchColumn::GetInputValue()
+{
+	return InputValue.GetMutablePtr<FChooserParameterBase>();
 }
-#endif
 
 #undef LOCTEXT_NAMESPACE
diff --git a/Engine/Plugins/Animation/PoseSearch/Source/Runtime/Private/PoseSearchConstraints.cpp b/Engine/Plugins/Animation/PoseSearch/Source/Runtime/Private/PoseSearchConstraints.cpp
--- a/Engine/Plugins/Animation/PoseSearch/Source/Runtime/Private/PoseSearchConstraints.cpp
+++ b/Engine/Plugins/Animation/PoseSearch/Source/Runtime/Private/PoseSearchConstraints.cpp
@@ -10,385 +10,397 @@
 
 namespace UE::PoseSearch
 {
-
-bool FConstraints::GetConstraint(FRole Role, FName SocketName, float& OutDesiredReach, FTransform& OutTransform) const
-{
-	for (const FConstraint& Constraint : Constraints)
+	bool FConstraints::GetConstraint(FRole Role, FName SocketName, float& OutDesiredReach, FTransform& OutTransform) const
 	{
-		if (const FTransform* SocketTransform = Constraint.GetSocketTransform(Role, SocketName))
+		for (const FConstraint& Constraint : Constraints)
 		{
-			OutDesiredReach = Constraint.DesiredReach;
-			OutTransform = *SocketTransform;
-			return true;
+			if (const FTransform* SocketTransform = Constraint.GetSocketTransform(Role, SocketName))
+			{
+				OutDesiredReach = Constraint.DesiredReach;
+				OutTransform = *SocketTransform;
+				return true;
+			}
 		}
-	}
 
-	OutDesiredReach = 0.f;
-	OutTransform = FTransform::Identity;
-	return false;
-}
-
-int32 FConstraints::Num() const
-{
-	return Constraints.Num();
-}
+		OutDesiredReach = 0.f;
+		OutTransform = FTransform::Identity;
+		return false;
+	}
 
-const FConstraint* FConstraints::Find(const FConstraint& Constraint) const
-{
-	const int32 ConstraintIndex = Constraints.Find(Constraint);
-	if (ConstraintIndex != INDEX_NONE)
+	int32 FConstraints::Num() const
 	{
-		return &Constraints[ConstraintIndex];
+		return Constraints.Num();
 	}
-	return nullptr;
-}
 
-FConstraint* FConstraints::Find(const FConstraint& Constraint)
-{
-	const int32 ConstraintIndex = Constraints.Find(Constraint);
-	if (ConstraintIndex != INDEX_NONE)
+	const FConstraint* FConstraints::Find(const FConstraint& Constraint) const
 	{
-		return &Constraints[ConstraintIndex];
+		const int32 ConstraintIndex = Constraints.Find(Constraint);
+		if (ConstraintIndex != INDEX_NONE)
+		{
+			return &Constraints[ConstraintIndex];
+		}
+		return nullptr;
 	}
-	return nullptr;
-}
-
-void FConstraints::ProcessAll(float DeltaTime, const FSearchResult& SearchResult, const FInteractionSearchContext& SearchContext, const FConstraints* SearchContextPreviousConstraints)
-{
-	// making sure nobody already called UpdateConstraints or we didn't reset the constraints...
-	check(Constraints.IsEmpty());
 
-	const UPoseSearchDatabase* Database = SearchResult.Database.Get();
-	check(Database);
-
-	const FSearchIndexAsset* SearchIndexAsset = SearchResult.GetSearchIndexAsset();
-	check(SearchIndexAsset);
-
-	const FPoseSearchDatabaseAnimationAssetBase* DatabaseAnimationAssetBase = Database->GetDatabaseAnimationAsset(*SearchIndexAsset);
-	check(DatabaseAnimationAssetBase);
-
-	if (const UMultiAnimAsset* MultiAnimAsset = Cast<UMultiAnimAsset>(DatabaseAnimationAssetBase->GetAnimationAsset()))
+	FConstraint* FConstraints::Find(const FConstraint& Constraint)
 	{
-		TConstArrayView<FConstraint> PreviousConstraints;
-		if (SearchContextPreviousConstraints)
+		const int32 ConstraintIndex = Constraints.Find(Constraint);
+		if (ConstraintIndex != INDEX_NONE)
 		{
-			PreviousConstraints = SearchContextPreviousConstraints->Constraints;
+			return &Constraints[ConstraintIndex];
 		}
-
-		ProcessAll(MultiAnimAsset, DeltaTime, SearchResult.GetAssetTime(), SearchIndexAsset->IsMirrored(), SearchIndexAsset->GetBlendParameters(), SearchContext.GetAnimContexts(), SearchContext.GetRoles(), PreviousConstraints);
+		return nullptr;
 	}
-}
-
-void FConstraints::ProcessAll(const UMultiAnimAsset* MultiAnimAsset, float DeltaTime, float AssetTime, bool bIsMirrored, const FVector& BlendParameters,
-	const TConstArrayView<TObjectPtr<const UObject>> AnimContexts, const TConstArrayView<FRole> Roles, TConstArrayView<FConstraint> PreviousConstraints)
-{
-	check(MultiAnimAsset);
-	const int32 NumAnimContexts = AnimContexts.Num();
-	check(NumAnimContexts == Roles.Num());
-
-	Constraints.Reset();
 
-	FMemMark Mark(FMemStack::Get());
-
-	TArray<bool, TInlineAllocator<NumPreallocatedContraints, TMemStackAllocator<>>> VisitedPreviousConstraints;
-	VisitedPreviousConstraints.SetNumZeroed(PreviousConstraints.Num());
-		
-	// extracting the contraints for this SearchContext from MultiAnimAsset
-	FAnimNotifyContext PreAllocatedNotifyContext;
-	for (int32 AnimContextIndex = 0; AnimContextIndex < NumAnimContexts; ++AnimContextIndex)
+	void FConstraints::ProcessAll(float DeltaTime, const FSearchResult& SearchResult, const FInteractionSearchContext& SearchContext, const FConstraints* SearchContextPreviousConstraints)
 	{
-		const FRole& Role = Roles[AnimContextIndex];
-		const UAnimationAsset* Animation = MultiAnimAsset->GetAnimationAsset(Role);
+		// making sure nobody already called UpdateConstraints or we didn't reset the constraints...
+		ensure(Constraints.IsEmpty());
 
-		if (!Animation)
+		const UPoseSearchDatabase* Database = SearchResult.Database.Get();
+		if (ensure(Database))
 		{
-			UE_LOG(LogPoseSearch, Error, TEXT("FConstraints::ProcessAll aborting because of invalid animation asset for Role '%s' in MultiAnimAsset '%s'"), *Role.ToString(), *MultiAnimAsset->GetName());
-			return;
-		}
-
-		// initializing FAnimationAssetSampler only for extracting AnimNotifyStates
-		FAnimationAssetSampler Sampler(Animation, FTransform::Identity, FVector::ZeroVector, FAnimationAssetSampler::DefaultRootTransformSamplingRate, false, false);
-		Sampler.ExtractAnimNotifyEvents(AssetTime, PreAllocatedNotifyContext, [AssetTime, DeltaTime, this, &PreviousConstraints, &VisitedPreviousConstraints](const FAnimNotifyEvent* AnimNotifyEvent)
+			const FSearchIndexAsset* SearchIndexAsset = SearchResult.GetSearchIndexAsset();
+			if (ensure(SearchIndexAsset))
 			{
-				if (const UAnimNotifyState_PoseSearchConstraint* ConstraintNotifyState = Cast<const UAnimNotifyState_PoseSearchConstraint>(AnimNotifyEvent->NotifyStateClass))
+				const FPoseSearchDatabaseAnimationAssetBase* DatabaseAnimationAssetBase = Database->GetDatabaseAnimationAsset(SearchIndexAsset->GetSourceAssetIdx());
+				if (ensure(DatabaseAnimationAssetBase))
 				{
-					// subtracting ConstraintNotifyState->CoolDownTime from the validity of this UAnimNotifyState_PoseSearchConstraint
-					const float NotifyEndTime = AnimNotifyEvent->GetEndTriggerTime();
-					if (AssetTime <= (NotifyEndTime - ConstraintNotifyState->CoolDownTime))
+					if (const UMultiAnimAsset* MultiAnimAsset = Cast<UMultiAnimAsset>(DatabaseAnimationAssetBase->GetAnimationAsset()))
 					{
-						FConstraint NewConstraint(ConstraintNotifyState);
-						const int32 PreviousConstraintIndex = PreviousConstraints.Find(NewConstraint);
-						if (PreviousConstraintIndex != INDEX_NONE)
+						TConstArrayView<FConstraint> PreviousConstraints;
+						if (SearchContextPreviousConstraints)
 						{
-							VisitedPreviousConstraints[PreviousConstraintIndex] = true;
-							const FConstraint& PreviousConstraint = PreviousConstraints[PreviousConstraintIndex];
-
-							// @todo: for now we just ignore conflicting constraints, but we should probably blend PreviousConstraint into NewConstraint :/
-							//        perhaps keeping the min or RampUpTime, to make the constraint as following the fastest one to get into full DesiredReach, and the max CoolDownTime to fade off with the slowest rate
-							//        not really sure what to do about TranslationWeight and RotationWeight.. maybe slowly blending them to the newer one?
-							if (PreviousConstraint.RampUpTime > UE_KINDA_SMALL_NUMBER)
-							{
-								NewConstraint.DesiredReach = FMath::Min(1.f, PreviousConstraint.DesiredReach + (DeltaTime / PreviousConstraint.RampUpTime));
-							}
-							else
-							{
-								NewConstraint.DesiredReach = 1.f;
-							}
-						}
-						else if (NewConstraint.RampUpTime <= UE_KINDA_SMALL_NUMBER)
-						{
-							NewConstraint.DesiredReach = 1.f;
+							PreviousConstraints = SearchContextPreviousConstraints->Constraints;
 						}
 
-						// @todo: for now we also ignore duplicate constraints, and perhaps we should use a similar strategy of blending the constraints
-						Constraints.AddUnique(NewConstraint);
+						ProcessAll(MultiAnimAsset, DeltaTime, SearchResult.GetAssetTime(), SearchIndexAsset->IsMirrored(), SearchIndexAsset->GetBlendParameters(), SearchContext.GetAnimContexts(), SearchContext.GetRoles(), PreviousConstraints);
 					}
 				}
-				return true;
-			});
-	}
-
-	// adding remaining unprocessed constraints from PreviousConstraints
-	for (int32 PreviousConstraintIndex = 0; PreviousConstraintIndex < PreviousConstraints.Num(); ++PreviousConstraintIndex)
-	{
-		if (!VisitedPreviousConstraints[PreviousConstraintIndex])
-		{
-			const FConstraint& PreviousConstraint = PreviousConstraints[PreviousConstraintIndex];
-			if (PreviousConstraint.CoolDownTime > UE_KINDA_SMALL_NUMBER)
-			{
-				const float DesiredReach = PreviousConstraint.DesiredReach - (DeltaTime / PreviousConstraint.CoolDownTime);
-				if (DesiredReach > UE_KINDA_SMALL_NUMBER)
-				{
-					FConstraint& NewConstraint = Constraints.Add_GetRef(PreviousConstraint);
-					NewConstraint.DesiredReach = DesiredReach;
-				}
 			}
 		}
 	}
 
-	// process all the collected constraints
-	if (!Constraints.IsEmpty())
+	void FConstraints::ProcessAll(const UMultiAnimAsset* MultiAnimAsset, float DeltaTime, float AssetTime, bool bIsMirrored, const FVector& BlendParameters,
+		const TConstArrayView<TObjectPtr<const UObject>> AnimContexts, const TConstArrayView<FRole> Roles, TConstArrayView<FConstraint> PreviousConstraints)
 	{
-		TArray<FTransform, TInlineAllocator<PreallocatedRolesNum, TMemStackAllocator<>>> ComponentToWorldTransforms;
-		TArray<TConstArrayView<FTransform>, TInlineAllocator<PreallocatedRolesNum, TMemStackAllocator<>>> AnimContextsComponentSpaceTransforms;
-		TArray<FBoneContainer, TInlineAllocator<PreallocatedRolesNum, TMemStackAllocator<>>> BoneContainers;
-		TArray<FCompactPose, TInlineAllocator<PreallocatedRolesNum, TMemStackAllocator<>>> Poses;
-		TArray<FCSPose<FCompactPose>, TInlineAllocator<PreallocatedRolesNum, TMemStackAllocator<>>> ComponentSpacePoses;
-		TArray<const USkeletalMeshComponent*, TInlineAllocator<PreallocatedRolesNum, TMemStackAllocator<>>> SkeletalMeshComponents;
-
-		ComponentToWorldTransforms.SetNum(NumAnimContexts);
-		AnimContextsComponentSpaceTransforms.SetNum(NumAnimContexts);
-		BoneContainers.SetNum(NumAnimContexts);
-		Poses.SetNum(NumAnimContexts);
-		ComponentSpacePoses.SetNum(NumAnimContexts);
-		SkeletalMeshComponents.SetNum(NumAnimContexts);
-
-		FRoleToIndex RoleToIndex;
-		for (int32 AnimContextIndex = 0; AnimContextIndex < NumAnimContexts; ++AnimContextIndex)
+		if (ensure(MultiAnimAsset))
 		{
-			const UObject* AnimContext = AnimContexts[AnimContextIndex];
-			check(AnimContext);
-			const USkeletalMeshComponent* SkeletalMeshComponent = GetContextSkeletalMeshComponent(AnimContext, false);
-			check(SkeletalMeshComponent);
-			SkeletalMeshComponents[AnimContextIndex] = SkeletalMeshComponent;
-
-			// @todo: do we need this logic about URO etc? if we're executing this is because we update the mesh, so we probably don't need any interpolation...
-			const bool bUROInSync = SkeletalMeshComponent->ShouldUseUpdateRateOptimizations() && SkeletalMeshComponent->AnimUpdateRateParams != nullptr;
-			const bool bUsingExternalInterpolation = SkeletalMeshComponent->IsUsingExternalInterpolation();
-
-			const TArray<FTransform>& CachedComponentSpaceTransforms = SkeletalMeshComponent->GetCachedComponentSpaceTransforms();
-			const bool bArraySizesMatch = CachedComponentSpaceTransforms.Num() == SkeletalMeshComponent->GetComponentSpaceTransforms().Num();
-
-			ComponentToWorldTransforms[AnimContextIndex] = SkeletalMeshComponent->GetComponentTransform();
-			AnimContextsComponentSpaceTransforms[AnimContextIndex] = (bUROInSync || bUsingExternalInterpolation) && bArraySizesMatch ? CachedComponentSpaceTransforms : SkeletalMeshComponent->GetComponentSpaceTransforms();
-
-			// @todo: check FAnimNode_CopyPoseFromMesh::PreUpdate to see how to GetAnimCurves and GetCustomAttributes
-			const FRole& Role = Roles[AnimContextIndex];
-			RoleToIndex.Add(Role) = AnimContextIndex;
-
-			const UAnimationAsset* Animation = MultiAnimAsset->GetAnimationAsset(Role);
-			check(Animation);
-			const FTransform Origin = MultiAnimAsset->GetOrigin(Role);
-
-			FAnimationAssetSampler Sampler(Animation, Origin, BlendParameters);
-			const FTransform RootTransform = Sampler.ExtractRootTransform(AssetTime);
-
-			// @todo add support for mirroring if needed
-			//if (MirrorDataTables.IsValidIndex(ItemIndex) && MirrorDataTables[ItemIndex])
-			//{
-			//	const FMirrorDataCache MirrorDataCache(MirrorDataTables[ItemIndex]);
-			//	AssetRootBoneTransforms[ItemIndex] = MirrorDataCache.MirrorTransform(AssetRootBoneTransforms[ItemIndex]);
-			//}
-
-			// array containing the bone index of the root bone (0)
-			const USkeleton* Skeleton = Animation->GetSkeleton();
-			check(Skeleton && SkeletalMeshComponent->GetSkeletalMeshAsset()->GetSkeleton() == Skeleton);
-
-			const FReferenceSkeleton& ReferenceSkeleton = Skeleton->GetReferenceSkeleton();
-
-			// @todo: refine this BoneIndices / cache bone containers
-			const int32 NumBones = ReferenceSkeleton.GetNum();
-			TArray<uint16, TMemStackAllocator<>> BoneIndices;
-			BoneIndices.SetNumUninitialized(NumBones);
-			for (int32 BoneIndex = 0; BoneIndex < NumBones; ++BoneIndex)
+			const int32 NumAnimContexts = AnimContexts.Num();
+			if (ensure(NumAnimContexts == Roles.Num()))
 			{
-				BoneIndices[BoneIndex] = BoneIndex;
-			}
-
-			// extracting the pose, containing only the root bone from the Sampler 
-			BoneContainers[AnimContextIndex].InitializeTo(BoneIndices, UE::Anim::FCurveFilterSettings(UE::Anim::ECurveFilterMode::DisallowAll), *Skeleton);
-			Poses[AnimContextIndex].SetBoneContainer(&BoneContainers[AnimContextIndex]);
-			Sampler.ExtractPose(AssetTime, Poses[AnimContextIndex]);
+				Constraints.Reset();
 
-			// making sure the animation root bone transform is Identity, so we can confuse the root with the root BONE transform and preserve performances!
-			FTransform& RootBoneTransform = Poses[AnimContextIndex][FCompactPoseBoneIndex(0)];
-			if (!RootBoneTransform.Equals(FTransform::Identity))
-			{
-				const FVector Pos = RootBoneTransform.GetTranslation();
-				const FRotator Rot(RootBoneTransform.GetRotation());
-				UE_LOG(LogPoseSearch, Error, TEXT("FConstraints::ProcessAll unsupported non identity root bone in %s at time %f Pos(%f, %f, %f), Rot(%f, %f, %f)"), *Animation->GetName(), AssetTime, Pos.X, Pos.Y, Pos.Z, Rot.Pitch, Rot.Yaw, Rot.Roll);
-			}
-			RootBoneTransform = Sampler.ExtractRootTransform(AssetTime);
-			ComponentSpacePoses[AnimContextIndex].InitPose(Poses[AnimContextIndex]);
+				FMemMark Mark(FMemStack::Get());
 
-#if ENABLE_VISUAL_LOG
-			if (FVisualLogger::IsRecording())
-			{
-				static const TCHAR* LogName = TEXT("PoseSearchInteraction_Constraints");
+				TArray<bool, TInlineAllocator<NumPreallocatedContraints, TMemStackAllocator<>>> VisitedPreviousConstraints;
+				VisitedPreviousConstraints.SetNumZeroed(PreviousConstraints.Num());
 
-				for (int32 BoneIndex = 0; BoneIndex < ReferenceSkeleton.GetNum(); ++BoneIndex)
+				// extracting the contraints for this SearchContext from MultiAnimAsset
+				FAnimNotifyContext PreAllocatedNotifyContext;
+				for (int32 AnimContextIndex = 0; AnimContextIndex < NumAnimContexts; ++AnimContextIndex)
 				{
-					const int32 ParentBoneIndex = ReferenceSkeleton.GetParentIndex(BoneIndex);
-					if (ParentBoneIndex != INDEX_NONE)
-					{
-						const FVector CurrentBoneWorldLocation = ComponentToWorldTransforms[AnimContextIndex].TransformPosition(AnimContextsComponentSpaceTransforms[AnimContextIndex][BoneIndex].GetTranslation());
-						const FVector CurrentParentBoneWorldLocation = ComponentToWorldTransforms[AnimContextIndex].TransformPosition(AnimContextsComponentSpaceTransforms[AnimContextIndex][ParentBoneIndex].GetTranslation());
+					const FRole& Role = Roles[AnimContextIndex];
+					const UAnimationAsset* Animation = MultiAnimAsset->GetAnimationAsset(Role);
 
-						const FVector AnimBoneWorldLocation = ComponentSpacePoses[AnimContextIndex].GetComponentSpaceTransform(FCompactPoseBoneIndex(BoneIndex)).GetTranslation();
-						const FVector AnimParentBoneWorldLocation = ComponentSpacePoses[AnimContextIndex].GetComponentSpaceTransform(FCompactPoseBoneIndex(ParentBoneIndex)).GetTranslation();
+					if (!Animation)
+					{
+						UE_LOG(LogPoseSearch, Error, TEXT("FConstraints::ProcessAll aborting because of invalid animation asset for Role '%s' in MultiAnimAsset '%s'"), *Role.ToString(), *MultiAnimAsset->GetName());
+						return;
+					}
 
-						for (int32 DrawAnimContextIndex = 0; DrawAnimContextIndex < NumAnimContexts; ++DrawAnimContextIndex)
+					// initializing FAnimationAssetSampler only for extracting AnimNotifyStates
+					FAnimationAssetSampler Sampler(Animation, FTransform::Identity, FVector::ZeroVector, FAnimationAssetSampler::DefaultRootTransformSamplingRate, false, false);
+					Sampler.ExtractAnimNotifyEvents(AssetTime, PreAllocatedNotifyContext, [AssetTime, DeltaTime, this, &PreviousConstraints, &VisitedPreviousConstraints](const FAnimNotifyEvent* AnimNotifyEvent)
 						{
-							const UObject* DrawAnimContext = AnimContexts[DrawAnimContextIndex];
-							const FColor Color = DrawAnimContext == AnimContext ? FColor::Red : FColor::Blue;
-
-							// current characters poses
-							UE_VLOG_SEGMENT(DrawAnimContext, LogName, VeryVerbose, CurrentBoneWorldLocation, CurrentParentBoneWorldLocation, Color, TEXT(""));
+							if (const UAnimNotifyState_PoseSearchConstraint* ConstraintNotifyState = Cast<const UAnimNotifyState_PoseSearchConstraint>(AnimNotifyEvent->NotifyStateClass))
+							{
+								// subtracting ConstraintNotifyState->CoolDownTime from the validity of this UAnimNotifyState_PoseSearchConstraint
+								const float NotifyEndTime = AnimNotifyEvent->GetEndTriggerTime();
+								if (AssetTime <= (NotifyEndTime - ConstraintNotifyState->CoolDownTime))
+								{
+									FConstraint NewConstraint(ConstraintNotifyState);
+									const int32 PreviousConstraintIndex = PreviousConstraints.Find(NewConstraint);
+									if (PreviousConstraintIndex != INDEX_NONE)
+									{
+										VisitedPreviousConstraints[PreviousConstraintIndex] = true;
+										const FConstraint& PreviousConstraint = PreviousConstraints[PreviousConstraintIndex];
+
+										// @todo: for now we just ignore conflicting constraints, but we should probably blend PreviousConstraint into NewConstraint :/
+										//        perhaps keeping the min or RampUpTime, to make the constraint as following the fastest one to get into full DesiredReach, and the max CoolDownTime to fade off with the slowest rate
+										//        not really sure what to do about TranslationWeight and RotationWeight.. maybe slowly blending them to the newer one?
+										if (PreviousConstraint.RampUpTime > UE_KINDA_SMALL_NUMBER)
+										{
+											NewConstraint.DesiredReach = FMath::Min(1.f, PreviousConstraint.DesiredReach + (DeltaTime / PreviousConstraint.RampUpTime));
+										}
+										else
+										{
+											NewConstraint.DesiredReach = 1.f;
+										}
+									}
+									else if (NewConstraint.RampUpTime <= UE_KINDA_SMALL_NUMBER)
+									{
+										NewConstraint.DesiredReach = 1.f;
+									}
+
+									// @todo: for now we also ignore duplicate constraints, and perhaps we should use a similar strategy of blending the constraints
+									Constraints.AddUnique(NewConstraint);
+								}
+							}
+							return true;
+						});
+				}
 
-							// sampled poses from UAnimationAsset
-							UE_VLOG_SEGMENT(DrawAnimContext, LogName, VeryVerbose, AnimBoneWorldLocation, AnimParentBoneWorldLocation, Color, TEXT(""));
+				// adding remaining unprocessed constraints from PreviousConstraints
+				for (int32 PreviousConstraintIndex = 0; PreviousConstraintIndex < PreviousConstraints.Num(); ++PreviousConstraintIndex)
+				{
+					if (!VisitedPreviousConstraints[PreviousConstraintIndex])
+					{
+						const FConstraint& PreviousConstraint = PreviousConstraints[PreviousConstraintIndex];
+						if (PreviousConstraint.CoolDownTime > UE_KINDA_SMALL_NUMBER)
+						{
+							const float DesiredReach = PreviousConstraint.DesiredReach - (DeltaTime / PreviousConstraint.CoolDownTime);
+							if (DesiredReach > UE_KINDA_SMALL_NUMBER)
+							{
+								FConstraint& NewConstraint = Constraints.Add_GetRef(PreviousConstraint);
+								NewConstraint.DesiredReach = DesiredReach;
+							}
 						}
 					}
 				}
-			}
+
+				// process all the collected constraints
+				if (!Constraints.IsEmpty())
+				{
+					TArray<FTransform, TInlineAllocator<PreallocatedRolesNum, TMemStackAllocator<>>> ComponentToWorldTransforms;
+					TArray<TConstArrayView<FTransform>, TInlineAllocator<PreallocatedRolesNum, TMemStackAllocator<>>> AnimContextsComponentSpaceTransforms;
+					TArray<FBoneContainer, TInlineAllocator<PreallocatedRolesNum, TMemStackAllocator<>>> BoneContainers;
+					TArray<FCompactPose, TInlineAllocator<PreallocatedRolesNum, TMemStackAllocator<>>> Poses;
+					TArray<FCSPose<FCompactPose>, TInlineAllocator<PreallocatedRolesNum, TMemStackAllocator<>>> ComponentSpacePoses;
+					TArray<const USkeletalMeshComponent*, TInlineAllocator<PreallocatedRolesNum, TMemStackAllocator<>>> SkeletalMeshComponents;
+
+					ComponentToWorldTransforms.SetNum(NumAnimContexts);
+					AnimContextsComponentSpaceTransforms.SetNum(NumAnimContexts);
+					BoneContainers.SetNum(NumAnimContexts);
+					Poses.SetNum(NumAnimContexts);
+					ComponentSpacePoses.SetNum(NumAnimContexts);
+					SkeletalMeshComponents.SetNum(NumAnimContexts);
+
+					FRoleToIndex RoleToIndex;
+					for (int32 AnimContextIndex = 0; AnimContextIndex < NumAnimContexts; ++AnimContextIndex)
+					{
+						const UObject* AnimContext = AnimContexts[AnimContextIndex];
+						if (ensure(AnimContext))
+						{
+							const USkeletalMeshComponent* SkeletalMeshComponent = GetContextSkeletalMeshComponent(AnimContext, false);
+							if (ensure(SkeletalMeshComponent))
+							{
+								SkeletalMeshComponents[AnimContextIndex] = SkeletalMeshComponent;
+
+								// @todo: do we need this logic about URO etc? if we're executing this is because we update the mesh, so we probably don't need any interpolation...
+								const bool bUROInSync = SkeletalMeshComponent->ShouldUseUpdateRateOptimizations() && SkeletalMeshComponent->AnimUpdateRateParams != nullptr;
+								const bool bUsingExternalInterpolation = SkeletalMeshComponent->IsUsingExternalInterpolation();
+
+								const TArray<FTransform>& CachedComponentSpaceTransforms = SkeletalMeshComponent->GetCachedComponentSpaceTransforms();
+								const bool bArraySizesMatch = CachedComponentSpaceTransforms.Num() == SkeletalMeshComponent->GetComponentSpaceTransforms().Num();
+
+								ComponentToWorldTransforms[AnimContextIndex] = SkeletalMeshComponent->GetComponentTransform();
+								AnimContextsComponentSpaceTransforms[AnimContextIndex] = (bUROInSync || bUsingExternalInterpolation) && bArraySizesMatch ? CachedComponentSpaceTransforms : SkeletalMeshComponent->GetComponentSpaceTransforms();
+
+								// @todo: check FAnimNode_CopyPoseFromMesh::PreUpdate to see how to GetAnimCurves and GetCustomAttributes
+								const FRole& Role = Roles[AnimContextIndex];
+								RoleToIndex.Add(Role) = AnimContextIndex;
+
+								const UAnimationAsset* Animation = MultiAnimAsset->GetAnimationAsset(Role);
+								if (ensure(Animation))
+								{
+									const FTransform Origin = MultiAnimAsset->GetOrigin(Role);
+
+									FAnimationAssetSampler Sampler(Animation, Origin, BlendParameters);
+									const FTransform RootTransform = Sampler.ExtractRootTransform(AssetTime);
+
+									// @todo add support for mirroring if needed
+									//if (MirrorDataTables.IsValidIndex(ItemIndex) && MirrorDataTables[ItemIndex])
+									//{
+									//	const FMirrorDataCache MirrorDataCache(MirrorDataTables[ItemIndex]);
+									//	AssetRootBoneTransforms[ItemIndex] = MirrorDataCache.MirrorTransform(AssetRootBoneTransforms[ItemIndex]);
+									//}
+
+									// array containing the bone index of the root bone (0)
+									const USkeleton* Skeleton = Animation->GetSkeleton();
+									if (ensure(Skeleton && SkeletalMeshComponent->GetSkeletalMeshAsset()->GetSkeleton() == Skeleton))
+									{
+										const FReferenceSkeleton& ReferenceSkeleton = Skeleton->GetReferenceSkeleton();
+
+										// @todo: refine this BoneIndices / cache bone containers
+										const int32 NumBones = ReferenceSkeleton.GetNum();
+										TArray<uint16, TMemStackAllocator<>> BoneIndices;
+										BoneIndices.SetNumUninitialized(NumBones);
+										for (int32 BoneIndex = 0; BoneIndex < NumBones; ++BoneIndex)
+										{
+											BoneIndices[BoneIndex] = BoneIndex;
+										}
+
+										// extracting the pose, containing only the root bone from the Sampler 
+										BoneContainers[AnimContextIndex].InitializeTo(BoneIndices, UE::Anim::FCurveFilterSettings(UE::Anim::ECurveFilterMode::DisallowAll), *Skeleton);
+										Poses[AnimContextIndex].SetBoneContainer(&BoneContainers[AnimContextIndex]);
+										Sampler.ExtractPose(AssetTime, Poses[AnimContextIndex]);
+
+										// making sure the animation root bone transform is Identity, so we can confuse the root with the root BONE transform and preserve performances!
+										FTransform& RootBoneTransform = Poses[AnimContextIndex][FCompactPoseBoneIndex(0)];
+										if (!RootBoneTransform.Equals(FTransform::Identity))
+										{
+											const FVector Pos = RootBoneTransform.GetTranslation();
+											const FRotator Rot(RootBoneTransform.GetRotation());
+											UE_LOG(LogPoseSearch, Error, TEXT("FConstraints::ProcessAll unsupported non identity root bone in %s at time %f Pos(%f, %f, %f), Rot(%f, %f, %f)"), *Animation->GetName(), AssetTime, Pos.X, Pos.Y, Pos.Z, Rot.Pitch, Rot.Yaw, Rot.Roll);
+										}
+										RootBoneTransform = Sampler.ExtractRootTransform(AssetTime);
+										ComponentSpacePoses[AnimContextIndex].InitPose(Poses[AnimContextIndex]);
+
+#if ENABLE_VISUAL_LOG
+										if (FVisualLogger::IsRecording())
+										{
+											static const TCHAR* LogName = TEXT("PoseSearchInteraction_Constraints");
+
+											for (int32 BoneIndex = 0; BoneIndex < ReferenceSkeleton.GetNum(); ++BoneIndex)
+											{
+												const int32 ParentBoneIndex = ReferenceSkeleton.GetParentIndex(BoneIndex);
+												if (ParentBoneIndex != INDEX_NONE)
+												{
+													const FVector CurrentBoneWorldLocation = ComponentToWorldTransforms[AnimContextIndex].TransformPosition(AnimContextsComponentSpaceTransforms[AnimContextIndex][BoneIndex].GetTranslation());
+													const FVector CurrentParentBoneWorldLocation = ComponentToWorldTransforms[AnimContextIndex].TransformPosition(AnimContextsComponentSpaceTransforms[AnimContextIndex][ParentBoneIndex].GetTranslation());
+
+													const FVector AnimBoneWorldLocation = ComponentSpacePoses[AnimContextIndex].GetComponentSpaceTransform(FCompactPoseBoneIndex(BoneIndex)).GetTranslation();
+													const FVector AnimParentBoneWorldLocation = ComponentSpacePoses[AnimContextIndex].GetComponentSpaceTransform(FCompactPoseBoneIndex(ParentBoneIndex)).GetTranslation();
+
+													for (int32 DrawAnimContextIndex = 0; DrawAnimContextIndex < NumAnimContexts; ++DrawAnimContextIndex)
+													{
+														const UObject* DrawAnimContext = AnimContexts[DrawAnimContextIndex];
+														const FColor Color = DrawAnimContext == AnimContext ? FColor::Red : FColor::Blue;
+
+														// current characters poses
+														UE_VLOG_SEGMENT(DrawAnimContext, LogName, VeryVerbose, CurrentBoneWorldLocation, CurrentParentBoneWorldLocation, Color, TEXT(""));
+
+														// sampled poses from UAnimationAsset
+														UE_VLOG_SEGMENT(DrawAnimContext, LogName, VeryVerbose, AnimBoneWorldLocation, AnimParentBoneWorldLocation, Color, TEXT(""));
+													}
+												}
+											}
+										}
 #endif //ENABLE_VISUAL_LOG
-		}
+									}
+								}
+							}
+						}
+					}
 
-		for (TConstraints::TIterator ConstraintIt = Constraints.CreateIterator(); ConstraintIt; ++ConstraintIt)
-		{
-			FConstraint& Constraint = *ConstraintIt;
-			const int32* FromSocketRoleIndex = RoleToIndex.Find(Constraint.FromSocketRole);
-			const int32* ToSocketRoleIndex = RoleToIndex.Find(Constraint.ToSocketRole);
+					for (TConstraints::TIterator ConstraintIt = Constraints.CreateIterator(); ConstraintIt; ++ConstraintIt)
+					{
+						FConstraint& Constraint = *ConstraintIt;
+						const int32* FromSocketRoleIndex = RoleToIndex.Find(Constraint.FromSocketRole);
+						const int32* ToSocketRoleIndex = RoleToIndex.Find(Constraint.ToSocketRole);
 
-			if (!FromSocketRoleIndex || !ToSocketRoleIndex)
-			{
+						if (!FromSocketRoleIndex || !ToSocketRoleIndex)
+						{
 #if !NO_LOGGING
-				if (!FromSocketRoleIndex)
-				{
-					UE_LOG(LogPoseSearch, Error, TEXT("FConstraints::ProcessAll, Failed to find Role '%s' used in a PoseSearchConstraint in asset '%s'"), *Constraint.FromSocketRole.ToString(), *MultiAnimAsset->GetName());
-				}
+							if (!FromSocketRoleIndex)
+							{
+								UE_LOG(LogPoseSearch, Error, TEXT("FConstraints::ProcessAll, Failed to find Role '%s' used in a PoseSearchConstraint in asset '%s'"), *Constraint.FromSocketRole.ToString(), *MultiAnimAsset->GetName());
+							}
 
-				if (!ToSocketRoleIndex)
-				{
-					UE_LOG(LogPoseSearch, Error, TEXT("FConstraints::ProcessAll, Failed to find Role '%s' used in a PoseSearchConstraint in asset '%s'"), *Constraint.ToSocketRole.ToString(), *MultiAnimAsset->GetName());
-				}
+							if (!ToSocketRoleIndex)
+							{
+								UE_LOG(LogPoseSearch, Error, TEXT("FConstraints::ProcessAll, Failed to find Role '%s' used in a PoseSearchConstraint in asset '%s'"), *Constraint.ToSocketRole.ToString(), *MultiAnimAsset->GetName());
+							}
 #endif // !NO_LOGGING
-				ConstraintIt.RemoveCurrent();
-			}
-			else
-			{
-				int32 FromBoneIndex = INDEX_NONE;
-				int32 ToBoneIndex = INDEX_NONE;
+							ConstraintIt.RemoveCurrent();
+						}
+						else
+						{
+							int32 FromBoneIndex = INDEX_NONE;
+							int32 ToBoneIndex = INDEX_NONE;
 
-				FTransform FromSocketLocalTransform(FTransform::Identity);
-				FTransform ToSocketLocalTransform(FTransform::Identity);
-					
-				if (!SkeletalMeshComponents[*FromSocketRoleIndex]->GetSocketInfoByName(Constraint.FromSocketName, FromSocketLocalTransform, FromBoneIndex))
-				{
-					FromBoneIndex = SkeletalMeshComponents[*FromSocketRoleIndex]->GetBoneIndex(Constraint.FromSocketName);
-				}
+							FTransform FromSocketLocalTransform(FTransform::Identity);
+							FTransform ToSocketLocalTransform(FTransform::Identity);
 
-				if (!SkeletalMeshComponents[*ToSocketRoleIndex]->GetSocketInfoByName(Constraint.ToSocketName, ToSocketLocalTransform, ToBoneIndex))
-				{
-					ToBoneIndex = SkeletalMeshComponents[*ToSocketRoleIndex]->GetBoneIndex(Constraint.ToSocketName);
-				}
+							if (!SkeletalMeshComponents[*FromSocketRoleIndex]->GetSocketInfoByName(Constraint.FromSocketName, FromSocketLocalTransform, FromBoneIndex))
+							{
+								FromBoneIndex = SkeletalMeshComponents[*FromSocketRoleIndex]->GetBoneIndex(Constraint.FromSocketName);
+							}
 
-				if (FromBoneIndex == INDEX_NONE || ToBoneIndex == INDEX_NONE)
-				{
-#if !NO_LOGGING
-					if (FromBoneIndex == INDEX_NONE)
-					{
-						UE_LOG(LogPoseSearch, Error, TEXT("FConstraints::ProcessAll, Failed to find Socket '%s' from SkeletalMeshComponent '%s' for Role '%s'"), *Constraint.FromSocketName.ToString(), *GetNameSafe(SkeletalMeshComponents[*FromSocketRoleIndex]), *Constraint.FromSocketRole.ToString());
-					}
+							if (!SkeletalMeshComponents[*ToSocketRoleIndex]->GetSocketInfoByName(Constraint.ToSocketName, ToSocketLocalTransform, ToBoneIndex))
+							{
+								ToBoneIndex = SkeletalMeshComponents[*ToSocketRoleIndex]->GetBoneIndex(Constraint.ToSocketName);
+							}
 
-					if (ToBoneIndex == INDEX_NONE)
-					{
-						UE_LOG(LogPoseSearch, Error, TEXT("FConstraints::ProcessAll, Failed to find Socket '%s' from SkeletalMeshComponent '%s' for Role '%s'"), *Constraint.ToSocketName.ToString(), *GetNameSafe(SkeletalMeshComponents[*ToSocketRoleIndex]), *Constraint.ToSocketRole.ToString());
-					}
+							if (FromBoneIndex == INDEX_NONE || ToBoneIndex == INDEX_NONE)
+							{
+#if !NO_LOGGING
+								if (FromBoneIndex == INDEX_NONE)
+								{
+									UE_LOG(LogPoseSearch, Error, TEXT("FConstraints::ProcessAll, Failed to find Socket '%s' from SkeletalMeshComponent '%s' for Role '%s'"), *Constraint.FromSocketName.ToString(), *GetNameSafe(SkeletalMeshComponents[*FromSocketRoleIndex]), *Constraint.FromSocketRole.ToString());
+								}
+
+								if (ToBoneIndex == INDEX_NONE)
+								{
+									UE_LOG(LogPoseSearch, Error, TEXT("FConstraints::ProcessAll, Failed to find Socket '%s' from SkeletalMeshComponent '%s' for Role '%s'"), *Constraint.ToSocketName.ToString(), *GetNameSafe(SkeletalMeshComponents[*ToSocketRoleIndex]), *Constraint.ToSocketRole.ToString());
+								}
 #endif // !NO_LOGGING
-					ConstraintIt.RemoveCurrent();
-				}
-				else
-				{
-					const FTransform CurrentFromTransform = FromSocketLocalTransform * (AnimContextsComponentSpaceTransforms[*FromSocketRoleIndex][FromBoneIndex] * ComponentToWorldTransforms[*FromSocketRoleIndex]);
-					const FTransform CurrentToTransform = ToSocketLocalTransform * (AnimContextsComponentSpaceTransforms[*ToSocketRoleIndex][ToBoneIndex] * ComponentToWorldTransforms[*ToSocketRoleIndex]);
+								ConstraintIt.RemoveCurrent();
+							}
+							else
+							{
+								const FTransform CurrentFromTransform = FromSocketLocalTransform * (AnimContextsComponentSpaceTransforms[*FromSocketRoleIndex][FromBoneIndex] * ComponentToWorldTransforms[*FromSocketRoleIndex]);
+								const FTransform CurrentToTransform = ToSocketLocalTransform * (AnimContextsComponentSpaceTransforms[*ToSocketRoleIndex][ToBoneIndex] * ComponentToWorldTransforms[*ToSocketRoleIndex]);
 
-					const FTransform AnimFromTransform = FromSocketLocalTransform * (ComponentSpacePoses[*FromSocketRoleIndex].GetComponentSpaceTransform(FCompactPoseBoneIndex(FromBoneIndex)));
-					const FTransform AnimToTransform = ToSocketLocalTransform * (ComponentSpacePoses[*ToSocketRoleIndex].GetComponentSpaceTransform(FCompactPoseBoneIndex(ToBoneIndex)));
+								const FTransform AnimFromTransform = FromSocketLocalTransform * (ComponentSpacePoses[*FromSocketRoleIndex].GetComponentSpaceTransform(FCompactPoseBoneIndex(FromBoneIndex)));
+								const FTransform AnimToTransform = ToSocketLocalTransform * (ComponentSpacePoses[*ToSocketRoleIndex].GetComponentSpaceTransform(FCompactPoseBoneIndex(ToBoneIndex)));
 
-					const FTransform WantedCurrentToTransform = AnimToTransform.GetRelativeTransform(AnimFromTransform) * CurrentFromTransform;
-					const FTransform WantedCurrentFromTransform = AnimFromTransform.GetRelativeTransform(AnimToTransform) * CurrentToTransform;
+								const FTransform WantedCurrentToTransform = AnimToTransform.GetRelativeTransform(AnimFromTransform) * CurrentFromTransform;
+								const FTransform WantedCurrentFromTransform = AnimFromTransform.GetRelativeTransform(AnimToTransform) * CurrentToTransform;
 
-					Constraint.FromSocketTransform.SetTranslation(FMath::Lerp(WantedCurrentFromTransform.GetTranslation(), CurrentFromTransform.GetTranslation(), Constraint.TranslationWeight));
-					Constraint.FromSocketTransform.SetRotation(FQuat::Slerp(WantedCurrentFromTransform.GetRotation(), CurrentFromTransform.GetRotation(), Constraint.RotationWeight));
-					Constraint.FromSocketTransform.SetScale3D(CurrentFromTransform.GetScale3D());
+								Constraint.FromSocketTransform.SetTranslation(FMath::Lerp(WantedCurrentFromTransform.GetTranslation(), CurrentFromTransform.GetTranslation(), Constraint.TranslationWeight));
+								Constraint.FromSocketTransform.SetRotation(FQuat::Slerp(WantedCurrentFromTransform.GetRotation(), CurrentFromTransform.GetRotation(), Constraint.RotationWeight));
+								Constraint.FromSocketTransform.SetScale3D(CurrentFromTransform.GetScale3D());
 
-					Constraint.ToSocketTransform.SetTranslation(FMath::Lerp(WantedCurrentToTransform.GetTranslation(), CurrentToTransform.GetTranslation(), 1.f - Constraint.TranslationWeight));
-					Constraint.ToSocketTransform.SetRotation(FQuat::Slerp(WantedCurrentToTransform.GetRotation(), CurrentToTransform.GetRotation(), 1.f - Constraint.RotationWeight));
-					Constraint.ToSocketTransform.SetScale3D(CurrentToTransform.GetScale3D());
+								Constraint.ToSocketTransform.SetTranslation(FMath::Lerp(WantedCurrentToTransform.GetTranslation(), CurrentToTransform.GetTranslation(), 1.f - Constraint.TranslationWeight));
+								Constraint.ToSocketTransform.SetRotation(FQuat::Slerp(WantedCurrentToTransform.GetRotation(), CurrentToTransform.GetRotation(), 1.f - Constraint.RotationWeight));
+								Constraint.ToSocketTransform.SetScale3D(CurrentToTransform.GetScale3D());
 
 #if WITH_EDITORONLY_DATA && ENABLE_ANIM_DEBUG && ENABLE_VISUAL_LOG
-					if (FVisualLogger::IsRecording())
-					{
-						static const TCHAR* LogName = TEXT("PoseSearchInteraction_Constraints");
+								if (FVisualLogger::IsRecording())
+								{
+									static const TCHAR* LogName = TEXT("PoseSearchInteraction_Constraints");
 
-						for (int32 DrawAnimContextIndex = 0; DrawAnimContextIndex < NumAnimContexts; ++DrawAnimContextIndex)
-						{
-							const UObject* DrawAnimContext = AnimContexts[DrawAnimContextIndex];
+									for (int32 DrawAnimContextIndex = 0; DrawAnimContextIndex < NumAnimContexts; ++DrawAnimContextIndex)
+									{
+										const UObject* DrawAnimContext = AnimContexts[DrawAnimContextIndex];
 
-							// current characters poses
-							UE_VLOG_SEGMENT(DrawAnimContext, LogName, VeryVerbose, CurrentFromTransform.GetTranslation(), CurrentToTransform.GetTranslation(), FColor::Yellow, TEXT(""));
+										// current characters poses
+										UE_VLOG_SEGMENT(DrawAnimContext, LogName, VeryVerbose, CurrentFromTransform.GetTranslation(), CurrentToTransform.GetTranslation(), FColor::Yellow, TEXT(""));
 
-							// sampled poses from UAnimationAsset
-							UE_VLOG_SEGMENT(DrawAnimContext, LogName, VeryVerbose, AnimFromTransform.GetTranslation(), AnimToTransform.GetTranslation(), FColor::Yellow, TEXT(""));
+										// sampled poses from UAnimationAsset
+										UE_VLOG_SEGMENT(DrawAnimContext, LogName, VeryVerbose, AnimFromTransform.GetTranslation(), AnimToTransform.GetTranslation(), FColor::Yellow, TEXT(""));
 
-							// wanted transforms
-							UE_VLOG_SEGMENT(DrawAnimContext, LogName, VeryVerbose, CurrentFromTransform.GetTranslation(), WantedCurrentToTransform.GetTranslation(), FColor::Green, TEXT(""));
-							UE_VLOG_SEGMENT(DrawAnimContext, LogName, VeryVerbose, CurrentToTransform.GetTranslation(), WantedCurrentFromTransform.GetTranslation(), FColor::Green, TEXT(""));
+										// wanted transforms
+										UE_VLOG_SEGMENT(DrawAnimContext, LogName, VeryVerbose, CurrentFromTransform.GetTranslation(), WantedCurrentToTransform.GetTranslation(), FColor::Green, TEXT(""));
+										UE_VLOG_SEGMENT(DrawAnimContext, LogName, VeryVerbose, CurrentToTransform.GetTranslation(), WantedCurrentFromTransform.GetTranslation(), FColor::Green, TEXT(""));
 
-							UE_VLOG_SEGMENT(DrawAnimContext, LogName, VeryVerbose, CurrentFromTransform.GetTranslation(), Constraint.FromSocketTransform.GetTranslation(), FColor::Orange, TEXT(""));
-							UE_VLOG_SEGMENT(DrawAnimContext, LogName, VeryVerbose, CurrentToTransform.GetTranslation(), Constraint.ToSocketTransform.GetTranslation(), FColor::Orange, TEXT(""));
+										UE_VLOG_SEGMENT(DrawAnimContext, LogName, VeryVerbose, CurrentFromTransform.GetTranslation(), Constraint.FromSocketTransform.GetTranslation(), FColor::Orange, TEXT(""));
+										UE_VLOG_SEGMENT(DrawAnimContext, LogName, VeryVerbose, CurrentToTransform.GetTranslation(), Constraint.ToSocketTransform.GetTranslation(), FColor::Orange, TEXT(""));
 
-							const FColor DesiredReachColor = (FLinearColor::Red* (1.f - Constraint.DesiredReach) + FLinearColor::Green * Constraint.DesiredReach).ToFColor(true);
-							UE_VLOG_SPHERE(DrawAnimContext, LogName, Display, Constraint.FromSocketTransform.GetTranslation(), 2.f, DesiredReachColor, TEXT(""));
-							UE_VLOG_SPHERE(DrawAnimContext, LogName, Display, Constraint.ToSocketTransform.GetTranslation(), 2.f, DesiredReachColor, TEXT(""));
+										const FColor DesiredReachColor = (FLinearColor::Red * (1.f - Constraint.DesiredReach) + FLinearColor::Green * Constraint.DesiredReach).ToFColor(true);
+										UE_VLOG_SPHERE(DrawAnimContext, LogName, Display, Constraint.FromSocketTransform.GetTranslation(), 2.f, DesiredReachColor, TEXT(""));
+										UE_VLOG_SPHERE(DrawAnimContext, LogName, Display, Constraint.ToSocketTransform.GetTranslation(), 2.f, DesiredReachColor, TEXT(""));
+									}
+								}
+#endif // WITH_EDITORONLY_DATA && ENABLE_ANIM_DEBUG && ENABLE_VISUAL_LOG
+							}
 						}
 					}
-#endif // WITH_EDITORONLY_DATA && ENABLE_ANIM_DEBUG && ENABLE_VISUAL_LOG
 				}
 			}
 		}
 	}
-}
 
 } // namespace UE::PoseSearch
diff --git a/Engine/Plugins/Animation/PoseSearch/Source/Runtime/Private/PoseSearchDatabase.cpp b/Engine/Plugins/Animation/PoseSearch/Source/Runtime/Private/PoseSearchDatabase.cpp
--- a/Engine/Plugins/Animation/PoseSearch/Source/Runtime/Private/PoseSearchDatabase.cpp
+++ b/Engine/Plugins/Animation/PoseSearch/Source/Runtime/Private/PoseSearchDatabase.cpp
@@ -1,16 +1,13 @@
 // Copyright Epic Games, Inc. All Rights Reserved.
 
 #include "PoseSearch/PoseSearchDatabase.h"
-
-// @todo: remove this include once the deprecated method FPoseSearchDatabaseAnimationAssetBase::GetFrameAtTime is removed. 
-//        used only to get UAnimationSettings::Get()->GetDefaultFrameRate()
-#include "Animation/AnimationSettings.h"
-
 #include "Animation/AnimComposite.h"
 #include "Animation/AnimMontage.h"
 #include "Animation/AnimSequence.h"
 #include "Animation/BlendSpace.h"
 #include "Animation/BlendSpace1D.h"
+#include "Chooser.h"
+#include "PoseSearch/Chooser/PoseSearchChooserColumn.h"
 #include "PoseSearch/MultiAnimAsset.h"
 #include "PoseSearch/PoseSearchAnimNotifies.h"
 #include "PoseSearch/PoseSearchContext.h"
@@ -195,44 +192,45 @@ static void IterateOverBlendSpaceSamplingParameter(const UBlendSpace* BlendSpace
 	bool bUseGridForSampling, int32 NumberOfHorizontalSamples, int32 NumberOfVerticalSamples,
 	const TFunction<void(const FVector& BlendParameters)>& ProcessSamplingParameter)
 {
-	check(BlendSpace);
-	
-	if (bUseSingleSample)
+	if (ensure(BlendSpace))
 	{
-		ProcessSamplingParameter(SingleSampleBlendParameters);
-	}
-	else if (bUseGridForSampling)
-	{
-		for (const FBlendSample& BlendSample : BlendSpace->GetBlendSamples())
+		if (bUseSingleSample)
 		{
-			ProcessSamplingParameter(BlendSample.SampleValue);
+			ProcessSamplingParameter(SingleSampleBlendParameters);
 		}
-	}
-	else
-	{
-		const int32 HorizontalBlendNum = FMath::Max(NumberOfHorizontalSamples, 1);
-		const int32 VerticalBlendNum = BlendSpace->IsA<UBlendSpace1D>() ? 1 : FMath::Max(NumberOfVerticalSamples, 1);
+		else if (bUseGridForSampling)
+		{
+			for (const FBlendSample& BlendSample : BlendSpace->GetBlendSamples())
+			{
+				ProcessSamplingParameter(BlendSample.SampleValue);
+			}
+		}
+		else
+		{
+			const int32 HorizontalBlendNum = FMath::Max(NumberOfHorizontalSamples, 1);
+			const int32 VerticalBlendNum = BlendSpace->IsA<UBlendSpace1D>() ? 1 : FMath::Max(NumberOfVerticalSamples, 1);
 
-		const FBlendParameter& HorizontalBlendParameter = BlendSpace->GetBlendParameter(0);
-		const FBlendParameter& VerticalBlendParameter = BlendSpace->GetBlendParameter(1);
+			const FBlendParameter& HorizontalBlendParameter = BlendSpace->GetBlendParameter(0);
+			const FBlendParameter& VerticalBlendParameter = BlendSpace->GetBlendParameter(1);
 
-		const int32 WrapInputHorizontalBlendNum = HorizontalBlendParameter.bWrapInput ? HorizontalBlendNum + 1 : HorizontalBlendNum;
-		const int32 WrapInputVerticalBlendNum = VerticalBlendParameter.bWrapInput ? VerticalBlendNum + 1 : VerticalBlendNum;
+			const int32 WrapInputHorizontalBlendNum = HorizontalBlendParameter.bWrapInput ? HorizontalBlendNum + 1 : HorizontalBlendNum;
+			const int32 WrapInputVerticalBlendNum = VerticalBlendParameter.bWrapInput ? VerticalBlendNum + 1 : VerticalBlendNum;
 
-		for (int32 HorizontalBlendIndex = 0; HorizontalBlendIndex < HorizontalBlendNum; HorizontalBlendIndex++)
-		{
-			for (int32 VerticalBlendIndex = 0; VerticalBlendIndex < VerticalBlendNum; VerticalBlendIndex++)
+			for (int32 HorizontalBlendIndex = 0; HorizontalBlendIndex < HorizontalBlendNum; HorizontalBlendIndex++)
 			{
-				const FVector BlendParameters = FVector(
-						WrapInputHorizontalBlendNum > 1 ? 
-							HorizontalBlendParameter.Min + (HorizontalBlendParameter.Max - HorizontalBlendParameter.Min) * ((float)HorizontalBlendIndex) / (WrapInputHorizontalBlendNum - 1) :
-							HorizontalBlendParameter.Min,
+				for (int32 VerticalBlendIndex = 0; VerticalBlendIndex < VerticalBlendNum; VerticalBlendIndex++)
+				{
+					const FVector BlendParameters = FVector(
+						WrapInputHorizontalBlendNum > 1 ?
+						HorizontalBlendParameter.Min + (HorizontalBlendParameter.Max - HorizontalBlendParameter.Min) * ((float)HorizontalBlendIndex) / (WrapInputHorizontalBlendNum - 1) :
+						HorizontalBlendParameter.Min,
 						WrapInputVerticalBlendNum > 1 ?
-							VerticalBlendParameter.Min + (VerticalBlendParameter.Max - VerticalBlendParameter.Min) * ((float)VerticalBlendIndex) / (WrapInputVerticalBlendNum - 1) :
-							VerticalBlendParameter.Min,
+						VerticalBlendParameter.Min + (VerticalBlendParameter.Max - VerticalBlendParameter.Min) * ((float)VerticalBlendIndex) / (WrapInputVerticalBlendNum - 1) :
+						VerticalBlendParameter.Min,
 						0.f);
 
-				ProcessSamplingParameter(BlendParameters);
+					ProcessSamplingParameter(BlendParameters);
+				}
 			}
 		}
 	}
@@ -240,18 +238,18 @@ static void IterateOverBlendSpaceSamplingParameter(const UBlendSpace* BlendSpace
 
 static bool IsBlendSpaceRootMotionEnabled(const UBlendSpace* BlendSpace)
 {
-	check(BlendSpace);
-
 	bool bIsRootMotionUsedInBlendSpace = false;
-	BlendSpace->ForEachImmutableSample([&bIsRootMotionUsedInBlendSpace](const FBlendSample& Sample)
-		{
-			const UAnimSequence* Sequence = Sample.Animation.Get();
-			if (IsValid(Sequence) && Sequence->HasRootMotion())
+	if (ensure(BlendSpace))
+	{
+		BlendSpace->ForEachImmutableSample([&bIsRootMotionUsedInBlendSpace](const FBlendSample& Sample)
 			{
-				bIsRootMotionUsedInBlendSpace = true;
-			}
-		});
-
+				const UAnimSequence* Sequence = Sample.Animation.Get();
+				if (IsValid(Sequence) && Sequence->HasRootMotion())
+				{
+					bIsRootMotionUsedInBlendSpace = true;
+				}
+			});
+	}
 	return bIsRootMotionUsedInBlendSpace;
 }
 
@@ -325,13 +323,13 @@ const FString FPoseSearchDatabaseAnimationAssetBase::GetName() const
 
 UAnimationAsset* FPoseSearchDatabaseAnimationAssetBase::GetAnimationAssetForRole(const UE::PoseSearch::FRole& Role) const
 {
-	check(GetNumRoles() == 1);
+	ensure(GetNumRoles() == 1);
 	return Cast<UAnimationAsset>(GetAnimationAsset());
 }
 
 FTransform FPoseSearchDatabaseAnimationAssetBase::GetRootTransformOriginForRole(const UE::PoseSearch::FRole& Role) const
 {
-	check(GetNumRoles() == 1);
+	ensure(GetNumRoles() == 1);
 	return FTransform::Identity;
 }
 
@@ -550,7 +548,7 @@ void FPoseSearchDatabaseAnimationAsset::IterateOverSamplingParameter(const TFunc
 				{
 					// @todo: right now we just sample using the first blend space, but we should probably make sure if there're multiple MultiAnimAssetBlendSpace they are consistent with each other
 					MultiAnimAssetBlendSpace = Cast<UBlendSpace>(AnimationAsset);
-					check(MultiAnimAssetBlendSpace);
+					ensure(MultiAnimAssetBlendSpace);
 					break;
 				}
 			}
@@ -802,7 +800,7 @@ void FPoseSearchDatabaseMultiAnimAsset::IterateOverSamplingParameter(const TFunc
 				{
 					// @todo: right now we just sample using the first blend space, but we should probably make sure if there're multiple blendspaces they are consistent with each other
 					BlendSpace = Cast<UBlendSpace>(AnimationAsset);
-					check(BlendSpace);
+					ensure(BlendSpace);
 					break;
 				}
 			}
@@ -831,7 +829,7 @@ PRAGMA_ENABLE_DEPRECATION_WARNINGS
 
 void UPoseSearchDatabase::SetSearchIndex(const UE::PoseSearch::FSearchIndex& SearchIndex)
 {
-	check(IsInGameThread());
+	ensure(IsInGameThread());
 	SearchIndexPrivate = SearchIndex;
 
 	UpdateCachedProperties();
@@ -846,7 +844,7 @@ void UPoseSearchDatabase::UpdateCachedProperties()
 	{
 		const FSearchIndexAsset& SearchIndexAsset = SearchIndexPrivate.Assets[AssetIdx];
 
-		if (const FPoseSearchDatabaseAnimationAssetBase* DatabaseAnimationAssetBase = GetDatabaseAnimationAsset(SearchIndexAsset))
+		if (const FPoseSearchDatabaseAnimationAssetBase* DatabaseAnimationAssetBase = GetDatabaseAnimationAsset(SearchIndexAsset.GetSourceAssetIdx()))
 		{
 			CachedAssetMap.FindOrAdd(DatabaseAnimationAssetBase->GetAnimationAsset()).Add(AssetIdx);
 		}
@@ -870,36 +868,37 @@ TConstArrayView<int32> UPoseSearchDatabase::GetAssetIndexesForSourceAsset(const
 TConstArrayView<float> UPoseSearchDatabase::CalculateDynamicWeightsSqrt(TArrayView<float> DynamicWeightsSqrtBuffer) const
 {
 	using namespace UE::PoseSearch;
-	check(IsAligned(DynamicWeightsSqrtBuffer.GetData(), alignof(VectorRegister4Float)));
-
 	const UE::PoseSearch::FSearchIndex& SearchIndex = GetSearchIndex();
-	check(DynamicWeightsSqrtBuffer.Num() == SearchIndex.WeightsSqrt.Num());
-
-	bool bInitialized = false;
-	Schema->IterateChannels([&bInitialized, &SearchIndex, &DynamicWeightsSqrtBuffer](const UPoseSearchFeatureChannel* Channel)
+	if (ensure(IsAligned(DynamicWeightsSqrtBuffer.GetData(), alignof(VectorRegister4Float))))
+	{
+		if (ensure(DynamicWeightsSqrtBuffer.Num() == SearchIndex.WeightsSqrt.Num()))
 		{
-			if (const UPoseSearchFeatureChannel_GroupBase* ChannelGroupBase = Cast<UPoseSearchFeatureChannel_GroupBase>(Channel))
-			{
-				if (ChannelGroupBase->DebugWeightGroupID != INDEX_NONE && ChannelGroupBase->DebugWeightGroupID != GVarMotionMatchDebugWeightGroupID)
+			bool bInitialized = false;
+			Schema->IterateChannels([&bInitialized, &SearchIndex, &DynamicWeightsSqrtBuffer](const UPoseSearchFeatureChannel* Channel)
 				{
-					if (!bInitialized)
+					if (const UPoseSearchFeatureChannel_GroupBase* ChannelGroupBase = Cast<UPoseSearchFeatureChannel_GroupBase>(Channel))
 					{
-						// initializing DynamicWeightsSqrtBuffer with the SearchIndex.WeightsSqrt
-						FMemory::Memcpy(DynamicWeightsSqrtBuffer.GetData(), SearchIndex.WeightsSqrt.GetData(), SearchIndex.WeightsSqrt.Num() * sizeof(float));
-						bInitialized = true;
+						if (ChannelGroupBase->DebugWeightGroupID != INDEX_NONE && ChannelGroupBase->DebugWeightGroupID != GVarMotionMatchDebugWeightGroupID)
+						{
+							if (!bInitialized)
+							{
+								// initializing DynamicWeightsSqrtBuffer with the SearchIndex.WeightsSqrt
+								FMemory::Memcpy(DynamicWeightsSqrtBuffer.GetData(), SearchIndex.WeightsSqrt.GetData(), SearchIndex.WeightsSqrt.Num() * sizeof(float));
+								bInitialized = true;
+							}
+
+							// zeroing out interval of weights in DynamicWeightsSqrtBuffer associated with ChannelGroupBase since its GroupID mismatch ValidWeightChannelGroup
+							FMemory::Memzero(DynamicWeightsSqrtBuffer.GetData() + ChannelGroupBase->GetChannelDataOffset(), ChannelGroupBase->GetChannelCardinality() * sizeof(float));
+						}
 					}
+				});
 
-					// zeroing out interval of weights in DynamicWeightsSqrtBuffer associated with ChannelGroupBase since its GroupID mismatch ValidWeightChannelGroup
-					FMemory::Memzero(DynamicWeightsSqrtBuffer.GetData() + ChannelGroupBase->GetChannelDataOffset(), ChannelGroupBase->GetChannelCardinality() * sizeof(float));
-				}
+			if (bInitialized)
+			{
+				return DynamicWeightsSqrtBuffer;
 			}
-		});
-
-	if (bInitialized)
-	{
-		return DynamicWeightsSqrtBuffer;
+		}
 	}
-
 	return SearchIndex.WeightsSqrt;
 }
 
@@ -918,7 +917,7 @@ void UPoseSearchDatabase::AppendToClassSchema(FAppendToClassSchemaContext& Conte
 const UE::PoseSearch::FSearchIndex& UPoseSearchDatabase::GetSearchIndex() const
 {
 	// making sure the search index is consistent. if it fails the calling code hasn't been protected by FAsyncPoseSearchDatabasesManagement::RequestAsyncBuildIndex
-	check(Schema && !SearchIndexPrivate.IsEmpty() && SearchIndexPrivate.GetNumDimensions() == Schema->SchemaCardinality);
+	ensure(Schema && !SearchIndexPrivate.IsEmpty() && SearchIndexPrivate.GetNumDimensions() == Schema->SchemaCardinality);
 	return SearchIndexPrivate;
 }
 
@@ -948,10 +947,9 @@ int32 UPoseSearchDatabase::GetPoseIndex(const UObject* AnimationAsset, float Ani
 
 					if (BlendParametersSquaredLength <= MinSquaredLength)
 					{
-						const FPoseSearchDatabaseAnimationAssetBase* DatabaseAnimationAssetBase = GetDatabaseAnimationAsset(SearchIndexAsset);
+						const FPoseSearchDatabaseAnimationAssetBase* DatabaseAnimationAssetBase = GetDatabaseAnimationAsset(SearchIndexAsset.GetSourceAssetIdx());
 
-						check(DatabaseAnimationAssetBase);
-						check(DatabaseAnimationAssetBase->GetAnimationAsset() == AnimationAsset);
+						ensure(DatabaseAnimationAssetBase && DatabaseAnimationAssetBase->GetAnimationAsset() == AnimationAsset);
 
 						const float RealAssetTime = AnimationAssetTime * SearchIndexAsset.GetToRealTimeFactor();
 						const int32 CurrentIdx = SearchIndexAsset.GetPoseIndexFromTime(RealAssetTime, SampleRate);
@@ -969,54 +967,133 @@ int32 UPoseSearchDatabase::GetPoseIndex(const UObject* AnimationAsset, float Ani
 	return PoseIdx;
 }
 
+#if WITH_EDITOR
 void UPoseSearchDatabase::AddAnimationAsset(FInstancedStruct AnimationAsset)
 {
 	using namespace UE::PoseSearch;
 	PRAGMA_DISABLE_DEPRECATION_WARNINGS
-	DatabaseAnimationAssets.Add(UpgradeToDatabaseAnimationAsset(AnimationAsset));
+	AddAnimationAsset(UpgradeToDatabaseAnimationAsset(AnimationAsset));
 	PRAGMA_ENABLE_DEPRECATION_WARNINGS
 }
 
 void UPoseSearchDatabase::AddAnimationAsset(const FPoseSearchDatabaseAnimationAsset& AnimationAsset)
 {
-	DatabaseAnimationAssets.Add(AnimationAsset);
+	InsertAnimationAssetAt(AnimationAsset, GetNumAnimationAssets());
 }
 
 void UPoseSearchDatabase::InsertAnimationAssetAt(const FPoseSearchDatabaseAnimationAsset& AnimationAsset, int32 AnimationAssetIndex)
 {
+	if (DatabaseAnimationAssets.IsEmpty())
+	{
+		if (UChooserTable* Chooser = Cast<UChooserTable>(GetOuter()))
+		{
+			return FPoseSearchColumn::InsertAnimationAssetAt(AnimationAsset, AnimationAssetIndex, Chooser, this);
+		}
+	}
+
+	ensure(AnimationAssetIndex >= 0 && AnimationAssetIndex <= DatabaseAnimationAssets.Num());
 	DatabaseAnimationAssets.Insert(AnimationAsset, AnimationAssetIndex);
 }
 
 void UPoseSearchDatabase::RemoveAnimationAssetAt(int32 AnimationAssetIndex)
 {
+	if (DatabaseAnimationAssets.IsEmpty())
+	{
+		if (UChooserTable* Chooser = Cast<UChooserTable>(GetOuter()))
+		{
+			FPoseSearchColumn::RemoveAnimationAssetAt(AnimationAssetIndex, Chooser, this);
+			return;
+		}
+	}
+
+	ensure(DatabaseAnimationAssets.IsValidIndex(AnimationAssetIndex));
 	DatabaseAnimationAssets.RemoveAt(AnimationAssetIndex);
 }
 
-const FPoseSearchDatabaseAnimationAsset* UPoseSearchDatabase::GetDatabaseAnimationAsset(int32 AnimationAssetIndex) const
+void UPoseSearchDatabase::SetAnimationAssetAt(const FPoseSearchDatabaseAnimationAsset& AnimationAsset, int32 AnimationAssetIndex)
 {
-	PRAGMA_DISABLE_DEPRECATION_WARNINGS
-	return GetDatabaseAnimationAsset<FPoseSearchDatabaseAnimationAsset>(AnimationAssetIndex);
-	PRAGMA_ENABLE_DEPRECATION_WARNINGS
+	if (DatabaseAnimationAssets.IsEmpty())
+	{
+		if (UChooserTable* Chooser = Cast<UChooserTable>(GetOuter()))
+		{
+			return FPoseSearchColumn::SetAnimationAssetAt(AnimationAsset, AnimationAssetIndex, Chooser, this);
+		}
+	}
+
+	if (DatabaseAnimationAssets.IsValidIndex(AnimationAssetIndex))
+	{
+		DatabaseAnimationAssets[AnimationAssetIndex] = AnimationAsset;
+	}
 }
+#endif // WITH_EDITOR
 
-const FPoseSearchDatabaseAnimationAsset* UPoseSearchDatabase::GetDatabaseAnimationAsset(const UE::PoseSearch::FSearchIndexAsset& SearchIndexAsset) const
+const FPoseSearchDatabaseAnimationAsset* UPoseSearchDatabase::GetDatabaseAnimationAsset(int32 AnimationAssetIndex) const
 {
-	PRAGMA_DISABLE_DEPRECATION_WARNINGS
-	return GetDatabaseAnimationAsset<FPoseSearchDatabaseAnimationAsset>(SearchIndexAsset);
-	PRAGMA_ENABLE_DEPRECATION_WARNINGS
+	if (DatabaseAnimationAssets.IsEmpty())
+	{
+		if (UChooserTable* Chooser = Cast<UChooserTable>(GetOuter()))
+		{
+			return FPoseSearchColumn::GetDatabaseAnimationAsset(AnimationAssetIndex, Chooser, this);
+		}
+	}
+
+	if (DatabaseAnimationAssets.IsValidIndex(AnimationAssetIndex))
+	{
+		return &DatabaseAnimationAssets[AnimationAssetIndex];
+	}
+	return nullptr;
 }
 
 FPoseSearchDatabaseAnimationAsset* UPoseSearchDatabase::GetMutableDatabaseAnimationAsset(int32 AnimationAssetIndex)
 {
-	PRAGMA_DISABLE_DEPRECATION_WARNINGS
-	return GetMutableDatabaseAnimationAsset<FPoseSearchDatabaseAnimationAsset>(AnimationAssetIndex);
-	PRAGMA_ENABLE_DEPRECATION_WARNINGS
+	if (DatabaseAnimationAssets.IsEmpty())
+	{
+		if (UChooserTable* Chooser = Cast<UChooserTable>(GetOuter()))
+		{
+			// unimplemented
+			ensure(false);
+			return nullptr;
+		}
+	}
+
+	if (DatabaseAnimationAssets.IsValidIndex(AnimationAssetIndex))
+	{
+		return &DatabaseAnimationAssets[AnimationAssetIndex];
+	}
+	return nullptr;
+}
+
+int32 UPoseSearchDatabase::GetNumAnimationAssets() const
+{
+	if (DatabaseAnimationAssets.IsEmpty())
+	{
+		if (UChooserTable* Chooser = Cast<UChooserTable>(GetOuter()))
+		{
+			return FPoseSearchColumn::GetNumAnimationAssets(Chooser, this);
+		}
+	}
+
+	return DatabaseAnimationAssets.Num();
+}
+
+UObject* UPoseSearchDatabase::GetAnimationAsset(int32 Index) const
+{
+	if (const FPoseSearchDatabaseAnimationAsset* DatabaseAnimationAsset = GetDatabaseAnimationAsset(Index))
+	{
+		return DatabaseAnimationAsset->GetAnimationAsset();
+	}
+	return nullptr;
+}
+
+const FPoseSearchDatabaseAnimationAsset* UPoseSearchDatabase::GetDatabaseAnimationAsset(const UE::PoseSearch::FSearchIndexAsset& SearchIndexAsset) const
+{
+	return GetDatabaseAnimationAsset(SearchIndexAsset.GetSourceAssetIdx());
 }
 
 FPoseSearchDatabaseAnimationAsset* UPoseSearchDatabase::GetMutableDatabaseAnimationAsset(const UE::PoseSearch::FSearchIndexAsset& SearchIndexAsset)
 {
 	PRAGMA_DISABLE_DEPRECATION_WARNINGS
-	return GetMutableDatabaseAnimationAsset<FPoseSearchDatabaseAnimationAsset>(SearchIndexAsset);
+	return GetMutableDatabaseAnimationAsset(SearchIndexAsset.GetSourceAssetIdx());
 	PRAGMA_ENABLE_DEPRECATION_WARNINGS
 }
 
@@ -1054,9 +1131,9 @@ void UPoseSearchDatabase::PostLoad()
 	}
 	PRAGMA_ENABLE_DEPRECATION_WARNINGS
 
-#if WITH_EDITOR
-
+	UpdateCachedProperties();
 
+#if WITH_EDITOR
 	// Delay indexing until BeginCacheForCookedPlatformData when running a CookCommandlet
 	if (!IsRunningCookCommandlet())
 	{
@@ -1085,20 +1162,6 @@ bool UPoseSearchDatabase::Contains(const UObject* Object) const
 	return !GetAssetIndexesForSourceAsset(Object).IsEmpty();
 }
 
-int32 UPoseSearchDatabase::GetNumAnimationAssets() const
-{
-	return DatabaseAnimationAssets.Num();
-}
-
-UObject* UPoseSearchDatabase::GetAnimationAsset(int32 Index) const
-{
-	if (DatabaseAnimationAssets.IsValidIndex(Index))
-	{
-		return DatabaseAnimationAssets[Index].GetAnimationAsset();
-	}
-	return nullptr;
-}
-
 #if WITH_EDITOR
 void UPoseSearchDatabase::SynchronizeWithExternalDependencies()
 {
@@ -1270,7 +1333,7 @@ void UPoseSearchDatabase::BeginCacheForCookedPlatformData(const ITargetPlatform*
 bool UPoseSearchDatabase::IsCachedCookedPlatformDataLoaded(const ITargetPlatform* TargetPlatform)
 {
 	using namespace UE::PoseSearch;
-	check(IsInGameThread());
+	ensure(IsInGameThread());
 	return EAsyncBuildIndexResult::InProgress != FAsyncPoseSearchDatabasesManagement::RequestAsyncBuildIndex(this, ERequestAsyncBuildFlag::ContinueRequest);
 }
 #endif // WITH_EDITOR
@@ -1325,27 +1388,34 @@ void UPoseSearchDatabase::Serialize(FArchive& Ar)
 		if (bSavingCooked || bLoadingCooked)
 		{
 			Ar << SearchIndexPrivate;
-			UpdateCachedProperties();
 		}
 	}
 }
 
 float UPoseSearchDatabase::GetRealAssetTime(int32 PoseIdx) const
 {
-	check(Schema);
-	const UE::PoseSearch::FSearchIndexAsset& Asset = GetSearchIndex().GetAssetForPose(PoseIdx);
-	return Asset.GetTimeFromPoseIndex(PoseIdx, Schema->SampleRate);
+	if (ensure(Schema))
+	{
+		const UE::PoseSearch::FSearchIndexAsset& Asset = GetSearchIndex().GetAssetForPose(PoseIdx);
+		return Asset.GetTimeFromPoseIndex(PoseIdx, Schema->SampleRate);
+	}
+	return 0.f;
 }
 
 float UPoseSearchDatabase::GetNormalizedAssetTime(int32 PoseIdx) const
 {
-	check(Schema);
-	const UE::PoseSearch::FSearchIndexAsset& Asset = GetSearchIndex().GetAssetForPose(PoseIdx);
-	const float ToRealTimeFactor = Asset.GetToRealTimeFactor();
-	check(ToRealTimeFactor > UE_KINDA_SMALL_NUMBER);
-	const float RealAssetTime = Asset.GetTimeFromPoseIndex(PoseIdx, Schema->SampleRate);
-	const float NormalizedAssetTime = RealAssetTime / ToRealTimeFactor;
-	return NormalizedAssetTime;
+	if (ensure(Schema))
+	{
+		const UE::PoseSearch::FSearchIndexAsset& Asset = GetSearchIndex().GetAssetForPose(PoseIdx);
+		const float ToRealTimeFactor = Asset.GetToRealTimeFactor();
+		if (ensure(ToRealTimeFactor > UE_KINDA_SMALL_NUMBER))
+		{
+			const float RealAssetTime = Asset.GetTimeFromPoseIndex(PoseIdx, Schema->SampleRate);
+			const float NormalizedAssetTime = RealAssetTime / ToRealTimeFactor;
+			return NormalizedAssetTime;
+		}
+	}
+	return 0.f;
 }
 
 void UPoseSearchDatabase::Search(UE::PoseSearch::FSearchContext& SearchContext, UE::PoseSearch::FSearchResults& SearchResults) const
@@ -1433,60 +1503,63 @@ void UPoseSearchDatabase::PopulateNonSelectableIdx(FNonSelectableIdx& NonSelecta
 			{
 				// excluding all the poses with CurrentIndexAsset->GetSourceAssetIdx()
 				const FPoseSearchDatabaseAnimationAssetBase* DatabaseAnimationAssetBase = GetDatabaseAnimationAsset(CurrentIndexAsset->GetSourceAssetIdx());
-				check(DatabaseAnimationAssetBase);
-
-				for (int32 AssetIndex : GetAssetIndexesForSourceAsset(DatabaseAnimationAssetBase->GetAnimationAsset()))
+				if (ensure(DatabaseAnimationAssetBase))
 				{
-					const FSearchIndexAsset& SearchIndexAsset = SearchIndex.Assets[AssetIndex];
-					const int32 FirstPoseIdx = SearchIndexAsset.GetFirstPoseIdx();
-					const int32 LastPoseIdx = FirstPoseIdx + SearchIndexAsset.GetNumPoses();
-					for (int32 PoseIdx = FirstPoseIdx; PoseIdx < LastPoseIdx; ++PoseIdx)
+					for (int32 AssetIndex : GetAssetIndexesForSourceAsset(DatabaseAnimationAssetBase->GetAnimationAsset()))
 					{
-						// no need to AddUnique since there's no overlapping between pose indexes in the FSearchIndexAsset(s)
-						NonSelectableIdx.Add(PoseIdx);
+						const FSearchIndexAsset& SearchIndexAsset = SearchIndex.Assets[AssetIndex];
+						const int32 FirstPoseIdx = SearchIndexAsset.GetFirstPoseIdx();
+						const int32 LastPoseIdx = FirstPoseIdx + SearchIndexAsset.GetNumPoses();
+						for (int32 PoseIdx = FirstPoseIdx; PoseIdx < LastPoseIdx; ++PoseIdx)
+						{
+							// no need to AddUnique since there's no overlapping between pose indexes in the FSearchIndexAsset(s)
+							NonSelectableIdx.Add(PoseIdx);
 
 #if UE_POSE_SEARCH_TRACE_ENABLED
-						SearchContext.Track(this, PoseIdx, EPoseCandidateFlags::DiscardedBy_AssetReselection, DynamicWeightsSqrt, ContinuingPoseCostAddend, ContinuingInteractionCostAddend, ContinuingContextInteractionCostAddend);
+							SearchContext.Track(this, PoseIdx, EPoseCandidateFlags::DiscardedBy_AssetReselection, DynamicWeightsSqrt, ContinuingPoseCostAddend, ContinuingInteractionCostAddend, ContinuingContextInteractionCostAddend);
 #endif // UE_POSE_SEARCH_TRACE_ENABLED
+						}
 					}
 				}
 			}
 			else if (!FMath::IsNearlyEqual(SearchContext.GetPoseJumpThresholdTime().Min, SearchContext.GetPoseJumpThresholdTime().Max))
 			{
 				const int32 ContinuingPoseIdx = SearchContext.GetContinuingPoseSearchResult().PoseIdx;
-				check(ContinuingPoseIdx != INDEX_NONE);
-				const int32 UnboundMinPoseIdx = ContinuingPoseIdx + FMath::FloorToInt(SearchContext.GetPoseJumpThresholdTime().Min * Schema->SampleRate);
-				const int32 UnboundMaxPoseIdx = ContinuingPoseIdx + FMath::CeilToInt(SearchContext.GetPoseJumpThresholdTime().Max * Schema->SampleRate);
-				const int32 CurrentIndexAssetFirstPoseIdx = CurrentIndexAsset->GetFirstPoseIdx();
-				const int32 CurrentIndexAssetNumPoses = CurrentIndexAsset->GetNumPoses();
-				const bool bIsLooping = CurrentIndexAsset->IsLooping();
-
-				if (bIsLooping)
+				if (ensure(ContinuingPoseIdx != INDEX_NONE))
 				{
-					for (int32 UnboundPoseIdx = UnboundMinPoseIdx; UnboundPoseIdx < UnboundMaxPoseIdx; ++UnboundPoseIdx)
+					const int32 UnboundMinPoseIdx = ContinuingPoseIdx + FMath::FloorToInt(SearchContext.GetPoseJumpThresholdTime().Min * Schema->SampleRate);
+					const int32 UnboundMaxPoseIdx = ContinuingPoseIdx + FMath::CeilToInt(SearchContext.GetPoseJumpThresholdTime().Max * Schema->SampleRate);
+					const int32 CurrentIndexAssetFirstPoseIdx = CurrentIndexAsset->GetFirstPoseIdx();
+					const int32 CurrentIndexAssetNumPoses = CurrentIndexAsset->GetNumPoses();
+					const bool bIsLooping = CurrentIndexAsset->IsLooping();
+
+					if (bIsLooping)
 					{
-						const int32 Modulo = (UnboundPoseIdx - CurrentIndexAssetFirstPoseIdx) % CurrentIndexAssetNumPoses;
-						const int32 CurrentIndexAssetFirstPoseIdxPlusModulo = CurrentIndexAssetFirstPoseIdx + Modulo;
-						const int32 PoseIdx = Modulo >= 0 ? CurrentIndexAssetFirstPoseIdxPlusModulo : CurrentIndexAssetFirstPoseIdxPlusModulo + CurrentIndexAssetNumPoses;
+						for (int32 UnboundPoseIdx = UnboundMinPoseIdx; UnboundPoseIdx < UnboundMaxPoseIdx; ++UnboundPoseIdx)
+						{
+							const int32 Modulo = (UnboundPoseIdx - CurrentIndexAssetFirstPoseIdx) % CurrentIndexAssetNumPoses;
+							const int32 CurrentIndexAssetFirstPoseIdxPlusModulo = CurrentIndexAssetFirstPoseIdx + Modulo;
+							const int32 PoseIdx = Modulo >= 0 ? CurrentIndexAssetFirstPoseIdxPlusModulo : CurrentIndexAssetFirstPoseIdxPlusModulo + CurrentIndexAssetNumPoses;
 
-						NonSelectableIdx.AddUnique(PoseIdx);
+							NonSelectableIdx.AddUnique(PoseIdx);
 
 #if UE_POSE_SEARCH_TRACE_ENABLED
-						SearchContext.Track(this, PoseIdx, EPoseCandidateFlags::DiscardedBy_PoseJumpThresholdTime, DynamicWeightsSqrt, ContinuingPoseCostAddend, ContinuingInteractionCostAddend, ContinuingContextInteractionCostAddend);
+							SearchContext.Track(this, PoseIdx, EPoseCandidateFlags::DiscardedBy_PoseJumpThresholdTime, DynamicWeightsSqrt, ContinuingPoseCostAddend, ContinuingInteractionCostAddend, ContinuingContextInteractionCostAddend);
 #endif // UE_POSE_SEARCH_TRACE_ENABLED
+						}
 					}
-				}
-				else
-				{
-					const int32 MinPoseIdx = FMath::Max(CurrentIndexAssetFirstPoseIdx, UnboundMinPoseIdx);
-					const int32 MaxPoseIdx = FMath::Min(CurrentIndexAssetFirstPoseIdx + CurrentIndexAssetNumPoses, UnboundMaxPoseIdx);
-
-					for (int32 PoseIdx = MinPoseIdx; PoseIdx < MaxPoseIdx; ++PoseIdx)
+					else
 					{
-						NonSelectableIdx.AddUnique(PoseIdx);
+						const int32 MinPoseIdx = FMath::Max(CurrentIndexAssetFirstPoseIdx, UnboundMinPoseIdx);
+						const int32 MaxPoseIdx = FMath::Min(CurrentIndexAssetFirstPoseIdx + CurrentIndexAssetNumPoses, UnboundMaxPoseIdx);
+
+						for (int32 PoseIdx = MinPoseIdx; PoseIdx < MaxPoseIdx; ++PoseIdx)
+						{
+							NonSelectableIdx.AddUnique(PoseIdx);
 #if UE_POSE_SEARCH_TRACE_ENABLED
-						SearchContext.Track(this, PoseIdx, EPoseCandidateFlags::DiscardedBy_PoseJumpThresholdTime, DynamicWeightsSqrt, ContinuingPoseCostAddend, ContinuingInteractionCostAddend, ContinuingContextInteractionCostAddend);
+							SearchContext.Track(this, PoseIdx, EPoseCandidateFlags::DiscardedBy_PoseJumpThresholdTime, DynamicWeightsSqrt, ContinuingPoseCostAddend, ContinuingInteractionCostAddend, ContinuingContextInteractionCostAddend);
 #endif // UE_POSE_SEARCH_TRACE_ENABLED
+						}
 					}
 				}
 			}
@@ -1495,8 +1568,6 @@ void UPoseSearchDatabase::PopulateNonSelectableIdx(FNonSelectableIdx& NonSelecta
 
 	for (const IPoseHistory* PoseHistory : SearchContext.GetPoseHistories())
 	{
-		check(PoseHistory);
-		
 		if (const FPoseIndicesHistory* PoseIndicesHistory = PoseHistory->GetPoseIndicesHistory())
 		{
 			const FObjectKey DatabaseKey(this);
@@ -1546,27 +1617,28 @@ void UPoseSearchDatabase::PopulateSelectableAssetIdx(FSelectableAssetIdx& Select
 
 void UPoseSearchDatabase::PopulateSelectableAssetIdx(FSelectableAssetIdx& SelectableAssetIdx, const UE::PoseSearch::FStackAssetSet* AssetsToConsider) const
 {
-	check(AssetsToConsider);
-
-	SelectableAssetIdx.Reset();
-	if (!AssetsToConsider->IsEmpty())
+	if (ensure(AssetsToConsider))
 	{
-		for (const UObject* AssetToConsider : *AssetsToConsider)
+		SelectableAssetIdx.Reset();
+		if (!AssetsToConsider->IsEmpty())
 		{
-			SelectableAssetIdx.Append(GetAssetIndexesForSourceAsset(AssetToConsider));
-		}
-
-		if (!SelectableAssetIdx.IsEmpty())
-		{
-			if (SelectableAssetIdx.Num() != GetSearchIndex().Assets.Num())
+			for (const UObject* AssetToConsider : *AssetsToConsider)
 			{
-				SelectableAssetIdx.Sort();
+				SelectableAssetIdx.Append(GetAssetIndexesForSourceAsset(AssetToConsider));
 			}
-			else
+
+			if (!SelectableAssetIdx.IsEmpty())
 			{
-				// SelectableAssetIdx contains ALL the Database->GetSearchIndex().Assets. 
-				// We reset SelectableAssetIdx since it has the same meaning, and it'll perform better
-				SelectableAssetIdx.Reset();
+				if (SelectableAssetIdx.Num() != GetSearchIndex().Assets.Num())
+				{
+					SelectableAssetIdx.Sort();
+				}
+				else
+				{
+					// SelectableAssetIdx contains ALL the Database->GetSearchIndex().Assets. 
+					// We reset SelectableAssetIdx since it has the same meaning, and it'll perform better
+					SelectableAssetIdx.Reset();
+				}
 			}
 		}
 	}
@@ -1576,33 +1648,34 @@ void UPoseSearchDatabase::PopulateSelectableAssetIdx(FSelectableAssetIdx& Select
 {
 	using namespace UE::PoseSearch;
 	
-	check(AssetIndexesToConsiderPerDatabase);
-
-	const int32 NumAnimationAssets = GetNumAnimationAssets();
-
-	SelectableAssetIdx.Reset();
-	if (const FStackAssetIndexes* AssetIndexesToConsider = AssetIndexesToConsiderPerDatabase->Find(this))
+	if (ensure(AssetIndexesToConsiderPerDatabase))
 	{
-		check(AssetIndexesToConsider->Num() <= NumAnimationAssets);
+		const int32 NumAnimationAssets = GetNumAnimationAssets();
 
-		// if AssetIndexesToConsider->Num() == NumAnimationAssets we want to search the full database and and empty SelectableAssetIdx does that
-		if (AssetIndexesToConsider->Num() != NumAnimationAssets)
+		SelectableAssetIdx.Reset();
+		if (const FStackAssetIndexes* AssetIndexesToConsider = AssetIndexesToConsiderPerDatabase->Find(this))
 		{
-			for (int32 AssetIndexToConsider : *AssetIndexesToConsider)
+			ensure(AssetIndexesToConsider->Num() <= NumAnimationAssets);
+
+			// if AssetIndexesToConsider->Num() == NumAnimationAssets we want to search the full database and and empty SelectableAssetIdx does that
+			if (AssetIndexesToConsider->Num() != NumAnimationAssets)
 			{
-				if (const TArray<int32>* IndexesForSourceAsset = CachedAssetMap.Find(DatabaseAnimationAssets[AssetIndexToConsider].AnimAsset))
+				for (int32 AssetIndexToConsider : *AssetIndexesToConsider)
 				{
-					// @todo: add a mapping like CachedAssetMap for asset indexes to SearchIndexPrivate.Assets indexes to avoid this cumbersome filtering
-					for (int32 IndexForSourceAsset : *IndexesForSourceAsset)
+					if (const TArray<int32>* IndexesForSourceAsset = CachedAssetMap.Find(GetAnimationAsset(AssetIndexToConsider)))
 					{
-						if (SearchIndexPrivate.Assets[IndexForSourceAsset].GetSourceAssetIdx() == AssetIndexToConsider)
+						// @todo: add a mapping like CachedAssetMap for asset indexes to SearchIndexPrivate.Assets indexes to avoid this cumbersome filtering
+						for (int32 IndexForSourceAsset : *IndexesForSourceAsset)
 						{
-							SelectableAssetIdx.Add(IndexForSourceAsset);
+							if (SearchIndexPrivate.Assets[IndexForSourceAsset].GetSourceAssetIdx() == AssetIndexToConsider)
+							{
+								SelectableAssetIdx.Add(IndexForSourceAsset);
+							}
 						}
 					}
 				}
+				SelectableAssetIdx.Sort();
 			}
-			SelectableAssetIdx.Sort();
 		}
 	}
 }
@@ -1669,8 +1742,11 @@ void UPoseSearchDatabase::ExtractCostAddends(const UE::PoseSearch::FSearchContex
 	// extracting notifies from the database animation asset at time SampleTime
 	const FSearchIndex& SearchIndex = GetSearchIndex();
 	const FSearchIndexAsset& SearchIndexAsset = SearchIndex.GetAssetForPose(ContinuingPoseIdx);
-	const FPoseSearchDatabaseAnimationAssetBase* DatabaseAnimationAssetBase = GetDatabaseAnimationAsset(SearchIndexAsset);
-	check(DatabaseAnimationAssetBase);
+	const FPoseSearchDatabaseAnimationAssetBase* DatabaseAnimationAssetBase = GetDatabaseAnimationAsset(SearchIndexAsset.GetSourceAssetIdx());
+	if (!ensure(DatabaseAnimationAssetBase))
+	{
+		return;
+	}
 
 	const float SampleTime = GetRealAssetTime(ContinuingPoseIdx);
 	FAnimNotifyContext PreAllocatedNotifyContext;
@@ -1746,7 +1822,7 @@ void UPoseSearchDatabase::SearchContinuingPose(UE::PoseSearch::FSearchContext& S
 
 	using namespace UE::PoseSearch;
 
-	check(SearchContext.GetContinuingPoseSearchResult().Database.Get() == this);
+	ensure(SearchContext.GetContinuingPoseSearchResult().Database.Get() == this);
 
 #if WITH_EDITOR
 	if (EAsyncBuildIndexResult::Success != FAsyncPoseSearchDatabasesManagement::RequestAsyncBuildIndex(this, ERequestAsyncBuildFlag::ContinueRequest))
@@ -1765,56 +1841,58 @@ void UPoseSearchDatabase::SearchContinuingPose(UE::PoseSearch::FSearchContext& S
 	if (SearchResults.ShouldPerformSearch(SearchIndex.MinCostAddend + ContinuingPoseCostAddend + ContinuingInteractionCostAddend + ContinuingContextInteractionCostAddend))
 	{
 		const int32 ContinuingPoseIdx = SearchContext.GetContinuingPoseSearchResult().PoseIdx;
-		check(ContinuingPoseIdx != INDEX_NONE);
-		const int32 NumDimensions = Schema->SchemaCardinality;
-		// FMemory_Alloca is forced 16 bytes aligned
-		TArrayView<float> ReconstructedPoseValuesBuffer((float*)FMemory_Alloca(NumDimensions * sizeof(float)), NumDimensions);
-		check(IsAligned(ReconstructedPoseValuesBuffer.GetData(), alignof(VectorRegister4Float)));
-		const TConstArrayView<float> PoseValues = SearchIndex.IsValuesEmpty() ? SearchIndex.GetReconstructedPoseValues(ContinuingPoseIdx, ReconstructedPoseValuesBuffer) : SearchIndex.GetPoseValues(ContinuingPoseIdx);
+		if (ensure(ContinuingPoseIdx != INDEX_NONE))
+		{
+			const int32 NumDimensions = Schema->SchemaCardinality;
+			// FMemory_Alloca is forced 16 bytes aligned
+			TArrayView<float> ReconstructedPoseValuesBuffer((float*)FMemory_Alloca(NumDimensions * sizeof(float)), NumDimensions);
+			ensure(IsAligned(ReconstructedPoseValuesBuffer.GetData(), alignof(VectorRegister4Float)));
+			const TConstArrayView<float> PoseValues = SearchIndex.IsValuesEmpty() ? SearchIndex.GetReconstructedPoseValues(ContinuingPoseIdx, ReconstructedPoseValuesBuffer) : SearchIndex.GetPoseValues(ContinuingPoseIdx);
 
-		// @todo: perhaps store the DynamicWeightsSqrt into the GetOrBuildQuery to share between the continuing pose and the full search
-		TArrayView<float> DynamicWeightsSqrtBuffer((float*)FMemory_Alloca(NumDimensions * sizeof(float)), NumDimensions);
-		const TConstArrayView<float> DynamicWeightsSqrt = CalculateDynamicWeightsSqrt(DynamicWeightsSqrtBuffer);
+			// @todo: perhaps store the DynamicWeightsSqrt into the GetOrBuildQuery to share between the continuing pose and the full search
+			TArrayView<float> DynamicWeightsSqrtBuffer((float*)FMemory_Alloca(NumDimensions * sizeof(float)), NumDimensions);
+			const TConstArrayView<float> DynamicWeightsSqrt = CalculateDynamicWeightsSqrt(DynamicWeightsSqrtBuffer);
 
-		FSearchResult Result;
-		Result.bIsContinuingPoseSearch = true;
+			FSearchResult Result;
+			Result.bIsContinuingPoseSearch = true;
 
-		// is the data padded at 16 bytes (and 16 bytes aligned by construction)?
-		if (NumDimensions % 4 == 0)
-		{
-			Result.PoseCost = FPoseSearchCost(CompareFeatureVectors<true>(PoseValues, SearchContext.GetOrBuildQuery(Schema), DynamicWeightsSqrt), SearchIndex.PoseMetadata[ContinuingPoseIdx].GetCostAddend(), ContinuingPoseCostAddend, ContinuingInteractionCostAddend, ContinuingContextInteractionCostAddend);
-		}
-		// data is not 16 bytes padded
-		else
-		{
-			Result.PoseCost = FPoseSearchCost(CompareFeatureVectors<false>(PoseValues, SearchContext.GetOrBuildQuery(Schema), DynamicWeightsSqrt), SearchIndex.PoseMetadata[ContinuingPoseIdx].GetCostAddend(), ContinuingPoseCostAddend, ContinuingInteractionCostAddend, ContinuingContextInteractionCostAddend);
-		}
+			// is the data padded at 16 bytes (and 16 bytes aligned by construction)?
+			if (NumDimensions % 4 == 0)
+			{
+				Result.PoseCost = FPoseSearchCost(CompareFeatureVectors<true>(PoseValues, SearchContext.GetOrBuildQuery(Schema), DynamicWeightsSqrt), SearchIndex.PoseMetadata[ContinuingPoseIdx].GetCostAddend(), ContinuingPoseCostAddend, ContinuingInteractionCostAddend, ContinuingContextInteractionCostAddend);
+			}
+			// data is not 16 bytes padded
+			else
+			{
+				Result.PoseCost = FPoseSearchCost(CompareFeatureVectors<false>(PoseValues, SearchContext.GetOrBuildQuery(Schema), DynamicWeightsSqrt), SearchIndex.PoseMetadata[ContinuingPoseIdx].GetCostAddend(), ContinuingPoseCostAddend, ContinuingInteractionCostAddend, ContinuingContextInteractionCostAddend);
+			}
 
-		Result.SetAssetTime(SearchContext.GetContinuingPoseSearchResult().GetAssetTime());
-		Result.PoseIdx = ContinuingPoseIdx;
-		Result.Database = this;
+			Result.SetAssetTime(SearchContext.GetContinuingPoseSearchResult().GetAssetTime());
+			Result.PoseIdx = ContinuingPoseIdx;
+			Result.Database = this;
 
 #if UE_POSE_SEARCH_TRACE_ENABLED
-		SearchContext.Track(this, ContinuingPoseIdx, EPoseCandidateFlags::Valid_ContinuingPose, Result.PoseCost);
+			SearchContext.Track(this, ContinuingPoseIdx, EPoseCandidateFlags::Valid_ContinuingPose, Result.PoseCost);
 #endif // UE_POSE_SEARCH_TRACE_ENABLED
 
-		check(Result.DebugValidate());
-		SearchResults.UpdateWith(Result);
+			ensure(Result.DebugValidate());
+			SearchResults.UpdateWith(Result);
 
 #if WITH_EDITOR && ENABLE_ANIM_DEBUG && UE_POSE_SEARCH_TRACE_ENABLED
-		if (GVarMotionMatchCompareAgainstBruteForce)
-		{
-			const FPoseSearchCost& BruteForcePoseCost = Result.PoseCost;
-			SearchResults.IterateOverSearchResults([&BruteForcePoseCost](FSearchResult& SearchResult)
-				{
-					if (BruteForcePoseCost < SearchResult.BruteForcePoseCost)
+			if (GVarMotionMatchCompareAgainstBruteForce)
+			{
+				const FPoseSearchCost& BruteForcePoseCost = Result.PoseCost;
+				SearchResults.IterateOverSearchResults([&BruteForcePoseCost](FSearchResult& SearchResult)
 					{
-						SearchResult.BruteForcePoseCost = BruteForcePoseCost;
-					}
-					return false;
-				});
-		}
+						if (BruteForcePoseCost < SearchResult.BruteForcePoseCost)
+						{
+							SearchResult.BruteForcePoseCost = BruteForcePoseCost;
+						}
+						return false;
+					});
+			}
 #endif // WITH_EDITOR && ENABLE_ANIM_DEBUG && UE_POSE_SEARCH_TRACE_ENABLED
+		}
 	}
 	else
 	{
@@ -1851,7 +1929,7 @@ void UPoseSearchDatabase::SearchPCAKDTree(UE::PoseSearch::FSearchContext& Search
 		const TConstArrayView<float> DynamicWeightsSqrt = CalculateDynamicWeightsSqrt(DynamicWeightsSqrtBuffer);
 
 		TConstArrayView<float> QueryValues = SearchContext.GetOrBuildQuery(Schema);
-		check(QueryValues.Num() == NumDimensions);
+		ensure(QueryValues.Num() == NumDimensions);
 
 		FSelectableAssetIdx SelectableAssetIdx;
 		PRAGMA_DISABLE_DEPRECATION_WARNINGS
@@ -1883,7 +1961,7 @@ void UPoseSearchDatabase::SearchPCAKDTree(UE::PoseSearch::FSearchContext& Search
 
 		// projecting QueryValues into the PCA space 
 		TConstArrayView<float> PCAQueryValues = SearchContext.GetOrBuildPCAQuery(this);
-		check(PCAQueryValues.Num() == ClampedNumberOfPrincipalComponents);
+		ensure(PCAQueryValues.Num() == ClampedNumberOfPrincipalComponents);
 
 		int32 NumResults = 0;
 
@@ -2014,7 +2092,7 @@ void UPoseSearchDatabase::SearchPCAKDTree(UE::PoseSearch::FSearchContext& Search
 			{
 				// FMemory_Alloca is forced 16 bytes aligned
 				TArrayView<float> ReconstructedPoseValuesBuffer((float*)FMemory_Alloca(NumDimensions * sizeof(float)), NumDimensions);
-				check(IsAligned(ReconstructedPoseValuesBuffer.GetData(), alignof(VectorRegister4Float)));
+				ensure(IsAligned(ReconstructedPoseValuesBuffer.GetData(), alignof(VectorRegister4Float)));
 				if (NumDimensions % 4 == 0)
 				{
 					int32 NumEvaluatePoseKernelCalls = 0;
@@ -2102,7 +2180,7 @@ void UPoseSearchDatabase::SearchPCAKDTree(UE::PoseSearch::FSearchContext& Search
 		{
 			// FMemory_Alloca is forced 16 bytes aligned
 			TArrayView<float> ReconstructedPoseValuesBuffer((float*)FMemory_Alloca(NumDimensions * sizeof(float)), NumDimensions);
-			check(IsAligned(ReconstructedPoseValuesBuffer.GetData(), alignof(VectorRegister4Float)));
+			ensure(IsAligned(ReconstructedPoseValuesBuffer.GetData(), alignof(VectorRegister4Float)));
 
 			// is the data padded at 16 bytes (and 16 bytes aligned by construction)?
 			if (NumDimensions % 4 == 0)
@@ -2209,7 +2287,7 @@ void UPoseSearchDatabase::SearchVPTree(UE::PoseSearch::FSearchContext& SearchCon
 		const TConstArrayView<float> DynamicWeightsSqrt = CalculateDynamicWeightsSqrt(DynamicWeightsSqrtBuffer);
 
 		TConstArrayView<float> QueryValues = SearchContext.GetOrBuildQuery(Schema);
-		check(QueryValues.Num() == NumDimensions);
+		ensure(QueryValues.Num() == NumDimensions);
 
 		FSelectableAssetIdx SelectableAssetIdx;
 		PRAGMA_DISABLE_DEPRECATION_WARNINGS
@@ -2363,7 +2441,7 @@ void UPoseSearchDatabase::SearchBruteForce(UE::PoseSearch::FSearchContext& Searc
 			{
 				// FMemory_Alloca is forced 16 bytes aligned
 				TArrayView<float> ReconstructedPoseValuesBuffer((float*)FMemory_Alloca(NumDimensions * sizeof(float)), NumDimensions);
-				check(IsAligned(ReconstructedPoseValuesBuffer.GetData(), alignof(VectorRegister4Float)));
+				ensure(IsAligned(ReconstructedPoseValuesBuffer.GetData(), alignof(VectorRegister4Float)));
 				for (int32 PoseIdx = 0; PoseIdx < SearchIndex.GetNumPoses(); ++PoseIdx)
 				{
 					EvaluatePoseKernel<true, false>(SearchResults, SearchIndex, QueryValues, ReconstructedPoseValuesBuffer, PoseIdx, INDEX_NONE, SearchFilters, ContinuingPoseCostAddend, ContinuingInteractionCostAddend, ContinuingContextInteractionCostAddend, this, DynamicWeightsSqrt
@@ -2416,7 +2494,7 @@ void UPoseSearchDatabase::SearchBruteForce(UE::PoseSearch::FSearchContext& Searc
 			{
 				// FMemory_Alloca is forced 16 bytes aligned
 				TArrayView<float> ReconstructedPoseValuesBuffer((float*)FMemory_Alloca(NumDimensions * sizeof(float)), NumDimensions);
-				check(IsAligned(ReconstructedPoseValuesBuffer.GetData(), alignof(VectorRegister4Float)));
+				ensure(IsAligned(ReconstructedPoseValuesBuffer.GetData(), alignof(VectorRegister4Float)));
 
 				for (int32 AssetIdx : SelectableAssetIdx)
 				{
diff --git a/Engine/Plugins/Animation/PoseSearch/Source/Runtime/Private/PoseSearchInteractionIsland.cpp b/Engine/Plugins/Animation/PoseSearch/Source/Runtime/Private/PoseSearchInteractionIsland.cpp
--- a/Engine/Plugins/Animation/PoseSearch/Source/Runtime/Private/PoseSearchInteractionIsland.cpp
+++ b/Engine/Plugins/Animation/PoseSearch/Source/Runtime/Private/PoseSearchInteractionIsland.cpp
@@ -34,8 +34,10 @@ static void ShowPrerequistes(FTickFunction& NestedTick, int32 Indent, TSet<FTick
 
 	if (bAlreadyInSet)
 	{
-		check(Indent > 0);
-		UE_LOG(LogPoseSearch, Log, TEXT("%s ERROR! Cycle detected! Use 'a.PoseSearchInteraction.DiagnoseTickDependencies' CVar to gather more infos to diagnose the problem"), FCString::Spc((Indent-1) * 2));
+		if (ensure(Indent > 0))
+		{
+			UE_LOG(LogPoseSearch, Log, TEXT("%s ERROR! Cycle detected! Use 'a.PoseSearchInteraction.DiagnoseTickDependencies' CVar to gather more infos to diagnose the problem"), FCString::Spc((Indent - 1) * 2));
+		}
 	}
 	else
 	{
@@ -60,7 +62,7 @@ static void LogTickFunction(FTickFunction& Tick, ENamedThreads::Type CurrentThre
 		{
 			TSet<FTickFunction*> VisitedTickFunctions;
 			ShowPrerequistes(Tick, Indent, VisitedTickFunctions);
-			check(VisitedTickFunctions.IsEmpty());
+			ensure(VisitedTickFunctions.IsEmpty());
 		}
 	}
 
@@ -132,12 +134,25 @@ static FInteractionSearchResult InitSearchResult(const FSearchResult& SearchResu
 		if (const IPoseHistory* PoseHistory = SearchContext.GetPoseHistory(AnimContextIndex))
 		{
 			const UObject* AnimContext = SearchContext.GetAnimContext(AnimContextIndex);
-			check(AnimContext);
-			const USkeleton* Skeleton = GetContextSkeleton(AnimContext, false);
-			check(Skeleton);
-
-			PoseHistory->GetTransformAtTime(0.f, InteractionSearchResult.ActorRootTransforms[AnimContextIndex], Skeleton, ComponentSpaceIndexType, WorldSpaceIndexType);
-			PoseHistory->GetTransformAtTime(0.f, InteractionSearchResult.ActorRootBoneTransforms[AnimContextIndex], Skeleton, RootBoneIndexType, ComponentSpaceIndexType);
+			if (ensure(AnimContext))
+			{
+				const USkeleton* Skeleton = GetContextSkeleton(AnimContext, false);
+				if (ensure(Skeleton))
+				{
+					PoseHistory->GetTransformAtTime(0.f, InteractionSearchResult.ActorRootTransforms[AnimContextIndex], Skeleton, ComponentSpaceIndexType, WorldSpaceIndexType);
+					PoseHistory->GetTransformAtTime(0.f, InteractionSearchResult.ActorRootBoneTransforms[AnimContextIndex], Skeleton, RootBoneIndexType, ComponentSpaceIndexType);
+				}
+				else
+				{
+					InteractionSearchResult = FInteractionSearchResult();
+					break;
+				}
+			}
+			else
+			{
+				InteractionSearchResult = FInteractionSearchResult();
+				break;
+			}
 		}
 		else
 		{
@@ -235,8 +250,10 @@ static void InitSearchResults(TArray<FInteractionSearchResult>& SearchResults, T
 				for (int32 AnimContextIndex = 0; AnimContextIndex < SearchContexts[SearchIndex].Num(); ++AnimContextIndex)
 				{
 					const UObject* SearchAnimContext = SearchContexts[SearchIndex].GetAnimContext(AnimContextIndex);
-					check(SearchAnimContext);
-					VisitedAnimContexts.Add(SearchAnimContext);
+					if (ensure(SearchAnimContext))
+					{
+						VisitedAnimContexts.Add(SearchAnimContext);
+					}
 				}
 			}
 		}
@@ -245,7 +262,7 @@ static void InitSearchResults(TArray<FInteractionSearchResult>& SearchResults, T
 
 static UActorComponent* FindComponentForTickDependencies(const UObject* AnimContext)
 {
-	check(AnimContext);
+	ensure(AnimContext);
 	
 	if (const UAnimInstance* AnimInstance = Cast<UAnimInstance>(AnimContext))
 	{
@@ -269,12 +286,12 @@ static void AddPrerequisite(FTickFunction& TickFunction, UObject* TargetObject,
 			{
 				LogTickFunction(TickFunction, Type, true, 1);
 			}
-			check(VisitedTickFunctions.IsEmpty());
+			ensure(VisitedTickFunctions.IsEmpty());
 			if (!ValidateTickDependenciesCycles(TargetTickFunction, VisitedTickFunctions))
 			{
 				LogTickFunction(TargetTickFunction, Type, true, 1);
 			}
-			check(VisitedTickFunctions.IsEmpty());
+			ensure(VisitedTickFunctions.IsEmpty());
 		}
 #endif // ENABLE_ANIM_DEBUG
 
@@ -294,7 +311,7 @@ static void AddPrerequisite(FTickFunction& TickFunction, UObject* TargetObject,
 			{
 				LogTickFunction(TickFunction, Type, true, 1);
 			}
-			check(VisitedTickFunctions.IsEmpty());
+			ensure(VisitedTickFunctions.IsEmpty());
 		}
 #endif // ENABLE_ANIM_DEBUG
 	}
@@ -304,10 +321,12 @@ static void AddPrerequisite(FTickFunction& TickFunction, UObject* TargetObject,
 ///////////////////////////////////////////////////////////
 const UE::PoseSearch::IPoseHistory* FInteractionSearchContextBase::GetPoseHistory(int32 Index) const
 {
-	check(PoseHistories.IsValidIndex(Index));
-	if (TSharedPtr<const IPoseHistory, ESPMode::ThreadSafe> Pin = PoseHistories[Index].Pin())
+	if (ensure(PoseHistories.IsValidIndex(Index)))
 	{
-		return Pin.Get();
+		if (TSharedPtr<const IPoseHistory, ESPMode::ThreadSafe> Pin = PoseHistories[Index].Pin())
+		{
+			return Pin.Get();
+		}
 	}
 	return nullptr;
 }
@@ -315,45 +334,51 @@ const UE::PoseSearch::IPoseHistory* FInteractionSearchContextBase::GetPoseHistor
 void FInteractionSearchContextBase::CacheHashesForEquivalence()
 {
 	// trying to inefficiently call CalculateHashesForEquivalence multiple times?
-	check(CachedHashForEquivalence == InvalidCachedHashForEquivalence && CachedContextHashForEquivalence == InvalidCachedHashForEquivalence);
+	ensure(CachedHashForEquivalence == InvalidCachedHashForEquivalence && CachedContextHashForEquivalence == InvalidCachedHashForEquivalence);
 	CalculateHashesForEquivalence(CachedHashForEquivalence, CachedContextHashForEquivalence);
 }
 
 void FInteractionSearchContextBase::CalculateHashesForEquivalence(uint32& OutCachedHashForEquivalence, uint32& OutCachedContextHashForEquivalence) const
 {
-	check(!AnimContexts.IsEmpty());
+	if (ensure(!AnimContexts.IsEmpty()))
+	{
+		OutCachedContextHashForEquivalence = GetTypeHash(AnimContexts[0]);
 
-	OutCachedContextHashForEquivalence = GetTypeHash(AnimContexts[0]);
+		for (int32 Index = 1; Index < AnimContexts.Num(); ++Index)
+		{
+			// AnimContexts must be sorted to have deterministic searches across multiple frames
+			ensure(AnimContexts[Index - 1] < AnimContexts[Index]);
+			OutCachedContextHashForEquivalence = HashCombineFast(OutCachedContextHashForEquivalence, GetTypeHash(AnimContexts[Index]));
+		}
 
-	for (int32 Index = 1; Index < AnimContexts.Num(); ++Index)
-	{
-		// AnimContexts must be sorted to have deterministic searches across multiple frames
-		check(AnimContexts[Index - 1] < AnimContexts[Index]);
-		OutCachedContextHashForEquivalence = HashCombineFast(OutCachedContextHashForEquivalence, GetTypeHash(AnimContexts[Index]));
-	}
+		ensure(OutCachedContextHashForEquivalence != InvalidCachedHashForEquivalence);
 
-	check(OutCachedContextHashForEquivalence != InvalidCachedHashForEquivalence);
+		// skipping bDisableCollisions and PoseHistories for hashing
+		OutCachedHashForEquivalence = HashCombineFast(OutCachedContextHashForEquivalence, GetTypeHash(Database));
 
-	// skipping bDisableCollisions and PoseHistories for hashing
-	OutCachedHashForEquivalence = HashCombineFast(OutCachedContextHashForEquivalence, GetTypeHash(Database));
+		for (int32 Index = 0; Index < Roles.Num(); ++Index)
+		{
+			OutCachedHashForEquivalence = HashCombineFast(OutCachedHashForEquivalence, GetTypeHash(Roles[Index]));
+		}
 
-	for (int32 Index = 0; Index < Roles.Num(); ++Index)
+		// only an invalid hash is zero!
+		ensure(OutCachedHashForEquivalence != InvalidCachedHashForEquivalence);
+	}
+	else
 	{
-		OutCachedHashForEquivalence = HashCombineFast(OutCachedHashForEquivalence, GetTypeHash(Roles[Index]));
+		OutCachedContextHashForEquivalence = InvalidCachedHashForEquivalence;
+		OutCachedContextHashForEquivalence = InvalidCachedHashForEquivalence;
 	}
-
-	// only an invalid hash is zero!
-	check(OutCachedHashForEquivalence != InvalidCachedHashForEquivalence);
 }
 
 uint32 FInteractionSearchContextBase::GetHashForEquivalence() const
 {
 #if DO_CHECK
-	check(CachedHashForEquivalence != InvalidCachedHashForEquivalence && CachedContextHashForEquivalence != InvalidCachedHashForEquivalence);
+	ensure(CachedHashForEquivalence != InvalidCachedHashForEquivalence && CachedContextHashForEquivalence != InvalidCachedHashForEquivalence);
 	// making sure properties didn't change without updating CachedHashForEquivalence
 	uint32 TestCachedHashForEquivalence, TestCachedContextHashForEquivalence;
 	CalculateHashesForEquivalence(TestCachedHashForEquivalence, TestCachedContextHashForEquivalence);
-	check(CachedHashForEquivalence == TestCachedHashForEquivalence && CachedContextHashForEquivalence == TestCachedContextHashForEquivalence);
+	ensure(CachedHashForEquivalence == TestCachedHashForEquivalence && CachedContextHashForEquivalence == TestCachedContextHashForEquivalence);
 #endif //DO_CHECK
 	return CachedHashForEquivalence;
 }
@@ -361,11 +386,11 @@ uint32 FInteractionSearchContextBase::GetHashForEquivalence() const
 uint32 FInteractionSearchContextBase::GetContextHashForEquivalence() const
 {
 #if DO_CHECK
-	check(CachedHashForEquivalence != InvalidCachedHashForEquivalence && CachedContextHashForEquivalence != InvalidCachedHashForEquivalence);
+	ensure(CachedHashForEquivalence != InvalidCachedHashForEquivalence && CachedContextHashForEquivalence != InvalidCachedHashForEquivalence);
 	// making sure properties didn't change without updating CachedHashForEquivalence
 	uint32 TestCachedHashForEquivalence, TestCachedContextHashForEquivalence;
 	CalculateHashesForEquivalence(TestCachedHashForEquivalence, TestCachedContextHashForEquivalence);
-	check(CachedHashForEquivalence == TestCachedHashForEquivalence && CachedContextHashForEquivalence == TestCachedContextHashForEquivalence);
+	ensure(CachedHashForEquivalence == TestCachedHashForEquivalence && CachedContextHashForEquivalence == TestCachedContextHashForEquivalence);
 #endif // DO_CHECK
 	return CachedContextHashForEquivalence;
 }
@@ -383,15 +408,23 @@ void FInteractionSearchContextBase::AddReferencedObjects(FReferenceCollector& Co
 void FInteractionSearchContextBase::Add(const UObject* AnimContext, const IPoseHistory* PoseHistory, const FRole Role, bool bWantsDisableCollisions)
 {
 	// trying to inefficiently call Add again after having calculated the hash via CalculateHashForEquivalence?
-	check(CachedHashForEquivalence == InvalidCachedHashForEquivalence && CachedContextHashForEquivalence == InvalidCachedHashForEquivalence);
-
-	check(AnimContext && PoseHistory);
+	ensure(CachedHashForEquivalence == InvalidCachedHashForEquivalence && CachedContextHashForEquivalence == InvalidCachedHashForEquivalence);
 
+	ensure(AnimContext);
 	// AnimContexts must be sorted to have deterministic searches across multiple frames
 	check(AnimContexts.IsEmpty() || AnimContexts.Last() < AnimContext);
 
 	AnimContexts.Add(AnimContext);
-	PoseHistories.Add(PoseHistory->AsWeak());
+	
+	if (ensure(PoseHistory))
+	{
+		PoseHistories.Add(PoseHistory->AsWeak());
+	}
+	else
+	{
+		PoseHistories.Add(nullptr);
+	}
+
 	Roles.Add(Role);
 
 	bDisableCollisions |= bWantsDisableCollisions;
@@ -551,7 +584,7 @@ void FValidInteractionSearch::DoDisableCollisions()
 	UE_MT_SCOPED_WRITE_ACCESS(ThreadSafeCounter);
 #endif // ENABLE_ANIM_DEBUG
 
-	check(DisabledCollisions.IsEmpty());
+	ensure(DisabledCollisions.IsEmpty());
 	if (IsDisableCollisions())
 	{
 		FMemMark Mark(FMemStack::Get());
@@ -1841,7 +1874,7 @@ bool FInteractionIsland::GetResult_AnyThread(const UObject* AnimContext, FPoseSe
 		const FSearchIndexAsset* SearchIndexAsset = SearchResult->GetSearchIndexAsset();
 		check(SearchIndexAsset);
 
-		const FPoseSearchDatabaseAnimationAssetBase* DatabaseAnimationAssetBase = Database->GetDatabaseAnimationAsset(*SearchIndexAsset);
+		const FPoseSearchDatabaseAnimationAssetBase* DatabaseAnimationAssetBase = Database->GetDatabaseAnimationAsset(SearchIndexAsset->GetSourceAssetIdx());
 		check(DatabaseAnimationAssetBase);
 
 		Result.SelectedAnim = DatabaseAnimationAssetBase->GetAnimationAsset();
diff --git a/Engine/Plugins/Animation/PoseSearch/Source/Runtime/Private/PoseSearchInteractionSubsystem.cpp b/Engine/Plugins/Animation/PoseSearch/Source/Runtime/Private/PoseSearchInteractionSubsystem.cpp
--- a/Engine/Plugins/Animation/PoseSearch/Source/Runtime/Private/PoseSearchInteractionSubsystem.cpp
+++ b/Engine/Plugins/Animation/PoseSearch/Source/Runtime/Private/PoseSearchInteractionSubsystem.cpp
@@ -1317,7 +1317,7 @@ void UPoseSearchInteractionSubsystem::GenerateSearchContexts(float DeltaSeconds,
 								const FInteractionSearchResult& PreviousSearchResult = ValidInteractionSearch->GetSearchResult();
 								if (const UE::PoseSearch::FSearchIndexAsset* SearchIndexAsset = PreviousSearchResult.GetSearchIndexAsset())
 								{
-									if (const FPoseSearchDatabaseAnimationAssetBase* DatabaseAsset = PreviousSearchResult.Database->GetDatabaseAnimationAsset(*SearchIndexAsset))
+									if (const FPoseSearchDatabaseAnimationAssetBase* DatabaseAsset = PreviousSearchResult.Database->GetDatabaseAnimationAsset(SearchIndexAsset->GetSourceAssetIdx()))
 									{
 										check(SearchIndexAsset->GetToRealTimeFactor() > UE_KINDA_SMALL_NUMBER);
 										// in case DatabaseAsset->GetAnimationAsset() is a blendspace, PreviousSearchResult.AssetTime is a normalized time in the interval [0,1] 
diff --git a/Engine/Plugins/Animation/PoseSearch/Source/Runtime/Private/PoseSearchLibrary.cpp b/Engine/Plugins/Animation/PoseSearch/Source/Runtime/Private/PoseSearchLibrary.cpp
--- a/Engine/Plugins/Animation/PoseSearch/Source/Runtime/Private/PoseSearchLibrary.cpp
+++ b/Engine/Plugins/Animation/PoseSearch/Source/Runtime/Private/PoseSearchLibrary.cpp
@@ -764,7 +764,7 @@ void UPoseSearchLibrary::TraceMotionMatching(
 		const UPoseSearchDatabase* CurrentResultDatabase = SearchResult.Database.Get();
 		if (SearchIndexAsset && CurrentResultDatabase)
 		{
-			const FPoseSearchDatabaseAnimationAssetBase* DatabaseAsset = CurrentResultDatabase->GetDatabaseAnimationAsset(*SearchIndexAsset);
+			const FPoseSearchDatabaseAnimationAssetBase* DatabaseAsset = CurrentResultDatabase->GetDatabaseAnimationAsset(SearchIndexAsset->GetSourceAssetIdx());
 			check(DatabaseAsset);
 			if (UAnimationAsset* AnimationAsset = Cast<UAnimationAsset>(DatabaseAsset->GetAnimationAsset()))
 			{
@@ -1199,29 +1199,31 @@ void UPoseSearchLibrary::MotionMatch(
 			if (const IPoseHistory* PoseHistory = InternalPoseHistories[RoleIndex])
 			{
 				const USkeleton* Skeleton = GetContextSkeleton(GetAnimContext(&Contexts[RoleIndex]), false);
-				check(Skeleton);
-				// @todo: add input BlendParameters to support sampling FutureAnimation blendspaces and support for multi character
-				const UAnimationAsset* AnimationAsset = Cast<UAnimationAsset>(Future.Animation);
-				if (!AnimationAsset)
+				if (ensure(Skeleton))
 				{
-					if (const UMultiAnimAsset* MultiAnimAsset = Cast<UMultiAnimAsset>(Future.Animation))
+					// @todo: add input BlendParameters to support sampling FutureAnimation blendspaces and support for multi character
+					const UAnimationAsset* AnimationAsset = Cast<UAnimationAsset>(Future.Animation);
+					if (!AnimationAsset)
 					{
-						AnimationAsset = MultiAnimAsset->GetAnimationAsset(Roles[RoleIndex]);
-					}
-					else
-					{
-						checkNoEntry();
+						if (const UMultiAnimAsset* MultiAnimAsset = Cast<UMultiAnimAsset>(Future.Animation))
+						{
+							AnimationAsset = MultiAnimAsset->GetAnimationAsset(Roles[RoleIndex]);
+						}
+						else
+						{
+							ensure(false);
+						}
 					}
-				}
 
-				MemStackPoseHistories[RoleIndex].Init(InternalPoseHistories[RoleIndex]);
-				MemStackPoseHistories[RoleIndex].ExtractAndAddFuturePoses(AnimationAsset, FutureAnimationTime, FiniteDelta, FVector::ZeroVector, FutureIntervalTime, Skeleton
+					MemStackPoseHistories[RoleIndex].Init(InternalPoseHistories[RoleIndex]);
+					MemStackPoseHistories[RoleIndex].ExtractAndAddFuturePoses(AnimationAsset, FutureAnimationTime, FiniteDelta, FVector::ZeroVector, FutureIntervalTime, Skeleton
 #if WITH_EDITORONLY_DATA
-				, false
-				, GetContextTransform(GetAnimContext(&Contexts[RoleIndex]), false)
+						, false
+						, GetContextTransform(GetAnimContext(&Contexts[RoleIndex]), false)
 #endif //WITH_EDITORONLY_DATA
-				);
-				InternalPoseHistories[RoleIndex] = MemStackPoseHistories[RoleIndex].GetThisOrPoseHistory();
+					);
+					InternalPoseHistories[RoleIndex] = MemStackPoseHistories[RoleIndex].GetThisOrPoseHistory();
+				}
 			}
 		}
 	}		
diff --git a/Engine/Plugins/Animation/PoseSearch/Source/Runtime/Private/PoseSearchModule.cpp b/Engine/Plugins/Animation/PoseSearch/Source/Runtime/Private/PoseSearchModule.cpp
--- a/Engine/Plugins/Animation/PoseSearch/Source/Runtime/Private/PoseSearchModule.cpp
+++ b/Engine/Plugins/Animation/PoseSearch/Source/Runtime/Private/PoseSearchModule.cpp
@@ -87,7 +87,7 @@ public:
 			if (const FSearchIndexAsset* SearchIndexAsset = SearchResult.GetSearchIndexAsset())
 			{
 				const UPoseSearchDatabase* Database = SearchResult.Database.Get();
-				if (const FPoseSearchDatabaseAnimationAssetBase* DatabaseAnimationAssetBase = Database->GetDatabaseAnimationAsset(*SearchIndexAsset))
+				if (const FPoseSearchDatabaseAnimationAssetBase* DatabaseAnimationAssetBase = Database->GetDatabaseAnimationAsset(SearchIndexAsset->GetSourceAssetIdx()))
 				{
 					ProviderResult.SelectedAsset = DatabaseAnimationAssetBase->GetAnimationAsset();
 					ProviderResult.Dissimilarity = SearchResult.PoseCost;
diff --git a/Engine/Plugins/Animation/PoseSearch/Source/Runtime/Private/PoseSearchResult.cpp b/Engine/Plugins/Animation/PoseSearch/Source/Runtime/Private/PoseSearchResult.cpp
--- a/Engine/Plugins/Animation/PoseSearch/Source/Runtime/Private/PoseSearchResult.cpp
+++ b/Engine/Plugins/Animation/PoseSearch/Source/Runtime/Private/PoseSearchResult.cpp
@@ -19,7 +19,10 @@ const FSearchIndexAsset* FDatabasePoseIdx::GetSearchIndexAsset(bool bMandatory)
 {
 	if (bMandatory)
 	{
-		check(IsValid());
+		if (!ensure(IsValid()))
+		{
+			return nullptr;
+		}
 	}
 	else if (!IsValid())
 	{
@@ -33,7 +36,7 @@ const UAnimationAsset* FDatabasePoseIdx::GetCurrentResultAnimationAsset() const
 {
 	if (const FSearchIndexAsset* SearchIndexAsset = GetSearchIndexAsset())
 	{
-		return Database->GetDatabaseAnimationAsset(*SearchIndexAsset)->GetAnimationAssetForRole(Database->Schema->GetDefaultRole());
+		return Database->GetDatabaseAnimationAsset(SearchIndexAsset->GetSourceAssetIdx())->GetAnimationAssetForRole(Database->Schema->GetDefaultRole());
 	}
 	return nullptr;
 }
@@ -42,7 +45,7 @@ const UAnimationAsset* FDatabasePoseIdx::GetCurrentResultAnimationAsset(const FR
 {
 	if (const FSearchIndexAsset* SearchIndexAsset = GetSearchIndexAsset())
 	{
-		return Database->GetDatabaseAnimationAsset(*SearchIndexAsset)->GetAnimationAssetForRole(Role);
+		return Database->GetDatabaseAnimationAsset(SearchIndexAsset->GetSourceAssetIdx())->GetAnimationAssetForRole(Role);
 	}
 	return nullptr;
 }
@@ -60,7 +63,7 @@ bool FSearchResult::DebugValidate() const
 	{
 		const FSearchIndex& SearchIndex = Database->GetSearchIndex();
 		const FSearchIndexAsset& SearchIndexAsset = SearchIndex.GetAssetForPose(PoseIdx);
-		if (const FPoseSearchDatabaseAnimationAssetBase* DatabaseAnimationAssetBase = Database->GetDatabaseAnimationAsset(SearchIndexAsset))
+		if (const FPoseSearchDatabaseAnimationAssetBase* DatabaseAnimationAssetBase = Database->GetDatabaseAnimationAsset(SearchIndexAsset.GetSourceAssetIdx()))
 		{
 			PRAGMA_DISABLE_DEPRECATION_WARNINGS
 			const float RealAssetTime = AssetTime * SearchIndexAsset.GetToRealTimeFactor();
@@ -84,7 +87,7 @@ bool FSearchResult::DebugValidate() const
 
 void FSearchResult::UpdateWithNormalizedTime(float NormalizedTime)
 {
-	check(DebugValidate());
+	ensure(DebugValidate());
 
 	if (IsValid())
 	{
@@ -93,27 +96,29 @@ void FSearchResult::UpdateWithNormalizedTime(float NormalizedTime)
 		const float RealTime = NormalizedTime * SearchIndexAsset.GetToRealTimeFactor();
 		
 #if WITH_EDITOR && DO_CHECK
-		const FPoseSearchDatabaseAnimationAssetBase* DatabaseAnimationAssetBase = Database->GetDatabaseAnimationAsset(SearchIndexAsset);
-		check(DatabaseAnimationAssetBase);
-		if (Cast<UBlendSpace>(DatabaseAnimationAssetBase->GetAnimationAssetForRole(Database->Schema->GetDefaultRole())))
+		const FPoseSearchDatabaseAnimationAssetBase* DatabaseAnimationAssetBase = Database->GetDatabaseAnimationAsset(SearchIndexAsset.GetSourceAssetIdx());
+		if (ensure(DatabaseAnimationAssetBase))
 		{
-			const float PlayLength = DatabaseAnimationAssetBase->GetPlayLength(SearchIndexAsset.GetBlendParameters());
-			if (PlayLength > UE_KINDA_SMALL_NUMBER)
+			if (Cast<UBlendSpace>(DatabaseAnimationAssetBase->GetAnimationAssetForRole(Database->Schema->GetDefaultRole())))
 			{
-				// Asset player time for blendspaces is normalized [0, 1] so we need to convert 
-				// to a real time before we advance it
-				check(NormalizedTime >= 0.f && NormalizedTime <= 1.f);
-				check(FMath::IsNearlyEqual(RealTime, NormalizedTime * PlayLength));
+				const float PlayLength = DatabaseAnimationAssetBase->GetPlayLength(SearchIndexAsset.GetBlendParameters());
+				if (PlayLength > UE_KINDA_SMALL_NUMBER)
+				{
+					// Asset player time for blendspaces is normalized [0, 1] so we need to convert 
+					// to a real time before we advance it
+					ensure(NormalizedTime >= 0.f && NormalizedTime <= 1.f);
+					ensure(FMath::IsNearlyEqual(RealTime, NormalizedTime * PlayLength));
+				}
+				else
+				{
+					ensure(FMath::IsNearlyEqual(SearchIndexAsset.GetToRealTimeFactor(), 1.f));
+				}
 			}
 			else
 			{
-				check(FMath::IsNearlyEqual(SearchIndexAsset.GetToRealTimeFactor(), 1.f));
+				ensure(FMath::IsNearlyEqual(SearchIndexAsset.GetToRealTimeFactor(), 1.f));
 			}
 		}
-		else
-		{
-			check(FMath::IsNearlyEqual(SearchIndexAsset.GetToRealTimeFactor(), 1.f));
-		}
 #endif // WITH_EDITOR && DO_CHECK
 		
 		PoseIdx = SearchIndexAsset.GetPoseIndexFromTime(RealTime, Database->Schema->SampleRate);
@@ -126,41 +131,43 @@ void FSearchResult::UpdateWithNormalizedTime(float NormalizedTime)
 		Reset();
 	}
 	
-	check(DebugValidate());
+	ensure(DebugValidate());
 }
 
 void FSearchResult::UpdateWithRealTime(float RealTime)
 {
-	check(DebugValidate());
+	ensure(DebugValidate());
 
 	if (IsValid())
 	{
 		const FSearchIndexAsset& SearchIndexAsset = Database->GetSearchIndex().GetAssetForPose(PoseIdx);
-		check(SearchIndexAsset.GetToRealTimeFactor() > UE_KINDA_SMALL_NUMBER);
+		ensure(SearchIndexAsset.GetToRealTimeFactor() > UE_KINDA_SMALL_NUMBER);
 		const float NormalizedTime = RealTime / SearchIndexAsset.GetToRealTimeFactor();
 
 #if WITH_EDITOR && DO_CHECK
-		const FPoseSearchDatabaseAnimationAssetBase* DatabaseAnimationAssetBase = Database->GetDatabaseAnimationAsset(SearchIndexAsset);
-		check(DatabaseAnimationAssetBase);
-		if (Cast<UBlendSpace>(DatabaseAnimationAssetBase->GetAnimationAssetForRole(Database->Schema->GetDefaultRole())))
+		const FPoseSearchDatabaseAnimationAssetBase* DatabaseAnimationAssetBase = Database->GetDatabaseAnimationAsset(SearchIndexAsset.GetSourceAssetIdx());
+		if (ensure(DatabaseAnimationAssetBase))
 		{
-			const float PlayLength = DatabaseAnimationAssetBase->GetPlayLength(SearchIndexAsset.GetBlendParameters());
-			if (PlayLength > UE_KINDA_SMALL_NUMBER)
+			if (Cast<UBlendSpace>(DatabaseAnimationAssetBase->GetAnimationAssetForRole(Database->Schema->GetDefaultRole())))
 			{
-				// Asset player time for blendspaces is normalized [0, 1] so we need to convert 
-				// to a real time before we advance it
-				check(RealTime >= 0.f && RealTime <= PlayLength);
-				check(FMath::IsNearlyEqual(NormalizedTime, PlayLength > UE_KINDA_SMALL_NUMBER ? RealTime / PlayLength : 0.f));
+				const float PlayLength = DatabaseAnimationAssetBase->GetPlayLength(SearchIndexAsset.GetBlendParameters());
+				if (PlayLength > UE_KINDA_SMALL_NUMBER)
+				{
+					// Asset player time for blendspaces is normalized [0, 1] so we need to convert 
+					// to a real time before we advance it
+					ensure(RealTime >= 0.f && RealTime <= PlayLength);
+					ensure(FMath::IsNearlyEqual(NormalizedTime, PlayLength > UE_KINDA_SMALL_NUMBER ? RealTime / PlayLength : 0.f));
+				}
+				else
+				{
+					ensure(FMath::IsNearlyEqual(SearchIndexAsset.GetToRealTimeFactor(), 1.f));
+				}
 			}
 			else
 			{
-				check(FMath::IsNearlyEqual(SearchIndexAsset.GetToRealTimeFactor(), 1.f));
+				ensure(FMath::IsNearlyEqual(SearchIndexAsset.GetToRealTimeFactor(), 1.f));
 			}
 		}
-		else
-		{
-			check(FMath::IsNearlyEqual(SearchIndexAsset.GetToRealTimeFactor(), 1.f));
-		}
 #endif // WITH_EDITOR && DO_CHECK
 
 		PoseIdx = SearchIndexAsset.GetPoseIndexFromTime(RealTime, Database->Schema->SampleRate);
@@ -173,21 +180,28 @@ void FSearchResult::UpdateWithRealTime(float RealTime)
 		Reset();
 	}
 	
-	check(DebugValidate());
+	ensure(DebugValidate());
 }
 
 bool FSearchResult::IsEventSearchFromTag(const FGameplayTag& EventTag) const
 {
-	check(IsValid());
-	return Database->GetSearchIndex().EventData.IsPoseFromEventTag(EventPoseIdx, EventTag);
+	if (ensure(IsValid()))
+	{
+		return Database->GetSearchIndex().EventData.IsPoseFromEventTag(EventPoseIdx, EventTag);
+	}
+	return false;
 }
 
 float FSearchResult::CalculateTimeToEvent() const
 {
-	check(IsValid() && IsEventSearchResult());
+	if (!ensure(IsValid() && IsEventSearchResult()))
+	{
+		return 0.f;
+	}
+	
 	const FSearchIndex& SearchIndex = Database->GetSearchIndex();
 
-	check(SearchIndex.PoseMetadata[PoseIdx].GetAssetIndex() == SearchIndex.PoseMetadata[EventPoseIdx].GetAssetIndex());
+	ensure(SearchIndex.PoseMetadata[PoseIdx].GetAssetIndex() == SearchIndex.PoseMetadata[EventPoseIdx].GetAssetIndex());
 	const FSearchIndexAsset& SearchIndexAsset = SearchIndex.GetAssetForPose(PoseIdx);
 
 	// @todo: the mathc here can be simplified between GetDeltaTimeBetweenPoseIndexes and GetTimeFromPoseIndex methods
@@ -229,7 +243,7 @@ FSearchResults_Single::FSearchResults_Single()
 {
 	UpdateWith = [this](const FSearchResult& SearchResult)
 		{
-			check(SearchResult.IsValid());
+			ensure(SearchResult.IsValid());
 			if (SearchResult.PoseCost < SingleSearchResult.PoseCost)
 			{
 				SingleSearchResult = SearchResult;
@@ -241,7 +255,7 @@ FSearchResults_Single::FSearchResults_Single()
 			if (SingleSearchResult.IsValid() && !SingleSearchResult.IsAssetTimeValid())
 			{
 				SingleSearchResult.SetAssetTime(SingleSearchResult.Database->GetNormalizedAssetTime(SingleSearchResult.PoseIdx));
-				check(SingleSearchResult.DebugValidate());
+				ensure(SingleSearchResult.DebugValidate());
 			}
 		};
 
@@ -266,12 +280,12 @@ FSearchResults_Single::FSearchResults_Single()
 FSearchResults_Multi::FSearchResults_Multi(int32 InMaxNumberOfWantedResults)
 	: MaxNumberOfWantedResults(InMaxNumberOfWantedResults)
 {
-	check(MaxNumberOfWantedResults > 0); // or else we need to refine the logic in UpdateWith
+	ensure(MaxNumberOfWantedResults > 0); // or else we need to refine the logic in UpdateWith
 	SearchResults.Reserve(MaxNumberOfWantedResults);
 
 	UpdateWith = [this](const FSearchResult& SearchResult)
 		{
-			check(SearchResult.IsValid());
+			ensure(SearchResult.IsValid());
 			if (SearchResults.Num() < MaxNumberOfWantedResults)
 			{
 				SearchResults.Add(SearchResult);
@@ -302,11 +316,11 @@ FSearchResults_Multi::FSearchResults_Multi(int32 InMaxNumberOfWantedResults)
 		{
 			for (FSearchResult& SearchResult : SearchResults)
 			{
-				check(SearchResult.IsValid());
+				ensure(SearchResult.IsValid());
 				if (!SearchResult.IsAssetTimeValid())
 				{
 					SearchResult.SetAssetTime(SearchResult.Database->GetNormalizedAssetTime(SearchResult.PoseIdx));
-					check(SearchResult.DebugValidate());
+					ensure(SearchResult.DebugValidate());
 				}
 			}
 		};
@@ -349,7 +363,7 @@ FSearchResults_AssetBests::FSearchResults_AssetBests()
 {
 	UpdateWith = [this](const FSearchResult& SearchResult)
 		{
-			check(SearchResult.IsValid());
+			ensure(SearchResult.IsValid());
 			if (const FSearchIndexAsset* SearchIndexAsset = SearchResult.GetSearchIndexAsset())
 			{
 				FSourceAssetIdxToBestSearchResult& SourceAssetIdxToBestSearchResult = PerDatabaseSourceAssetIdxToBestSearchResult.FindOrAdd(SearchResult.Database.Get());
@@ -368,11 +382,11 @@ FSearchResults_AssetBests::FSearchResults_AssetBests()
 				for (TPair<int32, FSearchResult>& SourceAssetIdxToBestSearchResultPair : PerDatabaseSourceAssetIdxToBestSearchResultPair.Value)
 				{
 					FSearchResult& BestSearchResultForSourceAssetIdx = SourceAssetIdxToBestSearchResultPair.Value;
-					check(BestSearchResultForSourceAssetIdx.IsValid());
+					ensure(BestSearchResultForSourceAssetIdx.IsValid());
 					if (!BestSearchResultForSourceAssetIdx.IsAssetTimeValid())
 					{
 						BestSearchResultForSourceAssetIdx.SetAssetTime(BestSearchResultForSourceAssetIdx.Database->GetNormalizedAssetTime(BestSearchResultForSourceAssetIdx.PoseIdx));
-						check(BestSearchResultForSourceAssetIdx.DebugValidate());
+						ensure(BestSearchResultForSourceAssetIdx.DebugValidate());
 					}
 				}
 			}
@@ -473,10 +487,12 @@ void FSearchResults_AssetBests::Shrink(int32 MaxNumResults)
 		for (const FSearchResult& BestSearchResult : BestSearchResults)
 		{
 			const FSearchIndexAsset* SearchIndexAsset = BestSearchResult.GetSearchIndexAsset();
-			check(SearchIndexAsset);
-			FSourceAssetIdxToBestSearchResult& SourceAssetIdxToBestSearchResult = PerDatabaseSourceAssetIdxToBestSearchResult.FindOrAdd(BestSearchResult.Database.Get());
-			FSearchResult& BestSearchResultForSourceAssetIdx = SourceAssetIdxToBestSearchResult.FindOrAdd(SearchIndexAsset->GetSourceAssetIdx());
-			BestSearchResultForSourceAssetIdx = BestSearchResult;
+			if (ensure(SearchIndexAsset))
+			{
+				FSourceAssetIdxToBestSearchResult& SourceAssetIdxToBestSearchResult = PerDatabaseSourceAssetIdxToBestSearchResult.FindOrAdd(BestSearchResult.Database.Get());
+				FSearchResult& BestSearchResultForSourceAssetIdx = SourceAssetIdxToBestSearchResult.FindOrAdd(SearchIndexAsset->GetSourceAssetIdx());
+				BestSearchResultForSourceAssetIdx = BestSearchResult;
+			}
 		}
 	}
 }
@@ -491,25 +507,26 @@ bool FPoseSearchBlueprintResult::InitFrom(const UE::PoseSearch::FSearchResult& S
 	{
 		const UPoseSearchDatabase* Database = SearchResult.Database.Get();
 
-		const FPoseSearchDatabaseAnimationAssetBase* DatabaseAsset = Database->GetDatabaseAnimationAsset(*SearchIndexAsset);
-		check(DatabaseAsset);
-		
-		SelectedAnim = DatabaseAsset->GetAnimationAsset();
-		SelectedTime = SearchResult.GetAssetTime();
-		bIsContinuingPoseSearch = SearchResult.bIsContinuingPoseSearch;
-		WantedPlayRate = InWantedPlayRate;
-		bLoop = SearchIndexAsset->IsLooping();
-		bIsMirrored = SearchIndexAsset->IsMirrored();
-		BlendParameters = SearchIndexAsset->GetBlendParameters();
-		SelectedDatabase = Database;
-		SearchCost = SearchResult.PoseCost;
-		Role = Database->Schema->GetDefaultRole();
-
-		bIsInteraction = false;
-		ActorRootTransforms.Reset();
-		ActorRootBoneTransforms.Reset();
-		AnimContexts.Reset();
-		return true;
+		const FPoseSearchDatabaseAnimationAssetBase* DatabaseAsset = Database->GetDatabaseAnimationAsset(SearchIndexAsset->GetSourceAssetIdx());
+		if (ensure(DatabaseAsset))
+		{
+			SelectedAnim = DatabaseAsset->GetAnimationAsset();
+			SelectedTime = SearchResult.GetAssetTime();
+			bIsContinuingPoseSearch = SearchResult.bIsContinuingPoseSearch;
+			WantedPlayRate = InWantedPlayRate;
+			bLoop = SearchIndexAsset->IsLooping();
+			bIsMirrored = SearchIndexAsset->IsMirrored();
+			BlendParameters = SearchIndexAsset->GetBlendParameters();
+			SelectedDatabase = Database;
+			SearchCost = SearchResult.PoseCost;
+			Role = Database->Schema->GetDefaultRole();
+
+			bIsInteraction = false;
+			ActorRootTransforms.Reset();
+			ActorRootBoneTransforms.Reset();
+			AnimContexts.Reset();
+			return true;
+		}
 	}
 	
 	*this = FPoseSearchBlueprintResult();
diff --git a/Engine/Plugins/Animation/PoseSearch/Source/Runtime/Public/PoseSearch/Chooser/PoseSearchChooserColumn.h b/Engine/Plugins/Animation/PoseSearch/Source/Runtime/Public/PoseSearch/Chooser/PoseSearchChooserColumn.h
--- a/Engine/Plugins/Animation/PoseSearch/Source/Runtime/Public/PoseSearch/Chooser/PoseSearchChooserColumn.h
+++ b/Engine/Plugins/Animation/PoseSearch/Source/Runtime/Public/PoseSearch/Chooser/PoseSearchChooserColumn.h
@@ -1,16 +1,12 @@
 // Copyright Epic Games, Inc. All Rights Reserved.
 #pragma once
 
-#include <IHasContext.h>
-
-#include "CoreMinimal.h"
 #include "ChooserPropertyAccess.h"
+#include "CoreMinimal.h"
 #include "IChooserColumn.h"
 #include "PoseSearch/Chooser/ChooserParameterPoseHistoryBase.h"
-#include "PoseSearch/PoseSearchLibrary.h"
-#include "Serialization/MemoryReader.h"
+#include "PoseSearch/PoseSearchDatabase.h"
 #include "StructUtils/InstancedStruct.h"
-#include "Templates/SubclassOf.h"
 #include "PoseSearchChooserColumn.generated.h"
 
 #define UE_API POSESEARCH_API
@@ -26,6 +22,44 @@ namespace UE::PoseSearch
 		float RowCost = 0.f;
 	};
 	FArchive& operator<<(FArchive& Ar, FActiveColumnCost& ActiveColumnCost);
+
+	// Experimental, this feature might be removed without warning, not for production use
+	struct FMapping 
+	{
+	private:
+		struct FTableCell 
+		{
+			FPoseSearchColumn* Column = nullptr;
+			int32 RowIndex = INDEX_NONE;
+
+			bool operator == (const FTableCell& Other) const
+			{
+				return Column == Other.Column && RowIndex == Other.RowIndex;
+			}
+		};
+
+		typedef TArray<FTableCell> FTableCells;
+		TMap<UPoseSearchDatabase*, FTableCells> Map;
+	
+		
+	public:
+		explicit FMapping(UChooserTable* Chooser = nullptr);
+		void Init(UChooserTable* Chooser);
+		bool Equals(const FMapping& Other) const;
+		int32 GetNumIndexesInDatabase(const UPoseSearchDatabase* Database) const;
+		int32 FindIndexInDatabase(int32 RowIndex, const FPoseSearchColumn* Column, const UPoseSearchDatabase* Database) const;
+		const FPoseSearchDatabaseAnimationAsset* FindDatabaseAnimationAsset(int32 IndexInDatabase, const UPoseSearchDatabase* Database) const;
+
+#if WITH_EDITOR
+		void SetAnimationAssetAt(const FPoseSearchDatabaseAnimationAsset& AnimationAsset, int32 IndexInDatabase, const UPoseSearchDatabase* Database);
+#endif // WITH_EDITOR
+
+		int32 GetRowIndex(int32 IndexInDatabase, const UPoseSearchDatabase* Database) const;
+		void ApplyMappingToAllColumns();
+
+		static FMapping* FindMapping(UChooserTable* Chooser);
+	};
+
 } // namespace UE::PoseSearch
 
 
@@ -37,7 +71,11 @@ struct FPoseHistoryContextProperty :  public FChooserParameterPoseHistoryBase
 	UPROPERTY(EditAnywhere, Meta = (BindingType = "FPoseHistoryReference", BindingAllowFunctions = "true", BindingColor = "StructPinTypeColor"), Category = "Binding")
 	FChooserPropertyBinding Binding;
 	
-	UE_API virtual bool GetValue(FChooserEvaluationContext& Context, FPoseHistoryReference& OutResult) const override;
+	virtual bool GetValue(FChooserEvaluationContext& Context, FPoseHistoryReference& OutResult) const override
+	{
+		return Binding.GetStructValue(Context, OutResult);
+	}
+
 	virtual bool IsBound() const override
 	{
 		return Binding.IsBoundToRoot || !Binding.PropertyBindingChain.IsEmpty();
@@ -46,16 +84,63 @@ struct FPoseHistoryContextProperty :  public FChooserParameterPoseHistoryBase
 	CHOOSER_PARAMETER_BOILERPLATE();
 };
 
-// reintroduced just to make the TRACE_CHOOSER_VALUE work, since it relies on ToCStr(InputValue.Get<FChooserParameterBase>().GetDebugName()) to work
 USTRUCT(Experimental)
-struct FPoseSearchColumnNullInputValue : public FChooserParameterBase
+struct FPoseSearchColumnRow
 {
 	GENERATED_BODY()
 
-	virtual void GetDisplayName(FText& OutName) const override
-	{
-		OutName = FText::FromString("FPoseSearchColumn");
-	}
+	UPROPERTY(EditAnywhere, Category = "Data")
+	FPoseSearchDatabaseAnimationAsset Data;
+
+	UPROPERTY(EditAnywhere, Category = "Data")
+	TObjectPtr<UPoseSearchDatabase> Database;
+};
+
+USTRUCT(Experimental)
+struct FPoseSearchColumnRowReflection
+{
+	GENERATED_BODY()
+
+#if WITH_EDITORONLY_DATA
+	// if bDisableReselection is true, poses from the same asset cannot be reselected. Useful to avoid jumping on frames on the same looping animations
+	UPROPERTY(EditAnywhere, Category = "Data")
+	bool bDisableReselection = false;
+
+	// This allows users to set if this animation is original only (no mirrored data), original and mirrored, or only the mirrored version of this animation.
+	// It requires the mirror table to be set up in the database Schema.
+	UPROPERTY(EditAnywhere, Category = "Data")
+	EPoseSearchMirrorOption MirrorOption = EPoseSearchMirrorOption::UnmirroredOnly;
+
+	// If true this BlendSpace will output a single segment in the database.
+	UPROPERTY(EditAnywhere, Category = "Data")
+	bool bUseSingleSample = false;
+
+	// When turned on, this will use the set grid samples of the blend space asset for sampling. This will override the Number of Horizontal/Vertical Samples.
+	UPROPERTY(EditAnywhere, Category = "Data", meta = (EditCondition = "!bUseSingleSample", EditConditionHides))
+	bool bUseGridForSampling = false;
+
+	// Sets the number of horizontal samples in the blend space to pull the animation data coverage from. The larger the samples the more the data, but also the more memory and performance it takes.
+	UPROPERTY(EditAnywhere, Category = "Data", meta = (EditCondition = "!bUseSingleSample && !bUseGridForSampling", EditConditionHides, ClampMin = "1", UIMin = "1", UIMax = "25"))
+	int32 NumberOfHorizontalSamples = 9;
+	
+	// Sets the number of vertical samples in the blend space to pull the animation data coverage from.The larger the samples the more the data, but also the more memory and performance it takes.
+	UPROPERTY(EditAnywhere, Category = "Data", meta = (EditCondition = "!bUseSingleSample && !bUseGridForSampling", EditConditionHides, ClampMin = "1", UIMin = "1", UIMax = "25"))
+	int32 NumberOfVerticalSamples = 2;
+
+	// BlendParams used to sample this BlendSpace
+	UPROPERTY(EditAnywhere, Category = "Data", meta = (EditCondition = "bUseSingleSample", EditConditionHides))
+	float BlendParamX = 0.f;
+
+	// BlendParams used to sample this BlendSpace
+	UPROPERTY(EditAnywhere, Category = "Data", meta = (EditCondition = "bUseSingleSample", EditConditionHides))
+	float BlendParamY = 0.f;
+
+	// @todo: make it EditAnywhere, but populate the asset list filter to look only for the 
+	//		  databases owned by the FPoseSearchColumn root chooser table owner
+	UPROPERTY(VisibleAnywhere, Category = "Data")
+	TObjectPtr<UPoseSearchDatabase> Database;
+
+#endif // WITH_EDITORONLY_DATA
 };
 
 USTRUCT(DisplayName = "Pose Match", Meta = (Experimental, Category = "Experimental", Tooltip = "This column filters out all assets except the one which is selected by motion matching query.  Results must be AnimationAssets with a PoseSearchBranchIn notify state.  It also outputs OutputStartTime to specify the frame which matched pose best.  To work as intended it must be placed last (furthest right) in the Chooser so that other filters are applied first."))
@@ -63,11 +148,27 @@ struct FPoseSearchColumn : public FChooserColumnBase
 {
 	GENERATED_BODY()
 	
-	UE_API FPoseSearchColumn();
-
 private:
+#if WITH_EDITORONLY_DATA
+	UPROPERTY()
+	TObjectPtr<UPoseSearchDatabase> InternalDatabase_DEPRECATED;
+#endif // WITH_EDITORONLY_DATA
+
+	UPROPERTY(VisibleAnywhere, Category = "Data")
+	TArray<FPoseSearchColumnRow> RowValues;
+
+	UPROPERTY(VisibleAnywhere, Category = "Data")
+	FPoseSearchColumnRow FallbackValue;
+
 	UPROPERTY(VisibleAnywhere, Category = "Data")
-	TObjectPtr<UPoseSearchDatabase> InternalDatabase;
+	TObjectPtr<UChooserTable> Chooser;
+
+	// @todo: currently Mapping is duplicated in every FPoseSearchColumn and it'd be nice to store it somewhere in the RootChooserTable only. 
+	//        That would also speed up the complexity of FMapping::FindMapping.
+	//
+	// structure containing the mapping between (FPoseSearchColumn, RowIndex) to (UPoseSearchDatabase, IndexInDatabase). 
+	// Not serialized, but generated on PostLoad and regenerated when Chooser->GetRootChooser() properties changes
+	UE::PoseSearch::FMapping Mapping;
 
 public:
 	// @todo: expose it as pinnable parameter
@@ -91,11 +192,6 @@ public:
 	UPROPERTY(EditAnywhere, DisplayName = "Pose History", NoClear, Meta = (ExcludeBaseStruct, BaseStruct = "/Script/PoseSearch.ChooserParameterPoseHistoryBase"), Category = "Data")
 	FInstancedStruct InputValue;
 	
-	// Experimental, this feature might be removed without warning, not for production use
-	// reintroduced just to make the TRACE_CHOOSER_VALUE work, since it relies on ToCStr(InputValue.Get<FChooserParameterBase>().GetDebugName()) to work
-	UPROPERTY()
-	FPoseSearchColumnNullInputValue NullInputValue;
-
 	// Float output for the start time with the best matching pose
 	UPROPERTY(EditAnywhere, NoClear, Meta = (ExcludeBaseStruct, BaseStruct = "/Script/Chooser.ChooserParameterFloatBase"), Category = "Output")
 	FInstancedStruct OutputStartTime;
@@ -112,12 +208,36 @@ public:
 	UPROPERTY(EditAnywhere, NoClear, Meta = (ExcludeBaseStruct, BaseStruct = "/Script/Chooser.ChooserParameterFloatBase"), Category = "Output")
 	FInstancedStruct OutputCost;
 
-	UE_API const UObject* GetDatabaseAsset(int32 RowIndex) const;
-	UE_API const UPoseSearchSchema* GetDatabaseSchema() const;
-	UE_API void SetDatabaseSchema(const UPoseSearchSchema* Schema);
+	FPoseSearchColumn() = default;
+	UE_API FPoseSearchColumn(const FPoseSearchColumn& Other);
+	UE_API FPoseSearchColumn& operator=(const FPoseSearchColumn& Other);
+
+	UE_API virtual void PostLoad() override;
 
+	UE_API UPoseSearchDatabase* GetDatabase(int32 RowIndex) const;
+#if WITH_EDITOR
+	UE_API void SwitchDatabase(UPoseSearchDatabase* NewDatabase, int32 RowIndex);
+#endif // WITH_EDITOR
+
+	const FPoseSearchColumnRow& GetRowValue(int32 RowIndex) const;
+	FPoseSearchColumnRow& EditRowValue(int32 RowIndex);
+
+#if WITH_EDITOR
+	static void InsertAnimationAssetAt(const FPoseSearchDatabaseAnimationAsset& AnimationAsset, int32 AnimationAssetIndex, UChooserTable* Chooser, const UPoseSearchDatabase* Database);
+	static void RemoveAnimationAssetAt(int32 AnimationAssetIndex, UChooserTable* Chooser, const UPoseSearchDatabase* Database);
+	static void SetAnimationAssetAt(const FPoseSearchDatabaseAnimationAsset& AnimationAsset, int32 AnimationAssetIndex, UChooserTable* Chooser, const UPoseSearchDatabase* Database);
+	const UChooserTable* GetChooser() const { return Chooser.Get(); }
+#endif // WITH_EDITOR
+
+	static const FPoseSearchDatabaseAnimationAsset* GetDatabaseAnimationAsset(int32 AnimationAssetIndex, UChooserTable* Chooser, const UPoseSearchDatabase* Database);
+	static int32 GetNumAnimationAssets(UChooserTable* Chooser, const UPoseSearchDatabase* Database);
+
+	UE::PoseSearch::FMapping& GetMapping() { return Mapping; }
+	void UpdateMapping(UE::PoseSearch::FMapping& NewMapping);
+
+	int32 GetNumRows() const { return RowValues.Num(); }
 private:
-	virtual bool HasFilters() const override;
+	virtual bool HasFilters() const override { return true; }
 	virtual void Filter(FChooserEvaluationContext& Context, const FChooserIndexArray& IndexListIn, FChooserIndexArray& IndexListOut, TArrayView<uint8> ScratchArea) const override;
 	virtual bool HasCosts() const override { return true; }
 	virtual void SetOutputs(FChooserEvaluationContext& Context, int RowIndex, TArrayView<uint8> ScratchArea) const override;
@@ -138,33 +258,26 @@ private:
 	virtual void AddToDetails(FInstancedPropertyBag& PropertyBag, int32 ColumnIndex, int32 RowIndex) override;
 	virtual void SetFromDetails(FInstancedPropertyBag& PropertyBag, int32 ColumnIndex, int32 RowIndex) override;
 
+	virtual void ReplaceReferences(UObject* ReferenceToReplace, UObject* ReplaceWith) override;
+
 	virtual FName RowValuesPropertyName() override;
 	virtual void SetNumRows(int32 NumRows) override;
-	virtual void InsertRows(int Index, int Count) override;
-	virtual void DeleteRows(TArrayView<int> RowIndices) override;
-	virtual void MoveRow(int SourceRowIndex, int TargetRowIndex) override;
+	virtual void InsertRows(int Index, int Count);
+	virtual void DeleteRows(TArrayView<int> RowIndices);
+	virtual void MoveRow(int SourceRowIndex, int TargetRowIndex);
 	virtual void CopyRow(FChooserColumnBase& SourceColumn, int SourceRowIndex, int TargetRowIndex) override;
-	virtual UScriptStruct* GetInputBaseType() const override;
-	virtual const UScriptStruct* GetInputType() const override;
-	virtual void SetInputType(const UScriptStruct* Type) override;
+	virtual void CopyFallback(FChooserColumnBase& SourceColumn) override;
 	
 	virtual bool AutoPopulates() const override { return true; }
 	virtual void AutoPopulate(int32 RowIndex, UObject* OutputObject) override;
-	
-	virtual bool HasPrimaryInput() override
-	{
-		return false;
-	}
 
+	virtual UScriptStruct* GetInputBaseType() const override;
+	virtual const UScriptStruct* GetInputType() const override;
+	virtual FInstancedStruct* GetInputValuePtr() override;
+	virtual void SetInputType(const UScriptStruct* Type) override;
 #endif // WITH_EDITOR
-	virtual FChooserParameterBase* GetInputValue() override;
-	const FChooserParameterBase* GetInputValue() const;
 
-#if WITH_EDITOR
-	UObject* GetReferencedObject(int32 RowIndex) const;
-#endif // WITH_EDITOR
-	
-	void CheckConsistency() const;
+	virtual FChooserParameterBase* GetInputValue() override;
 };
 
 #undef UE_API
diff --git a/Engine/Plugins/Animation/PoseSearch/Source/Runtime/Public/PoseSearch/PoseSearchDatabase.h b/Engine/Plugins/Animation/PoseSearch/Source/Runtime/Public/PoseSearch/PoseSearchDatabase.h
--- a/Engine/Plugins/Animation/PoseSearch/Source/Runtime/Public/PoseSearch/PoseSearchDatabase.h
+++ b/Engine/Plugins/Animation/PoseSearch/Source/Runtime/Public/PoseSearch/PoseSearchDatabase.h
@@ -643,6 +643,7 @@ public:
 	// It'll return INDEX_NONE if it cannot find a pose
 	UE_API int32 GetPoseIndex(const UObject* AnimationAsset, float AnimationAssetTime, bool bMirrored, const FVector& BlendParameters, const UE::PoseSearch::FStackAssetIndexes* AssetIndexesToConsider = nullptr) const;
 
+#if WITH_EDITOR
 	UE_DEPRECATED(5.7, "Use AddAnimationAsset(const FPoseSearchDatabaseAnimationAsset& AnimationAsset) instead")
 	UE_API void AddAnimationAsset(FInstancedStruct AnimationAsset);
 
@@ -651,9 +652,17 @@ public:
 	UE_API void InsertAnimationAssetAt(const FPoseSearchDatabaseAnimationAsset& AnimationAsset, int32 AnimationAssetIndex);
 	UE_API void RemoveAnimationAssetAt(int32 AnimationAssetIndex);
 
+	UE_API void SetAnimationAssetAt(const FPoseSearchDatabaseAnimationAsset& AnimationAsset, int32 AnimationAssetIndex);
+#endif // WITH_EDITOR
+
 	UE_API const FPoseSearchDatabaseAnimationAsset* GetDatabaseAnimationAsset(int32 AnimationAssetIndex) const;
+
+	UE_DEPRECATED(5.8, "Use GetDatabaseAnimationAsset(SearchIndexAsset.GetSourceAssetIdx()) instead")
 	UE_API const FPoseSearchDatabaseAnimationAsset* GetDatabaseAnimationAsset(const UE::PoseSearch::FSearchIndexAsset& SearchIndexAsset) const;
+
+	UE_DEPRECATED(5.8, "Use SetAnimationAssetAt(AnimationAsset, AnimationAssetIndex) instead")
 	UE_API FPoseSearchDatabaseAnimationAsset* GetMutableDatabaseAnimationAsset(int32 AnimationAssetIndex);
+	UE_DEPRECATED(5.8, "Use SetAnimationAssetAt(AnimationAsset, SearchIndexAsset.GetSourceAssetIdx()) instead")
 	UE_API FPoseSearchDatabaseAnimationAsset* GetMutableDatabaseAnimationAsset(const UE::PoseSearch::FSearchIndexAsset& SearchIndexAsset);
 
 	template<typename TDatabaseAnimationAsset> UE_DEPRECATED(5.7, "Use the non templated version of this API") const TDatabaseAnimationAsset* GetDatabaseAnimationAsset(int32 AnimationAssetIndex) const;
@@ -750,11 +759,8 @@ inline const TDatabaseAnimationAsset* UPoseSearchDatabase::GetDatabaseAnimationA
 	if (std::is_same<TDatabaseAnimationAsset, FPoseSearchDatabaseAnimationAsset>::value ||
 		std::is_same<TDatabaseAnimationAsset, FPoseSearchDatabaseAnimationAssetBase>::value)
 	{
-		if (DatabaseAnimationAssets.IsValidIndex(AnimationAssetIndex))
-		{
-			// using reinterpret_cast to be able to compile this method even with other types (that will return null)
-			return reinterpret_cast<const TDatabaseAnimationAsset*>(&DatabaseAnimationAssets[AnimationAssetIndex]);
-		}
+		// using reinterpret_cast to be able to compile this method even with other types (that will return null)
+		return reinterpret_cast<const TDatabaseAnimationAsset*>(GetDatabaseAnimationAsset(AnimationAssetIndex));
 	}
 	
 	return nullptr;
@@ -774,11 +780,8 @@ inline TDatabaseAnimationAsset* UPoseSearchDatabase::GetMutableDatabaseAnimation
 	if (std::is_same<TDatabaseAnimationAsset, FPoseSearchDatabaseAnimationAsset>::value ||
 		std::is_same<TDatabaseAnimationAsset, FPoseSearchDatabaseAnimationAssetBase>::value)
 	{
-		if (DatabaseAnimationAssets.IsValidIndex(AnimationAssetIndex))
-		{
-			// using reinterpret_cast to be able to compile this method even with other types (that will return null)
-			return reinterpret_cast<TDatabaseAnimationAsset*>(&DatabaseAnimationAssets[AnimationAssetIndex]);
-		}
+		// using reinterpret_cast to be able to compile this method even with other types (that will return null)
+		return reinterpret_cast<TDatabaseAnimationAsset*>(GetMutableDatabaseAnimationAsset(AnimationAssetIndex));
 	}
 
 	return nullptr;
diff --git a/Engine/Plugins/Chooser/Source/Chooser/Private/Chooser.cpp b/Engine/Plugins/Chooser/Source/Chooser/Private/Chooser.cpp
--- a/Engine/Plugins/Chooser/Source/Chooser/Private/Chooser.cpp
+++ b/Engine/Plugins/Chooser/Source/Chooser/Private/Chooser.cpp
@@ -36,7 +36,6 @@ void UChooserTable::PostLoad()
 {
 	Super::PostLoad();
 
-#if WITH_EDITORONLY_DATA
 	for (FInstancedStruct& ColumnData : ColumnsStructs)
 	{
 		if (ColumnData.IsValid())
@@ -46,6 +45,7 @@ void UChooserTable::PostLoad()
 		}
 	}
 
+#if WITH_EDITORONLY_DATA
 	CachedPreviousOutputObjectType = OutputObjectType;
 	CachedPreviousResultType = ResultType;
 
@@ -108,6 +108,9 @@ void UChooserTable::BeginDestroy()
 	ColumnsStructs.Empty();
 #if WITH_EDITORONLY_DATA
 	ResultsStructs.Empty();
+
+	NestedChoosers.Empty();
+	NestedObjects.Empty();
 #endif
 	CookedResults.Empty();
 	Super::BeginDestroy();
@@ -377,6 +380,14 @@ static void ReplaceReferencesInTable(UChooserTable* ChooserToReplace, UChooserTa
 			}	
 		}
 	}
+
+	for (FInstancedStruct& ColumnData : Table->ColumnsStructs)
+	{
+		FChooserColumnBase& Column = ColumnData.GetMutable<FChooserColumnBase>();
+		Column.ReplaceReferences(ChooserToReplace, ReplaceWith);
+	}
+
+	Table->PostEditChange();
 }
 
 static void ReplaceReferences(UChooserTable* ChooserToReplace, UChooserTable* ReplaceWith, UChooserTable* RootTable)
@@ -408,6 +419,38 @@ void UChooserTable::DeleteNestedChooser(UChooserTable* ChooserToDelete)
 	ReplaceReferences(ChooserToDelete, nullptr, this);
 }
 
+static void RemoveReferencesFromColumns(UObject* ObjectToDelete, UChooserTable* ChooserTable)
+{
+	for (FInstancedStruct& ColumnData : ChooserTable->ColumnsStructs)
+	{
+		FChooserColumnBase& Column = ColumnData.GetMutable<FChooserColumnBase>();
+		Column.ReplaceReferences(ObjectToDelete, nullptr);
+	}
+
+	for (UObject* NestedObject : ChooserTable->NestedObjects)
+	{
+		if (UChooserTable* NestedChooser = Cast<UChooserTable>(NestedObject))
+		{
+			for (FInstancedStruct& ColumnData : NestedChooser->ColumnsStructs)
+			{
+				FChooserColumnBase& Column = ColumnData.GetMutable<FChooserColumnBase>();
+				Column.ReplaceReferences(ObjectToDelete, nullptr);
+			}
+		}
+	}
+}
+
+void UChooserTable::DeleteNestedObject(UObject* ObjectToDelete)
+{
+	static FName DeletedObjectName = "DeletedNestedObject";
+	DeletedObjectName.SetNumber(DeletedObjectName.GetNumber() + 1);
+	ObjectToDelete->Rename(*DeletedObjectName.ToString());
+	Modify(true);
+	RemoveNestedObject(ObjectToDelete);
+	RemoveReferencesFromColumns(ObjectToDelete, this);
+	PostEditChange();
+}
+
 #endif
 
 #if WITH_EDITOR
@@ -968,6 +1011,10 @@ void UChooserTable::AddNestedChooser(UChooserTable* Chooser)
 
 void UChooserTable::AddNestedObject(UObject* Object)
 {
+	check(Object);
+	ensure(!RootChooser);
+	ensure(!NestedObjects.Contains(Object));
+
 	FName NewObjectName = Object->GetFName();
 	FString NewChooserPlainName = NewObjectName.GetPlainNameString();
 	bool bNameConflict = false;
@@ -996,11 +1043,34 @@ void UChooserTable::AddNestedObject(UObject* Object)
 
 void UChooserTable::RemoveNestedChooser(UChooserTable* Chooser)
 {
+	// deleting the Chooser by removing all the rows and resetting ResultsStructs / DisabledRows / FallbackResult
+	// so columns are aware of dependencies being changed
+	TArray<int32> RowsToDelete;
+	RowsToDelete.Reserve(Chooser->ResultsStructs.Num() + 1);
+	for (int32 RowIndex = Chooser->ResultsStructs.Num() - 1; RowIndex >= 0; --RowIndex)
+	{
+		RowsToDelete.Add(RowIndex);
+	}
+	RowsToDelete.Add(ChooserColumn_SpecialIndex_Fallback);
+	
+	// calling DeleteRows on columns before removing them from FallbackResult / DisabledRows / ResultsStructs,
+	// so they are aware of what's gonna be deleted, and access those properties to remove references, etc.
+	for(FInstancedStruct& ColumnData : Chooser->ColumnsStructs)
+	{
+		FChooserColumnBase& Column = ColumnData.GetMutable<FChooserColumnBase>();
+		Column.DeleteRows(RowsToDelete);
+	}
+
+	Chooser->ResultsStructs.Reset();
+	Chooser->DisabledRows.Reset();
+	Chooser->FallbackResult.Reset();
+
 	RemoveNestedObject(Chooser);
 }
 
 void UChooserTable::RemoveNestedObject(UObject* Object)
 {
+	ensure(NestedObjects.Contains(Object));
 	NestedObjects.Remove(Object);
 	NestedObjectsChanged.Broadcast();
 }
diff --git a/Engine/Plugins/Chooser/Source/Chooser/Public/Chooser.h b/Engine/Plugins/Chooser/Source/Chooser/Public/Chooser.h
--- a/Engine/Plugins/Chooser/Source/Chooser/Public/Chooser.h
+++ b/Engine/Plugins/Chooser/Source/Chooser/Public/Chooser.h
@@ -34,6 +34,7 @@ public:
 	CHOOSER_API void RemoveDisabledData();
 	CHOOSER_API void PopulateCookedData(bool bCooking);
 	CHOOSER_API void DeleteNestedChooser(UChooserTable* NestedChooser);
+	CHOOSER_API void DeleteNestedObject(UObject* ObjectToDelete);
 #endif
 
 	bool IsRowDisabled(int32 RowIndex) const
diff --git a/Engine/Plugins/Chooser/Source/Chooser/Public/IChooserColumn.h b/Engine/Plugins/Chooser/Source/Chooser/Public/IChooserColumn.h
--- a/Engine/Plugins/Chooser/Source/Chooser/Public/IChooserColumn.h
+++ b/Engine/Plugins/Chooser/Source/Chooser/Public/IChooserColumn.h
@@ -117,6 +117,8 @@ public:
 
 	virtual void AddToDetails(FInstancedPropertyBag& PropertyBag, int32 ColumnIndex, int32 RowIndex) {}
 	virtual void SetFromDetails(FInstancedPropertyBag& PropertyBag, int32 ColumnIndex, int32 RowIndex) {}
+
+	virtual void ReplaceReferences(UObject* ReferenceToReplace, UObject* ReplaceWith) {}
 #endif
 
 #if WITH_EDITORONLY_DATA
diff --git a/Engine/Plugins/Chooser/Source/ChooserEditor/Private/ChooserTableViewModel.cpp b/Engine/Plugins/Chooser/Source/ChooserEditor/Private/ChooserTableViewModel.cpp
--- a/Engine/Plugins/Chooser/Source/ChooserEditor/Private/ChooserTableViewModel.cpp
+++ b/Engine/Plugins/Chooser/Source/ChooserEditor/Private/ChooserTableViewModel.cpp
@@ -54,7 +54,7 @@
 #include "SChooserTableWidget.h"
 #include "AssetRegistry/IAssetRegistry.h"
 #include "UObject/AssetRegistryTagsContext.h"
-
+#include "Templates/Greater.h"
 
 #include UE_INLINE_GENERATED_CPP_BY_NAME(ChooserTableViewModel)
 
@@ -362,8 +362,9 @@ void FChooserTableViewModel::AddColumn(const UScriptStruct* ColumnType)
 
 	FInstancedStruct NewColumn;
 	NewColumn.InitializeAs(ColumnType);
-	NewColumn.GetMutable<FChooserColumnBase>().Initialize(Chooser);
-	const FChooserColumnBase& NewColumnRef = NewColumn.Get<FChooserColumnBase>();
+	FChooserColumnBase& NewColumnRef = NewColumn.GetMutable<FChooserColumnBase>();
+	NewColumnRef.Initialize(Chooser);
+
 	int32 InsertIndex = 0;
 	if (NewColumnRef.IsRandomizeColumn())
 	{
@@ -437,7 +438,7 @@ void FChooserTableViewModel::RefreshRowSelectionDetails()
 		
 			if (Chooser->ResultType != EObjectChooserResultType::NoPrimaryResult)
 			{
-				FInstancedStruct& Result = SelectedRowIndex == ColumnWidget_SpecialIndex_Fallback? Chooser->FallbackResult : Chooser->ResultsStructs[SelectedRowIndex];;
+				FInstancedStruct& Result = SelectedRowIndex == ColumnWidget_SpecialIndex_Fallback? Chooser->FallbackResult : Chooser->ResultsStructs[SelectedRowIndex];
 				Selection->Properties.AddProperties({ResultPropertyDesc});
 				Selection->Properties.SetValueStruct("Result", FConstStructView(FInstancedStruct::StaticStruct(), reinterpret_cast<uint8*>(&Result)));
 			}
@@ -507,7 +508,17 @@ int32 FChooserTableViewModel::DeleteRowsInternal(TConstArrayView<int32> InRowsTo
 		// sort indices in reverse
 		TArray<int32> RowsToDelete;
 		RowsToDelete = InRowsToDelete;
-		RowsToDelete.Sort([](int32 A, int32 B){ return A>B; });
+		RowsToDelete.Sort(TGreater<>());
+
+		// calling DeleteRows on columns before removing them from FallbackResult / DisabledRows / ResultsStructs,
+		// so they are aware of what's gonna be deleted, and access those properties to remove references, etc.
+		for(FInstancedStruct& ColumnData : Chooser->ColumnsStructs)
+		{
+			FChooserColumnBase& Column = ColumnData.GetMutable<FChooserColumnBase>();
+			Column.DeleteRows(RowsToDelete);
+		}
+
+		// removing RowsToDelete from FallbackResult / DisabledRows / ResultsStructs
 		for(int32 RowIndex : RowsToDelete)
 		{
 			if (RowIndex == ColumnWidget_SpecialIndex_Fallback)
@@ -525,11 +536,7 @@ int32 FChooserTableViewModel::DeleteRowsInternal(TConstArrayView<int32> InRowsTo
 			}
 		}
 
-		for(FInstancedStruct& ColumnData : Chooser->ColumnsStructs)
-		{
-			FChooserColumnBase& Column = ColumnData.GetMutable<FChooserColumnBase>();
-			Column.DeleteRows(RowsToDelete);
-		}
+		
 		UpdateTableRows();
 	}
 
@@ -596,24 +603,20 @@ void FChooserTableViewModel::AutoPopulateRow(int32 Index)
 	{
 		if (Chooser->ResultsStructs.IsValidIndex(Index) && Chooser->ResultsStructs[Index].IsValid())
 		{
-			if (UObject* ReferencedObject = Chooser->ResultsStructs[Index].Get<FObjectChooserBase>().GetReferencedObject()) 
+			UObject* ReferencedObject = Chooser->ResultsStructs[Index].Get<FObjectChooserBase>().GetReferencedObject();
+			for (FInstancedStruct& ColumnData : Chooser->ColumnsStructs)
 			{
-				for (FInstancedStruct& ColumnData : Chooser->ColumnsStructs)
-				{
-					FChooserColumnBase& Column = ColumnData.GetMutable<FChooserColumnBase>();
-					Column.AutoPopulate(Index, ReferencedObject);
-				}
+				FChooserColumnBase& Column = ColumnData.GetMutable<FChooserColumnBase>();
+				Column.AutoPopulate(Index, ReferencedObject);
 			}
 		}
 		else if (Index == ColumnWidget_SpecialIndex_Fallback && Chooser->FallbackResult.IsValid())
 		{
-			if (UObject* ReferencedObject = Chooser->FallbackResult.Get<FObjectChooserBase>().GetReferencedObject())
+			UObject* ReferencedObject = Chooser->FallbackResult.Get<FObjectChooserBase>().GetReferencedObject();
+			for (FInstancedStruct& ColumnData : Chooser->ColumnsStructs)
 			{
-				for (FInstancedStruct& ColumnData : Chooser->ColumnsStructs)
-				{
-					FChooserColumnBase& Column = ColumnData.GetMutable<FChooserColumnBase>();
-					Column.AutoPopulate(Index, ReferencedObject);
-				}
+				FChooserColumnBase& Column = ColumnData.GetMutable<FChooserColumnBase>();
+				Column.AutoPopulate(Index, ReferencedObject);
 			}
 		}
 	}
@@ -757,6 +760,8 @@ void FChooserTableViewModel::DuplicateColumn(int32 ColumnIndex)
 		Chooser->Modify();
 		FInstancedStruct Column = Chooser->ColumnsStructs[ColumnIndex];
 		Chooser->ColumnsStructs.Insert(Column,ColumnIndex);
+		// technically unnecessary, but let the column know that it needs to be initialized since it got duplicated
+		Chooser->ColumnsStructs[ColumnIndex].GetMutable<FChooserColumnBase>().Initialize(Chooser);
 	}
 
 	UpdateTableColumns();
@@ -906,11 +911,39 @@ UChooserTable* DuplicateNestedChooser(UChooserTable* Chooser, UChooserTable* New
 	}
 	
 	UChooserTable* NewTable = NewObject<UChooserTable>(NewOuter, Chooser->GetFName(), RF_Transactional);
-	NewTable->ColumnsStructs = Chooser->ColumnsStructs;
 	NewTable->ResultsStructs = Chooser->ResultsStructs;
+	NewTable->DisabledRows = Chooser->DisabledRows;
+	
+	// @todo: add fallback FallbackResult copy support (NewTable->FallbackResult = Chooser->FallbackResult)
 	NewTable->RootChooser = RootTable;
 	RootTable->AddNestedObject(NewTable);
 
+	// copy all columns
+	NewTable->ColumnsStructs = Chooser->ColumnsStructs;
+	for (FInstancedStruct& ColumnData : NewTable->ColumnsStructs)
+	{
+		FChooserColumnBase& Column = ColumnData.GetMutable<FChooserColumnBase>();
+		Column.Initialize(NewTable);
+	}
+
+	// clear all column's cell data
+	for (FInstancedStruct& ColumnData : NewTable->ColumnsStructs)
+	{
+		FChooserColumnBase& Column = ColumnData.GetMutable<FChooserColumnBase>();
+		Column.SetNumRows(0);
+		Column.SetNumRows(Chooser->ResultsStructs.Num());
+	}
+
+	for (int32 RowIndex = 0; RowIndex < Chooser->ResultsStructs.Num(); RowIndex++)
+	{
+		for (int32 ColumnIndex = 0; ColumnIndex < NewTable->ColumnsStructs.Num(); ColumnIndex++)
+		{
+			FChooserColumnBase& SourceColumn = Chooser->ColumnsStructs[ColumnIndex].GetMutable<FChooserColumnBase>();
+			FChooserColumnBase& TargetColumn = NewTable->ColumnsStructs[ColumnIndex].GetMutable<FChooserColumnBase>();
+			TargetColumn.CopyRow(SourceColumn, RowIndex, RowIndex);
+		}
+	}
+
 	for (FInstancedStruct& ResultData : NewTable->ResultsStructs)
 	{
 		if (FNestedChooser* NestedChooser = ResultData.GetMutablePtr<FNestedChooser>())
@@ -1117,6 +1150,7 @@ void FChooserTableViewModel::PasteInternal(UChooserTable* PastedContent, int32 P
 			}
 		}
 		Chooser->ColumnsStructs.Insert(PastedContent->ColumnsStructs, InsertColumnIndex);
+		Chooser->ColumnsStructs[InsertColumnIndex].GetMutable<FChooserColumnBase>().Initialize(Chooser);
 		SelectColumn(Chooser, InsertColumnIndex);
 	}
 	else
@@ -1257,6 +1291,7 @@ void FChooserTableViewModel::PasteInternal(UChooserTable* PastedContent, int32 P
 				// if we couldn't find a match, paste a new column
 				Chooser->ColumnsStructs.Insert(PastedColumnData, InsertColumnIndex);
 				FChooserColumnBase& Column = Chooser->ColumnsStructs[InsertColumnIndex].GetMutable<FChooserColumnBase>();
+				Column.Initialize(Chooser);
 				InsertColumnIndex++;
 				Column.SetNumRows(0);
 				Column.SetNumRows(Chooser->ResultsStructs.Num());
diff --git a/Engine/Plugins/Chooser/Source/ChooserEditor/Private/SChooserTableWidget.cpp b/Engine/Plugins/Chooser/Source/ChooserEditor/Private/SChooserTableWidget.cpp
--- a/Engine/Plugins/Chooser/Source/ChooserEditor/Private/SChooserTableWidget.cpp
+++ b/Engine/Plugins/Chooser/Source/ChooserEditor/Private/SChooserTableWidget.cpp
@@ -138,6 +138,7 @@ namespace UE::ChooserEditor
 	void SChooserTableWidget::UpdateTableRows()
 	{
 		UChooserTable* Chooser = ChooserViewModel->GetChooser();
+		Chooser->Modify();
 		int32 NewNum = Chooser->ResultsStructs.Num();
 		Chooser->DisabledRows.SetNum(NewNum);
     
@@ -163,6 +164,7 @@ namespace UE::ChooserEditor
 			FChooserColumnBase& Column = ColumnData.GetMutable<FChooserColumnBase>();
 			Column.SetNumRows(NewNum);
 		}
+		Chooser->PostEditChange();
     
 		if (TableView.IsValid())
 		{
diff --git a/Engine/Plugins/Chooser/Source/ChooserEditor/Private/SNestedChooserTree.cpp b/Engine/Plugins/Chooser/Source/ChooserEditor/Private/SNestedChooserTree.cpp
--- a/Engine/Plugins/Chooser/Source/ChooserEditor/Private/SNestedChooserTree.cpp
+++ b/Engine/Plugins/Chooser/Source/ChooserEditor/Private/SNestedChooserTree.cpp
@@ -40,9 +40,16 @@ namespace UE::ChooserEditor
 		{
 			ChooserEditor->SetChooserTableToEdit(Chooser);
 		}
+		else
+		{
+			if (UAssetEditorSubsystem* AssetEditorSS = GEditor->GetEditorSubsystem<UAssetEditorSubsystem>())
+			{
+				AssetEditorSS->OpenEditorForAsset(SelectedObject->Object);
+			}
+		}
 	}
 
-	void SNestedChooserTree::DeleteNestedChooser()
+	void SNestedChooserTree::DeleteNestedObject()
 	{
 		TArray<TSharedPtr<FNestedChooserTreeEntry>>SelectedItems = TreeView->GetSelectedItems();
 		if (SelectedItems.Num() == 1)
@@ -55,6 +62,11 @@ namespace UE::ChooserEditor
 					RootChooser->DeleteNestedChooser(ChooserToDelete);
 				}
 			}
+			else if (UObject* ObjectToDelete = SelectedItems[0]->Object)
+			{
+				FScopedTransaction Transaction(LOCTEXT("Delete Nested Object", "Delete Nested Object"));
+				RootChooser->DeleteNestedObject(ObjectToDelete);
+			}
 		}
 	}
 
@@ -120,7 +132,7 @@ namespace UE::ChooserEditor
 							{
 								if (Event.GetKey() == EKeys::Delete)
 								{
-									DeleteNestedChooser();
+									DeleteNestedObject();
 									return FReply::Handled();
 								}
 								else if (Event.GetKey() == EKeys::F2)
@@ -193,9 +205,9 @@ namespace UE::ChooserEditor
 	TSharedPtr<SWidget> SNestedChooserTree::TreeViewContextMenuOpening()
 	{
 		FMenuBuilder MenuBuilder(true, nullptr);
-		MenuBuilder.AddMenuEntry(LOCTEXT("Delete", "Delete"), LOCTEXT("Delete Tooltip", "Delete Nested Chooser Table"), FSlateIcon(), FUIAction
+		MenuBuilder.AddMenuEntry(LOCTEXT("Delete", "Delete"), LOCTEXT("Delete Tooltip", "Delete Nested Object"), FSlateIcon(), FUIAction
 			(
-				FExecuteAction::CreateRaw(this, &SNestedChooserTree::DeleteNestedChooser)
+				FExecuteAction::CreateRaw(this, &SNestedChooserTree::DeleteNestedObject)
 				)
 				);
 		MenuBuilder.AddMenuEntry(LOCTEXT("Rename", "Rename"), LOCTEXT("Rename Tooltip", "Rename Nested Object"), FSlateIcon(), FUIAction
diff --git a/Engine/Plugins/Chooser/Source/ChooserEditor/Private/SNestedChooserTree.h b/Engine/Plugins/Chooser/Source/ChooserEditor/Private/SNestedChooserTree.h
--- a/Engine/Plugins/Chooser/Source/ChooserEditor/Private/SNestedChooserTree.h
+++ b/Engine/Plugins/Chooser/Source/ChooserEditor/Private/SNestedChooserTree.h
@@ -34,7 +34,7 @@ private:
 	TSharedRef<ITableRow> TreeViewGenerateRow(TSharedPtr<FNestedChooserTreeEntry> InItem, const TSharedRef<STableViewBase>& OwnerTable);
 	void TreeViewGetChildren(TSharedPtr<FNestedChooserTreeEntry> InItem, TArray<TSharedPtr<FNestedChooserTreeEntry>>& OutChildren);
 	void TreeViewDoubleClicked(TSharedPtr<FNestedChooserTreeEntry> SelectedObject);
-	void DeleteNestedChooser();
+	void DeleteNestedObject();
 	void RenameNestedObject();
 	
 	FChooserTableEditor* ChooserEditor = nullptr;
diff --git a/Engine/Plugins/Experimental/MovieScenePoseSearchTracks/Source/MovieScenePoseSearchTracks/Private/Sections/MovieSceneStitchAnimSection.cpp b/Engine/Plugins/Experimental/MovieScenePoseSearchTracks/Source/MovieScenePoseSearchTracks/Private/Sections/MovieSceneStitchAnimSection.cpp
--- a/Engine/Plugins/Experimental/MovieScenePoseSearchTracks/Source/MovieScenePoseSearchTracks/Private/Sections/MovieSceneStitchAnimSection.cpp
+++ b/Engine/Plugins/Experimental/MovieScenePoseSearchTracks/Source/MovieScenePoseSearchTracks/Private/Sections/MovieSceneStitchAnimSection.cpp
@@ -148,7 +148,7 @@ void FMovieSceneStitchAnimEvaluationTask::Execute(UE::UAF::FEvaluationVM& VM) co
 			if (const FSearchIndexAsset* SearchIndexAsset = SearchResult.GetSearchIndexAsset())
 			{
 				const UPoseSearchDatabase* Database = SearchResult.Database.Get();
-				if (const FPoseSearchDatabaseAnimationAssetBase* DatabaseAnimationAssetBase = Database->GetDatabaseAnimationAsset(*SearchIndexAsset))
+				if (const FPoseSearchDatabaseAnimationAssetBase* DatabaseAnimationAssetBase = Database->GetDatabaseAnimationAsset(SearchIndexAsset->GetSourceAssetIdx()))
 				{
 					MatchedAsset = Cast<UAnimationAsset>(DatabaseAnimationAssetBase->GetAnimationAsset());
 					MatchedAssetTime = SearchResult.GetAssetTime();
