PoseSearch - demoting check to ensure in FSearchContext::Track

#Synced-CL 50943617

diff --git a/Engine/Plugins/Animation/PoseSearch/Source/Runtime/Private/PoseSearchContext.cpp b/Engine/Plugins/Animation/PoseSearch/Source/Runtime/Private/PoseSearchContext.cpp
--- a/Engine/Plugins/Animation/PoseSearch/Source/Runtime/Private/PoseSearchContext.cpp
+++ b/Engine/Plugins/Animation/PoseSearch/Source/Runtime/Private/PoseSearchContext.cpp
@@ -1339,11 +1339,13 @@ bool FSearchContext::CanUseContinuingPoseValues() const
 #if UE_POSE_SEARCH_TRACE_ENABLED
 void FSearchContext::Track(const UPoseSearchDatabase* Database, int32 PoseIdx, EPoseCandidateFlags PoseCandidateFlags, const FPoseSearchCost& Cost)
 {
-	check(Database);
-	FBestPoseCandidates& BestPoseCandidates = BestPoseCandidatesMap.FindOrAdd(Database);
-	if (PoseIdx != INDEX_NONE)
+	if (ensure(Database))
 	{
-		BestPoseCandidates.Add(PoseIdx, PoseCandidateFlags, Cost);
+		FBestPoseCandidates& BestPoseCandidates = BestPoseCandidatesMap.FindOrAdd(Database);
+		if (PoseIdx != INDEX_NONE)
+		{
+			BestPoseCandidates.Add(PoseIdx, PoseCandidateFlags, Cost);
+		}
 	}
 }
 
