PoseSearch - making sure there're no flying animation tasks when on game thread by calling GetProxyOnAnyThread so it's safe to access animinstance variables in UPoseSearchLibrary::MotionMatch

#Synced-CL 50943617

diff --git a/Engine/Plugins/Animation/PoseSearch/Source/Runtime/Private/PoseSearchLibrary.cpp b/Engine/Plugins/Animation/PoseSearch/Source/Runtime/Private/PoseSearchLibrary.cpp
--- a/Engine/Plugins/Animation/PoseSearch/Source/Runtime/Private/PoseSearchLibrary.cpp
+++ b/Engine/Plugins/Animation/PoseSearch/Source/Runtime/Private/PoseSearchLibrary.cpp
@@ -1045,6 +1045,19 @@ void UPoseSearchLibrary::MotionMatch(
 		return;
 	}
 
+	// @todo: improve this apprach, perhaps promoting UAnimInstance::GetProxyOnAnyThread from protected to public
+	// making sure there're no flying animation tasks when on game thread by calling GetProxyOnAnyThread so it's safe to access animinstance variables etc
+	class UFinishFlyingAnimInstanceTasks : public UAnimInstance
+	{
+	public:
+		static void Execute(UAnimInstance* AnimInstance)
+		{
+			static_cast<UFinishFlyingAnimInstanceTasks*>(AnimInstance)->GetProxyOnAnyThread<FAnimInstanceProxy>();
+		}
+	};
+	UFinishFlyingAnimInstanceTasks::Execute(AnimInstance);
+
+
 	FMemMark Mark(FMemStack::Get());
 
 	TArray<const IPoseHistory*, TInlineAllocator<PreallocatedRolesNum, TMemStackAllocator<>>> PoseHistories;
