PoseSearch - fix for pose search column crash when InputValue is uninitialized

#Synced-CL 50943617

diff --git a/Engine/Plugins/Animation/PoseSearch/Source/Runtime/Private/Chooser/PoseSearchChooserColumn.cpp b/Engine/Plugins/Animation/PoseSearch/Source/Runtime/Private/Chooser/PoseSearchChooserColumn.cpp
--- a/Engine/Plugins/Animation/PoseSearch/Source/Runtime/Private/Chooser/PoseSearchChooserColumn.cpp
+++ b/Engine/Plugins/Animation/PoseSearch/Source/Runtime/Private/Chooser/PoseSearchChooserColumn.cpp
@@ -826,7 +826,17 @@ void FPoseSearchColumn::Filter(FChooserEvaluationContext& Context, const FChoose
 	}
 
 #if WITH_EDITOR
-	TRACE_CHOOSER_VALUE(Context, ToCStr(InputValue.GetPtr<FChooserParameterBase>()->GetDebugName()), TracedActiveColumnCosts);
+	// @todo: improve how we trace chooser values, without relying on GetDebugName, since it's not a unique identifier
+	FString DebugName;
+	if (const FChooserParameterBase* ChooserParameterBase = InputValue.GetPtr<FChooserParameterBase>())
+	{
+		DebugName = ChooserParameterBase->GetDebugName();
+	}
+	else
+	{
+		DebugName = "FPoseSearchColumn";
+	}
+	TRACE_CHOOSER_VALUE(Context, ToCStr(DebugName), TracedActiveColumnCosts);
 	if (Context.DebuggingInfo.bCurrentDebugTarget)
 	{
 		ActiveColumnCosts = TracedActiveColumnCosts;
