PoseSearch - added fault tolerance to FPoseSearchDatabaseAsyncCacheTask PreCancelIfDependsOn, PreModified, PostModified. Improved logging on OnGetComplete

#Synced-CL 51037391

diff --git a/Engine/Plugins/Animation/PoseSearch/Source/Runtime/Private/PoseSearchDerivedData.cpp b/Engine/Plugins/Animation/PoseSearch/Source/Runtime/Private/PoseSearchDerivedData.cpp
--- a/Engine/Plugins/Animation/PoseSearch/Source/Runtime/Private/PoseSearchDerivedData.cpp
+++ b/Engine/Plugins/Animation/PoseSearch/Source/Runtime/Private/PoseSearchDerivedData.cpp
@@ -1622,10 +1622,10 @@ void FPoseSearchDatabaseAsyncCacheTask::StartNewRequestIfNeeded(bool bPerformCon
 // it cancels and waits for the task to be done and set the state to PreCancelled, so no other new requests can start until the task gets cancelled
 void FPoseSearchDatabaseAsyncCacheTask::PreCancelIfDependsOn(const UObject* Object)
 {
-	check(IsInGameThread());
-
 	if (FKeyBuilder::ShouldInclude(Object))
 	{
+		ensure(IsInGameThread());
+
 		// DatabaseDependencies is updated only in StartNewRequestIfNeeded when there are no active requests, so it's thread safe to access it 
 		if (DatabaseDependencies.Contains(Object))
 		{
@@ -1898,7 +1898,7 @@ void FPoseSearchDatabaseAsyncCacheTask::OnGetComplete(UE::DerivedData::FCacheGet
 							const UAnimationAsset* AnimationAsset = Cast<const UAnimationAsset>(DatabaseAsset->GetAnimationAsset());
 							if (DatabaseAsset->GetAnimationAsset() == nullptr)
 							{
-								UE_LOG(LogPoseSearch, Warning, TEXT("OnGetComplete - - No asset has been selected."));
+								UE_LOG(LogPoseSearch, Warning, TEXT("OnGetComplete - Database %s entry %d has no associated animation asset."), *DatabaseAsset->GetName(), AnimationAssetIndex);
 							}
 							else if (!DatabaseAsset->GetAnimationAsset()->GetPackage() || DatabaseAsset->GetAnimationAsset()->HasAnyFlags(RF_MirroredGarbage))
 							{
@@ -2358,10 +2358,10 @@ void FAsyncPoseSearchDatabasesManagement::OnObjectModified(UObject* Object)
 
 void FAsyncPoseSearchDatabasesManagement::PreModified(UObject* Object)
 {
-	check(IsInGameThread());
-
 	if (FKeyBuilder::ShouldInclude(Object))
 	{
+		ensure(IsInGameThread());
+
 		FScopeLock Lock(&Mutex);
 
 		PartialKeyHashes.InvalidateKeysFor(Object);
@@ -2378,10 +2378,10 @@ void FAsyncPoseSearchDatabasesManagement::PreModified(UObject* Object)
 
 void FAsyncPoseSearchDatabasesManagement::PostModified(UObject* Object)
 {
-	check(IsInGameThread());
-
 	if (FKeyBuilder::ShouldInclude(Object))
 	{
+		ensure(IsInGameThread());
+
 		FScopeLock Lock(&Mutex);
 
 		// collecting databases to synchronize, and merging the results with the PreModified collection
