PoseSearch - adding missing FPoseSearchColumn::HasOutputs

#Synced-CL 51099401

diff --git a/Engine/Plugins/Animation/PoseSearch/Source/Runtime/Public/PoseSearch/Chooser/PoseSearchChooserColumn.h b/Engine/Plugins/Animation/PoseSearch/Source/Runtime/Public/PoseSearch/Chooser/PoseSearchChooserColumn.h
--- a/Engine/Plugins/Animation/PoseSearch/Source/Runtime/Public/PoseSearch/Chooser/PoseSearchChooserColumn.h
+++ b/Engine/Plugins/Animation/PoseSearch/Source/Runtime/Public/PoseSearch/Chooser/PoseSearchChooserColumn.h
@@ -240,6 +240,7 @@ private:
 	virtual bool HasFilters() const override { return true; }
 	virtual void Filter(FChooserEvaluationContext& Context, const FChooserIndexArray& IndexListIn, FChooserIndexArray& IndexListOut, TArrayView<uint8> ScratchArea) const override;
 	virtual bool HasCosts() const override { return true; }
+	virtual bool HasOutputs() const { return true; }
 	virtual void SetOutputs(FChooserEvaluationContext& Context, int RowIndex, TArrayView<uint8> ScratchArea) const override;
 
 	virtual int32 GetScratchAreaSize() const override;
