PoseSearch
- fix for DDC key determinism when supporting choosers
- improved how to run DDC related validations via cvars
- added new validation for PartialKeyHashes determinism

#Synced-CL 51142787

diff --git a/Engine/Plugins/Animation/PoseSearch/Source/Runtime/PoseSearch.Build.cs b/Engine/Plugins/Animation/PoseSearch/Source/Runtime/PoseSearch.Build.cs
--- a/Engine/Plugins/Animation/PoseSearch/Source/Runtime/PoseSearch.Build.cs
+++ b/Engine/Plugins/Animation/PoseSearch/Source/Runtime/PoseSearch.Build.cs
@@ -44,6 +44,8 @@ public class PoseSearch : ModuleRules
 					"AnimationModifiers",
 					"AnimationBlueprintLibrary",
 					"DerivedDataCache",
+					"RigVM",
+					"RigVMDeveloper",
 					"UnrealEd"
 				}
 			);
diff --git a/Engine/Plugins/Animation/PoseSearch/Source/Runtime/Private/PoseSearchAssetIndexer.cpp b/Engine/Plugins/Animation/PoseSearch/Source/Runtime/Private/PoseSearchAssetIndexer.cpp
--- a/Engine/Plugins/Animation/PoseSearch/Source/Runtime/Private/PoseSearchAssetIndexer.cpp
+++ b/Engine/Plugins/Animation/PoseSearch/Source/Runtime/Private/PoseSearchAssetIndexer.cpp
@@ -19,10 +19,10 @@ namespace UE::PoseSearch
 
 #if ENABLE_ANIM_DEBUG
 static bool GVarMotionMatchTestDisableIndexerCaching = false;
-static FAutoConsoleVariableRef CVarMotionMatchTestDisableIndexerCaching(TEXT("a.MotionMatch.TestDisableIndexerCaching"), GVarMotionMatchTestDisableIndexerCaching, TEXT("Disable Motion Matching Indexer Caching"));
+static FAutoConsoleVariableRef CVarMotionMatchTestDisableIndexerCaching(TEXT("a.MotionMatch.Test.DisableIndexerCaching"), GVarMotionMatchTestDisableIndexerCaching, TEXT("Disable Motion Matching Indexer Caching"));
 
 static int32 GVarMotionMatchTestExtractPoseDeterminismNumIterations = 0;
-static FAutoConsoleVariableRef CVarMotionMatchTestExtractPoseDeterminismNumIterations(TEXT("a.MotionMatch.TestExtractPoseDeterminismNumIterations"), GVarMotionMatchTestExtractPoseDeterminismNumIterations, TEXT("Test Motion Matching ExtractPose Determinism via this NumIterations retries"));
+static FAutoConsoleVariableRef CVarMotionMatchTestExtractPoseDeterminismNumIterations(TEXT("a.MotionMatch.Test.ExtractPoseDeterminismNumIterations"), GVarMotionMatchTestExtractPoseDeterminismNumIterations, TEXT("Test Motion Matching ExtractPose Determinism via this NumIterations retries"));
 #endif // ENABLE_ANIM_DEBUG
 
 //////////////////////////////////////////////////////////////////////////
diff --git a/Engine/Plugins/Animation/PoseSearch/Source/Runtime/Private/PoseSearchDerivedData.cpp b/Engine/Plugins/Animation/PoseSearch/Source/Runtime/Private/PoseSearchDerivedData.cpp
--- a/Engine/Plugins/Animation/PoseSearch/Source/Runtime/Private/PoseSearchDerivedData.cpp
+++ b/Engine/Plugins/Animation/PoseSearch/Source/Runtime/Private/PoseSearchDerivedData.cpp
@@ -38,69 +38,65 @@ namespace UE::PoseSearch
 {
 #if ENABLE_ANIM_DEBUG
 
-enum EMotionMatchTestFlags
-{
-	// no additional tests will be performed
-	None = 0,
+static bool GVarMotionMatchTestInvalidateCache = false;
+static FAutoConsoleVariableRef CVarMotionMatchTestInvalidateCache(TEXT("a.MotionMatch.Test.InvalidateCache"), GVarMotionMatchTestInvalidateCache, TEXT("Invalidate DDC cache on a frame basis (to stress test DDC cancellation while tasks are flying if !WaitForTaskCompletion)"));
 
-	// cache will be invalidated every frame (to stress test DDC cancellation while tasks are flying if !WaitForTaskCompletion)
-	InvalidateCache = 1 << 0,
+static bool GVarMotionMatchTestWaitForTaskCompletion = false;
+static FAutoConsoleVariableRef CVarMotionMatchTestWaitForTaskCompletion(TEXT("a.MotionMatch.Test.WaitForTaskCompletion"), GVarMotionMatchTestWaitForTaskCompletion, TEXT("Wait synchronously for DDC tasks to complete"));
 
-	// cache will be invalidated once the flying tasks are ended
-	WaitForTaskCompletion = 1 << 1,
+static bool GVarMotionMatchTestForceIndexing = false;
+static FAutoConsoleVariableRef CVarMotionMatchTestForceIndexing(TEXT("a.MotionMatch.Test.ForceIndexing"), GVarMotionMatchTestForceIndexing, TEXT("we'll force the database re-indexing and compare the result SearchIndex with the one retrieved via DDC"));
 
-	// we'll force the database re-indexing and compare the result SearchIndex with the one retrieved via DDC
-	ForceIndexing = 1 << 2,
+static bool GVarMotionMatchTestDeterminismKDTreeConstruct = false;
+static FAutoConsoleVariableRef CVarMotionMatchTestDeterminismKDTreeConstruct(TEXT("a.MotionMatch.Test.Determinism.KDTreeConstruct"), GVarMotionMatchTestDeterminismKDTreeConstruct, TEXT("test KDTree Construct determinism"));
 
-	// test KDTree Construct determinism
-	TestKDTreeConstructDeterminism = 1 << 3,
+static bool GVarMotionMatchTestKDTreeConstruct = false;
+static FAutoConsoleVariableRef CVarMotionMatchTestKDTreeConstruct(TEXT("a.MotionMatch.Test.KDTreeConstruct"), GVarMotionMatchTestKDTreeConstruct, TEXT("validating the kdtree construction"));
 
-	// validating the kdtree construction
-	ValidateKDTreeConstruct = 1 << 4,
+static bool GVarMotionMatchTestDeterminismVPTreeConstruct = false;
+static FAutoConsoleVariableRef CVarMotionMatchTestDeterminismVPTreeConstruct(TEXT("a.MotionMatch.Test.Determinism.VPTreeConstruct"), GVarMotionMatchTestDeterminismVPTreeConstruct, TEXT("test VPTree Construct determinism"));
 
-	// test VPTree Construct determinism
-	TestVPTreeConstructDeterminism = 1 << 5,
+static bool GVarMotionMatchTestVPTreeConstruct = false;
+static FAutoConsoleVariableRef CVarMotionMatchTestVPTreeConstruct(TEXT("a.MotionMatch.Test.VPTreeConstruct"), GVarMotionMatchTestVPTreeConstruct, TEXT("validating the vptree construction"));
 
-	// validating the vptree construction
-	ValidateVPTreeConstruct = 1 << 6,
+static bool GVarMotionMatchTestDeterminismIndexDatabase = false;
+static FAutoConsoleVariableRef CVarMotionMatchTestDeterminismIndexDatabase(TEXT("a.MotionMatch.Test.Determinism.IndexDatabase"), GVarMotionMatchTestDeterminismIndexDatabase, TEXT("test IndexDatabase determinism"));
 
-	// test IndexDatabase determinism
-	TestIndexDatabaseDeterminism = 1 << 7,
+static bool GVarMotionMatchTestDeterminismPruneDuplicateValues = false;
+static FAutoConsoleVariableRef CVarMotionMatchTestDeterminismPruneDuplicateValues(TEXT("a.MotionMatch.Test.Determinism.PruneDuplicateValues"), GVarMotionMatchTestDeterminismPruneDuplicateValues, TEXT("test PruneDuplicateValues determinism"));
 
-	// test PruneDuplicateValues determinism
-	TestPruneDuplicateValuesDeterminism = 1 << 8,
+static bool GVarMotionMatchTestDeterminismPruneDuplicatePCAValues = false;
+static FAutoConsoleVariableRef CVarMotionMatchTestDeterminismPruneDuplicatePCAValues(TEXT("a.MotionMatch.Test.Determinism.PruneDuplicatePCAValues"), GVarMotionMatchTestDeterminismPruneDuplicatePCAValues, TEXT("test PruneDuplicatePCAValues determinism"));
 
-	// test PruneDuplicatePCAValues determinism
-	TestPruneDuplicatePCAValuesDeterminism = 1 << 9,
+static bool GVarMotionMatchTestDDCConsistency = false;
+static FAutoConsoleVariableRef CVarMotionMatchTestDDCConsistency(TEXT("a.MotionMatch.Test.DDCConsistency"), GVarMotionMatchTestDDCConsistency, TEXT("validating the data we gave to DDC is stored correctly"));
 
-	// validating the data we gave to DDC is stored correctly
-	ValidateDDC = 1 << 10,
+static bool GVarMotionMatchTestDeterminismSynchronizeWithExternalDependencies = false;
+static FAutoConsoleVariableRef CVarMotionMatchTestDeterminismSynchronizeWithExternalDependencies(TEXT("a.MotionMatch.Test.Determinism.SynchronizeWithExternalDependencies"), GVarMotionMatchTestDeterminismSynchronizeWithExternalDependencies, TEXT("validating SynchronizeWithExternalDependencies doesn't alter the database AnimationAssets order"));
 
-	// validating SynchronizeWithExternalDependencies doesn't alter the database AnimationAssets order
-	ValidateSynchronizeWithExternalDependenciesDeterminism = 1 << 11,
+static bool GVarMotionMatchTestDeterminismAssetSampler = false;
+static FAutoConsoleVariableRef CVarMotionMatchTestDeterminismAssetSampler(TEXT("a.MotionMatch.Test.Determinism.AssetSampler"), GVarMotionMatchTestDeterminismAssetSampler, TEXT("test FAnimationAssetSampler determinism"));
 
-	// test FAnimationAssetSampler determinism
-	TestAssetSamplerDeterminism = 1 << 12,
+static bool GVarMotionMatchTestDeterminismAssetSamplerFromPreviousExecution = false;
+static FAutoConsoleVariableRef CVarMotionMatchTestDeterminismAssetSamplerFromPreviousExecution(TEXT("a.MotionMatch.Test.Determinism.AssetSamplerFromPreviousExecution"), GVarMotionMatchTestDeterminismAssetSamplerFromPreviousExecution, TEXT("test FAnimationAssetSampler determinism across multiple editor executions. It'll store some bin files in \\Engine\\TestAssetSamplerDeterminism"));
 
-	// test FAnimationAssetSampler determinism across multiple editro executions. It'll store some bin files in \Engine\TestAssetSamplerDeterminism
-	TestAssetSamplerDeterminismFromPreviousExecution = 1 << 13,
+static bool GVarMotionMatchTestDeterminismDDCKey = false;
+static FAutoConsoleVariableRef CVarMotionMatchTestDeterminismDDCKey(TEXT("a.MotionMatch.Test.Determinism.DDCKey"), GVarMotionMatchTestDeterminismDDCKey, TEXT("test DDC key generation determinism"));
 
-	// test DDC key generation determinism
-	TestDDCKeyDeterminism = 1 << 14,
-};
+static int32 GVarMotionMatchTestNumIterations = 10;
+static FAutoConsoleVariableRef CVarMotionMatchTestNumIterations(TEXT("a.MotionMatch.Test.NumIterations"), GVarMotionMatchTestNumIterations, TEXT("Test Motion Matching Num Iterations"));
 
-static int32 GVarMotionMatchTestFlags = EMotionMatchTestFlags::None;
-static FAutoConsoleVariableRef CVarMotionMatchTestFlags(TEXT("a.MotionMatch.TestFlags"), GVarMotionMatchTestFlags, TEXT("Test Motion Matching using EMotionMatchTestFlags"));
+static bool GVarMotionMatchTestDeterminismPartialKeyHashesSaveFile = false;
+static FAutoConsoleVariableRef CVarMotionMatchTestDeterminismPartialKeyHashesSaveFile(TEXT("a.MotionMatch.Test.Determinism.PartialKeyHashesSaveFile"), GVarMotionMatchTestDeterminismPartialKeyHashesSaveFile, TEXT("save file \\Engine\\TestPartialKeyHashesDeterminism\\PartialKeyHashes.bin for testing PartialKeyHashes determinism across multiple editor executions"));
 
-static int32 GVarMotionMatchTestNumIterations = 10;
-static FAutoConsoleVariableRef CVarMotionMatchTestNumIterations(TEXT("a.MotionMatch.TestNumIterations"), GVarMotionMatchTestNumIterations, TEXT("Test Motion Matching Num Iterations"));
-static bool AnyTestFlags(int32 Flags) { return (GVarMotionMatchTestFlags & Flags) != 0; }
-#endif // ENABLE_ANIM_DEBUG
+static bool GVarMotionMatchTestDeterminismPartialKeyHashes = false;
+static FAutoConsoleVariableRef CVarMotionMatchTestDeterminismPartialKeyHashes(TEXT("a.MotionMatch.Test.Determinism.PartialKeyHashes"), GVarMotionMatchTestDeterminismPartialKeyHashes, TEXT("test PartialKeyHashes determinism across multiple editor executions from stored file \\Engine\\TestPartialKeyHashesDeterminism\\PartialKeyHashes.bin"));
 
 static bool GVarMotionMatchReindexCancelledDatabases = false;
 static FAutoConsoleVariableRef CVarMotionMatchReindexCancelledDatabases(TEXT("a.MotionMatch.ReindexCancelledDatabases"), GVarMotionMatchReindexCancelledDatabases, TEXT("Reindex Cancelled Databases"));
 
-// Experimental, this feature might be removed without warning, not for production use
+#endif // ENABLE_ANIM_DEBUG
+
 static bool GVarMotionMatchReindexAllReferencedDatabases = true;
 static FAutoConsoleVariableRef CVarMotionMatchReindexAllReferencedDatabases(TEXT("a.MotionMatch.ReindexAllReferencedDatabases"), GVarMotionMatchReindexAllReferencedDatabases, TEXT("Reindex All Referenced Databases"));
 
@@ -115,6 +111,94 @@ static FCookStatsManager::FAutoRegisterCallback RegisterCookStats([](FCookStatsM
 	});
 #endif // ENABLE_COOK_STATS
 
+#if ENABLE_ANIM_DEBUG
+struct FSerializablePartialKeyHashes
+{
+	FSerializablePartialKeyHashes() = default;
+
+	explicit FSerializablePartialKeyHashes(const FPartialKeyHashes& PartialKeyHashes)
+	{
+		Entries.Reserve(PartialKeyHashes.GetEntries().Num());
+		for (const FPartialKeyHashes::FPair& Pair : PartialKeyHashes.GetEntries())
+		{
+			if (UObject* Object = Pair.Key.ResolveObjectPtr())
+			{
+				FString ObjectName = Object->GetFullName();
+				if (Entries.Find(ObjectName))
+				{
+					UE_LOG(LogPoseSearch, Error, TEXT("FSerializablePartialKeyHashes duplicate asset name '%s'"), *ObjectName);
+				}
+				else
+				{
+					FEntry& Entry = Entries.Add(ObjectName);
+
+					Entry.Hash = Pair.Value.Hash;
+					Entry.Dependencies.Reserve(Pair.Value.Dependencies.Num());
+
+					for (const FObjectKey& DependencyKey : Pair.Value.Dependencies)
+					{
+						if (UObject* Dependency = DependencyKey.ResolveObjectPtr())
+						{
+							FString DependencyName = Dependency->GetFullName();
+							Entry.Dependencies.Add(DependencyName);
+						}
+						else
+						{
+							Entry.Dependencies.Add("");
+						}
+					}
+				}
+			}
+		}
+	}
+
+	// checking if the intersection of entries between this and Other are consistent
+	bool LogDifferencesAgains(const FSerializablePartialKeyHashes& Other) const
+	{
+		bool bSuccess = true;
+		for (const TPair<FString, FEntry>& Pair : Entries)
+		{
+			if (const FEntry* OtherEntry = Other.Entries.Find(Pair.Key))
+			{
+				if (Pair.Value.Hash != OtherEntry->Hash)
+				{
+					UE_LOG(LogPoseSearch, Error, TEXT("FSerializablePartialKeyHashes hash mismatch for asset '%s'"), *Pair.Key);
+					bSuccess = false;
+				}
+
+				if (Pair.Value.Dependencies != OtherEntry->Dependencies)
+				{
+					UE_LOG(LogPoseSearch, Warning, TEXT("FSerializablePartialKeyHashes Dependencies mismatch for asset '%s'"), *Pair.Key);
+					bSuccess = false;
+				}
+			}
+		}
+		return bSuccess;
+	}
+
+	friend FArchive& operator<<(FArchive& Ar, FSerializablePartialKeyHashes& SerializablePartialKeyHashes)
+	{
+		Ar << SerializablePartialKeyHashes.Entries;
+		return Ar;
+	}
+
+	struct FEntry
+	{
+		HashDigestType Hash;
+		TArray<FString> Dependencies;
+
+		friend FArchive& operator<<(FArchive& Ar, FEntry& Entry)
+		{
+			Ar << Entry.Hash;
+			Ar << Entry.Dependencies;
+			return Ar;
+		}
+	};
+
+	TMap<FString, FEntry> Entries;
+};
+#endif // ENABLE_ANIM_DEBUG
+
 typedef TSet<const UPoseSearchDatabase*, DefaultKeyFuncs<const UPoseSearchDatabase*>, TInlineSetAllocator<256>> FDatabaseSet;
 
 static void RecursivePopulateDependentDatabases(const UPoseSearchDatabase* Database, FDatabaseSet& DatabaseSet)
@@ -571,7 +655,7 @@ static bool PreprocessSearchIndexWeights(FSearchIndex& SearchIndex, const UPoseS
 static Eigen::ComputationInfo PreprocessSearchIndexPCAData(FSearchIndex& SearchIndex, int32 NumDimensions, int32 NumberOfPrincipalComponents, EPoseSearchMode PoseSearchMode)
 {
 #if ENABLE_ANIM_DEBUG
-	if (AnyTestFlags(EMotionMatchTestFlags::ValidateKDTreeConstruct))
+	if (GVarMotionMatchTestKDTreeConstruct)
 	{
 		// @todo: move this into a unit test.
 		// this code will fail with nanoflann 1.5.5
@@ -643,7 +727,7 @@ static Eigen::ComputationInfo PreprocessSearchIndexPCAData(FSearchIndex& SearchI
 			const ColMajorMatrix EigenVectors = EigenSolver.eigenvectors().real();
 
 #if ENABLE_ANIM_DEBUG
-			if (AnyTestFlags(EMotionMatchTestFlags::ValidateKDTreeConstruct) && NumberOfPrincipalComponents == NumDimensions)
+			if (GVarMotionMatchTestKDTreeConstruct && NumberOfPrincipalComponents == NumDimensions)
 			{
 				const RowMajorVector ReciprocalWeightsSqrt = MapWeightsSqrt.cwiseInverse();
 				const RowMajorMatrix ProjectedValues = CenteredValues * EigenVectors;
@@ -689,7 +773,7 @@ static Eigen::ComputationInfo PreprocessSearchIndexPCAData(FSearchIndex& SearchI
 			MapPCAValues = CenteredValues * PCAProjectionMatrix;
 
 #if ENABLE_ANIM_DEBUG
-			if (AnyTestFlags(EMotionMatchTestFlags::ValidateKDTreeConstruct) && NumberOfPrincipalComponents == NumDimensions)
+			if (GVarMotionMatchTestKDTreeConstruct && NumberOfPrincipalComponents == NumDimensions)
 			{
 				const RowMajorVector ReciprocalWeightsSqrt = MapWeightsSqrt.cwiseInverse();
 				for (Eigen::Index RowIndex = 0; RowIndex < MapValues.rows(); ++RowIndex)
@@ -738,7 +822,7 @@ static void PreprocessSearchIndexKDTree(FSearchIndex& SearchIndex, const UPoseSe
 
 #if ENABLE_ANIM_DEBUG
 		// testing kdtree Construct determinism
-		if (AnyTestFlags(EMotionMatchTestFlags::TestKDTreeConstructDeterminism))
+		if (GVarMotionMatchTestDeterminismKDTreeConstruct)
 		{
 			TAlignedArray<float> PCAValuesTest;
 			for (int32 Iteration = 0; Iteration < GVarMotionMatchTestNumIterations; ++Iteration)
@@ -756,7 +840,7 @@ static void PreprocessSearchIndexKDTree(FSearchIndex& SearchIndex, const UPoseSe
 			}
 		}
 
-		if (AnyTestFlags(EMotionMatchTestFlags::ValidateKDTreeConstruct))
+		if (GVarMotionMatchTestKDTreeConstruct)
 		{
 			// testing the KDTree is returning the proper searches for all the points in pca space
 			const int32 KDTreeQueryNumNeighbors = Database->KDTreeQueryNumNeighbors;
@@ -845,7 +929,7 @@ static void PreprocessSearchIndexVPTree(FSearchIndex& SearchIndex, const UPoseSe
 
 #if ENABLE_ANIM_DEBUG
 
-		if (AnyTestFlags(EMotionMatchTestFlags::TestVPTreeConstructDeterminism))
+		if (GVarMotionMatchTestDeterminismVPTreeConstruct)
 		{
 			if (!SearchIndex.VPTree.TestConstruct(DataSource))
 			{
@@ -866,7 +950,7 @@ static void PreprocessSearchIndexVPTree(FSearchIndex& SearchIndex, const UPoseSe
 			}
 		}
 
-		if (AnyTestFlags(EMotionMatchTestFlags::ValidateVPTreeConstruct))
+		if (GVarMotionMatchTestVPTreeConstruct)
 		{
 			// we can validate vantage point tree only if there are no duplicates poses (points)
 			if (Database->PosePruningSimilarityThreshold > 0.f)
@@ -900,7 +984,7 @@ static void PreprocessSearchIndexVPTree(FSearchIndex& SearchIndex, const UPoseSe
 			}
 			else
 			{
-				UE_LOG(LogPoseSearch, Warning, TEXT("PreprocessSearchIndexVPTree - cannot ValidateVPTreeConstruct for %s if there could be potential duplicate poses: set PosePruningSimilarityThreshold > 0 to enforce duplicate pruning"), *Database->GetName());
+				UE_LOG(LogPoseSearch, Warning, TEXT("PreprocessSearchIndexVPTree - cannot VPTreeConstruct for %s if there could be potential duplicate poses: set PosePruningSimilarityThreshold > 0 to enforce duplicate pruning"), *Database->GetName());
 			}
 		}
 #endif // ENABLE_ANIM_DEBUG
@@ -1771,7 +1855,7 @@ void FPoseSearchDatabaseAsyncCacheTask::OnGetComplete(UE::DerivedData::FCacheGet
 	
 	bool bForceBuildIndex = false;
 #if ENABLE_ANIM_DEBUG
-	bForceBuildIndex = AnyTestFlags(EMotionMatchTestFlags::ForceIndexing);
+	bForceBuildIndex = GVarMotionMatchTestForceIndexing;
 #endif // ENABLE_ANIM_DEBUG
 
 	if (Response.Status == EStatus::Error || bCacheCorrupted || bForceBuildIndex)
@@ -1991,7 +2075,7 @@ void FPoseSearchDatabaseAsyncCacheTask::OnGetComplete(UE::DerivedData::FCacheGet
 					}
 
 #if ENABLE_ANIM_DEBUG
-					if (AnyTestFlags(EMotionMatchTestFlags::TestIndexDatabaseDeterminism))
+					if (GVarMotionMatchTestDeterminismIndexDatabase)
 					{
 						for (int32 Iteration = 0; Iteration < GVarMotionMatchTestNumIterations; ++Iteration)
 						{
@@ -2014,7 +2098,7 @@ void FPoseSearchDatabaseAsyncCacheTask::OnGetComplete(UE::DerivedData::FCacheGet
 				
 #if ENABLE_ANIM_DEBUG
 				// testing PruneDuplicateValues determinism
-				if (AnyTestFlags(EMotionMatchTestFlags::TestPruneDuplicateValuesDeterminism))
+				if (GVarMotionMatchTestDeterminismPruneDuplicateValues)
 				{
 					const int32 NumIterations = GVarMotionMatchTestNumIterations;
 
@@ -2085,7 +2169,7 @@ void FPoseSearchDatabaseAsyncCacheTask::OnGetComplete(UE::DerivedData::FCacheGet
 
 #if ENABLE_ANIM_DEBUG
 				// testing PruneDuplicatePCAValues determinism
-				if (AnyTestFlags(EMotionMatchTestFlags::TestPruneDuplicatePCAValuesDeterminism))
+				if (GVarMotionMatchTestDeterminismPruneDuplicatePCAValues)
 				{
 					FSearchIndex TestSearchIndexA = SearchIndex;
 					TestSearchIndexA.PruneDuplicatePCAValues(MainDatabase->PCAValuesPruningSimilarityThreshold, MainDatabase->GetNumberOfPrincipalComponents());
@@ -2189,7 +2273,7 @@ void FPoseSearchDatabaseAsyncCacheTask::OnGetComplete(UE::DerivedData::FCacheGet
 							UE_LOG(LogPoseSearch, Log, TEXT("%s - %s BuildIndex stored to DDC"), *LexToString(FullIndexKey.Hash), *MainDatabaseName);
 
 #if ENABLE_ANIM_DEBUG
-							if (AnyTestFlags(EMotionMatchTestFlags::ValidateDDC))
+							if (GVarMotionMatchTestDDCConsistency)
 							{
 								const FCacheKey CacheKey{ Bucket, DerivedDataKey };
 								const FCacheGetRequest CacheRequest = { { Database->GetPathName() }, CacheKey, ECachePolicy::Default };
@@ -2470,8 +2554,48 @@ void FAsyncPoseSearchDatabasesManagement::Tick(float DeltaTime)
 
 #if ENABLE_ANIM_DEBUG
 	// testing sampler determinism
-	const bool bTestAssetSamplerDeterminism = AnyTestFlags(EMotionMatchTestFlags::TestAssetSamplerDeterminism);
-	const bool bTestAssetSamplerDeterminismFromPreviousExecution = AnyTestFlags(EMotionMatchTestFlags::TestAssetSamplerDeterminismFromPreviousExecution);
+	if (GVarMotionMatchTestDeterminismPartialKeyHashesSaveFile || GVarMotionMatchTestDeterminismPartialKeyHashes)
+	{
+		FSerializablePartialKeyHashes SerializablePartialKeyHashes(PartialKeyHashes);
+		FString FileName = FString::Printf(TEXT("%s//TestPartialKeyHashesDeterminism//PartialKeyHashes.bin"), *FPaths::EngineDir());
+
+		if (GVarMotionMatchTestDeterminismPartialKeyHashesSaveFile)
+		{
+			TArray<uint8> RawBytes;
+			RawBytes.Reserve(1024 * 1024);
+			FMemoryWriter Writer(RawBytes);
+			Writer << SerializablePartialKeyHashes;
+			FFileHelper::SaveArrayToFile(RawBytes, *FileName);
+
+			GVarMotionMatchTestDeterminismPartialKeyHashesSaveFile = false;
+			UE_LOG(LogPoseSearch, Log, TEXT("FAsyncPoseSearchDatabasesManagement::Tick PartialKeyHashes serialized in '%s'"), *FileName);
+		}
+
+		if (GVarMotionMatchTestDeterminismPartialKeyHashes)
+		{
+			TArray<uint8> LoadedData;
+			const bool bLoadedFile = FFileHelper::LoadFileToArray(LoadedData, *FileName, FILEREAD_Silent);
+			if (bLoadedFile)
+			{
+				FMemoryReaderView Reader(LoadedData);
+				FSerializablePartialKeyHashes SerializablePartialKeyHashesFromFile;
+				Reader << SerializablePartialKeyHashesFromFile;
+				if (SerializablePartialKeyHashes.LogDifferencesAgains(SerializablePartialKeyHashesFromFile))
+				{
+					UE_LOG(LogPoseSearch, Log, TEXT("FAsyncPoseSearchDatabasesManagement::Tick PartialKeyHashes determinism validation succeeded"));
+				}
+			}
+			else
+			{
+				UE_LOG(LogPoseSearch, Error, TEXT("FAsyncPoseSearchDatabasesManagement::Tick Failed to load '%s' with serialized PartialKeyHashes"), *FileName);
+			}
+			GVarMotionMatchTestDeterminismPartialKeyHashes = false;
+		}
+	}
+
+	// testing sampler determinism
+	const bool bTestAssetSamplerDeterminism = GVarMotionMatchTestDeterminismAssetSampler;
+	const bool bTestAssetSamplerDeterminismFromPreviousExecution = GVarMotionMatchTestDeterminismAssetSamplerFromPreviousExecution;
 	if (bTestAssetSamplerDeterminism || bTestAssetSamplerDeterminismFromPreviousExecution)
 	{
 		struct FTestSample
@@ -2654,9 +2778,9 @@ void FAsyncPoseSearchDatabasesManagement::Tick(float DeltaTime)
 		}
 	}
 
-	if (AnyTestFlags(EMotionMatchTestFlags::InvalidateCache))
+	if (GVarMotionMatchTestInvalidateCache)
 	{
-		if (AnyTestFlags(EMotionMatchTestFlags::WaitForTaskCompletion))
+		if (GVarMotionMatchTestWaitForTaskCompletion)
 		{
 			// iterating backwards because of the possible RemoveAtSwap 
 			for (int32 TaskIndex = Tasks.Num() - 1; TaskIndex >= 0; --TaskIndex)
@@ -2679,7 +2803,7 @@ void FAsyncPoseSearchDatabasesManagement::Tick(float DeltaTime)
 		}
 	}
 
-	if (AnyTestFlags(EMotionMatchTestFlags::ValidateSynchronizeWithExternalDependenciesDeterminism))
+	if (GVarMotionMatchTestDeterminismSynchronizeWithExternalDependencies)
 	{
 		for (int32 TaskIndex = 0; TaskIndex < Tasks.Num(); ++TaskIndex)
 		{
@@ -2704,10 +2828,12 @@ void FAsyncPoseSearchDatabasesManagement::Tick(float DeltaTime)
 			UE_LOG(LogPoseSearch, Log, TEXT("%s - %s Removed because is Cancelled"), *LexToString(Tasks[TaskIndex]->GetDerivedDataKey()), *GetDatabaseAndOuterName(TaskDatabase));
 			Tasks.RemoveAtSwap(TaskIndex, EAllowShrinking::No);
 
+#if ENABLE_ANIM_DEBUG
 			if (GVarMotionMatchReindexCancelledDatabases)
 			{
 				RequestAsyncBuildIndex(TaskDatabase, ERequestAsyncBuildFlag::NewRequest);
 			}
+#endif // ENABLE_ANIM_DEBUG
 		}
 		else
 		{
@@ -2716,7 +2842,7 @@ void FAsyncPoseSearchDatabasesManagement::Tick(float DeltaTime)
 	}
 	
 #if ENABLE_ANIM_DEBUG
-	if (AnyTestFlags(EMotionMatchTestFlags::TestDDCKeyDeterminism))
+	if (GVarMotionMatchTestDeterminismDDCKey)
 	{
 		for (int32 TaskIndex = 0; TaskIndex < Tasks.Num(); ++TaskIndex)
 		{
diff --git a/Engine/Plugins/Animation/PoseSearch/Source/Runtime/Private/PoseSearchDerivedDataKey.cpp b/Engine/Plugins/Animation/PoseSearch/Source/Runtime/Private/PoseSearchDerivedDataKey.cpp
--- a/Engine/Plugins/Animation/PoseSearch/Source/Runtime/Private/PoseSearchDerivedDataKey.cpp
+++ b/Engine/Plugins/Animation/PoseSearch/Source/Runtime/Private/PoseSearchDerivedDataKey.cpp
@@ -10,12 +10,16 @@
 #include "Animation/IAnimationSequenceCompiler.h"
 #include "AnimationModifier.h"
 #include "Components/SkeletalMeshComponent.h"
+#include "EdGraph/EdGraphNode.h"
 #include "Engine/SkeletalMesh.h"
 #include "Engine/SkinnedAsset.h"
 #include "Engine/StreamableRenderAsset.h"
 #include "Factories/FbxAnimSequenceImportData.h"
 #include "PoseSearch/PoseSearchAnimNotifies.h"
 #include "PoseSearch/PoseSearchDefines.h"
+#include "RigVMCore/RigVMGraphFunctionHost.h"
+#include "RigVMHost.h"
+#include "RigVMModel/RigVMClient.h"
 #include "UObject/DevObjectVersion.h"
 
 namespace UE::PoseSearch
@@ -198,6 +202,8 @@ FArchive& FKeyBuilder::operator<<(FText& Value)
 
 FArchive& FKeyBuilder::TryAddDependency(UObject* Object)
 {
+	check(Object);
+
 	// @todo: add RF_NeedPostLoadSubobjects, RF_NeedInitialization, RF_NeedLoad, RF_WillBeLoaded? 
 	if (Object->HasAnyFlags(RF_NeedPostLoad))
 	{
@@ -497,6 +503,21 @@ bool FKeyBuilder::IsExcludedType(const UObject* Object)
 		return true;
 	}
 
+	if (Object->IsA<UEdGraphNode>())
+	{
+		return true;
+	}
+	
+	if (Object->IsA<URigVMHost>() || Object->IsA<URigVM>())
+	{
+		return true;
+	}
+
+	if (Cast<IRigVMClientHost>(Object) || Cast<IRigVMGraphFunctionHost>(Object) || Cast<IRigVMClientExternalModelHost>(Object))
+	{
+		return true;
+	}
+
 	return false;
 }
 
diff --git a/Engine/Plugins/Animation/PoseSearch/Source/Runtime/Public/PoseSearch/PoseSearchDerivedDataKey.h b/Engine/Plugins/Animation/PoseSearch/Source/Runtime/Public/PoseSearch/PoseSearchDerivedDataKey.h
--- a/Engine/Plugins/Animation/PoseSearch/Source/Runtime/Public/PoseSearch/PoseSearchDerivedDataKey.h
+++ b/Engine/Plugins/Animation/PoseSearch/Source/Runtime/Public/PoseSearch/PoseSearchDerivedDataKey.h
@@ -58,6 +58,8 @@ struct FPartialKeyHashes
 
 			return true;
 		}
+
+		friend FArchive& operator<<(FArchive& Ar, FEntry& Entry);
 	};
 
 	typedef TPair<FObjectKey, FEntry> FPair;
@@ -123,8 +125,7 @@ struct FPartialKeyHashes
 	{
 		return Entries.Find(ObjectKey);
 	}
-		
-private:
+	
 	struct FEntries : public TMap<FObjectKey, FEntry>
 	{
 		FPair* FindPair(const FObjectKey& ObjectKey)
@@ -133,6 +134,12 @@ private:
 		}
 	};
 
+	const FEntries& GetEntries() const
+	{
+		return Entries;
+	}
+
+private:
 	FEntries Entries;
 };
 
