PoseSearch - improved logging for reporting mismatch between MM column and related PSDs

#Synced-CL 51167957

diff --git a/Engine/Plugins/Animation/PoseSearch/Source/Runtime/Private/Chooser/PoseSearchChooserColumn.cpp b/Engine/Plugins/Animation/PoseSearch/Source/Runtime/Private/Chooser/PoseSearchChooserColumn.cpp
--- a/Engine/Plugins/Animation/PoseSearch/Source/Runtime/Private/Chooser/PoseSearchChooserColumn.cpp
+++ b/Engine/Plugins/Animation/PoseSearch/Source/Runtime/Private/Chooser/PoseSearchChooserColumn.cpp
@@ -43,7 +43,6 @@ namespace UE::PoseSearch
 		return nullptr;
 	}
 
-#if WITH_EDITORONLY_DATA
 	static const FInstancedStruct& GetResult(const UChooserTable* Chooser, int32 RowIndex)
 	{
 		check(Chooser);
@@ -52,14 +51,17 @@ namespace UE::PoseSearch
 			return Chooser->FallbackResult;
 		}
 
+#if WITH_EDITORONLY_DATA
 		if (!Chooser->IsCookedData())
 		{
 			return Chooser->ResultsStructs[RowIndex];
 		}
+#endif // WITH_EDITORONLY_DATA
 		
 		return Chooser->CookedResults[RowIndex];
 	}
 
+#if WITH_EDITOR
 	static UObject* GetReferencedObject(const UChooserTable* Chooser, int32 RowIndex)
 	{
 		const FInstancedStruct& Result = GetResult(Chooser, RowIndex);
@@ -69,8 +71,7 @@ namespace UE::PoseSearch
 		}
 		return nullptr;
 	}
-
-#endif // WITH_EDITORONLY_DATA
+#endif // WITH_EDITOR
 
 	struct FPoseSearchColumnScratchArea
 	{
@@ -286,9 +287,13 @@ namespace UE::PoseSearch
 			{
 				const FTableCell& TableCell = (*TableCells)[IndexInDatabase];
 				DatabaseAnimationAsset = &TableCell.Column->GetRowValue(TableCell.RowIndex).Data;
-#if WITH_EDITORONLY_DATA
-				ensure(GetReferencedObject(TableCell.Column->GetChooser(), TableCell.RowIndex) == DatabaseAnimationAsset->AnimAsset);
-#endif // WITH_EDITORONLY_DATA
+#if WITH_EDITOR
+				const UChooserTable* Chooser = TableCell.Column->GetChooser();
+				if (!ensure(GetReferencedObject(Chooser, TableCell.RowIndex) == DatabaseAnimationAsset->AnimAsset))
+				{
+					UE_LOG(LogPoseSearch, Error, TEXT("FMapping::FindDatabaseAnimationAsset - Mismatch between chooser table '%s' with root '%s' and Database '%s' for row %d"), *Chooser->GetName(), *Chooser->GetRootChooser()->GetName(), *Database->GetName(), TableCell.RowIndex);
+				}
+#endif // WITH_EDITOR
 			}
 		}
 		return DatabaseAnimationAsset;
@@ -1089,7 +1094,10 @@ void FPoseSearchColumn::AddToDetails(FInstancedPropertyBag& PropertyBag, int32 C
 	PropertyBag.AddProperties({PropertyDesc});
 
 	const FPoseSearchColumnRow& RowValue = GetRowValue(RowIndex);
-	ensure(GetReferencedObject(GetChooser(), RowIndex) == RowValue.Data.AnimAsset);
+	if (!ensure(GetReferencedObject(GetChooser(), RowIndex) == RowValue.Data.AnimAsset))
+	{
+		UE_LOG(LogPoseSearch, Error, TEXT("FPoseSearchColumn::AddToDetails - Mismatch between chooser table '%s' with root '%s' and Database '%s' for row %d"), *GetChooser()->GetName(), *GetChooser()->GetRootChooser()->GetName(), *GetNameSafe(RowValue.Database), RowIndex);
+	}
 
 	FPoseSearchColumnRowReflection RowReflection;
 	RowReflection.bDisableReselection = RowValue.Data.bDisableReselection;
