PoseSearch - fix for static analysis warning C6011: Dereferencing NULL pointer

#Synced-CL 51167957

#Synced-CL 51167957

diff --git a/Engine/Plugins/Animation/PoseSearch/Source/Runtime/Private/Chooser/PoseSearchChooserColumn.cpp b/Engine/Plugins/Animation/PoseSearch/Source/Runtime/Private/Chooser/PoseSearchChooserColumn.cpp
--- a/Engine/Plugins/Animation/PoseSearch/Source/Runtime/Private/Chooser/PoseSearchChooserColumn.cpp
+++ b/Engine/Plugins/Animation/PoseSearch/Source/Runtime/Private/Chooser/PoseSearchChooserColumn.cpp
@@ -280,20 +280,27 @@ namespace UE::PoseSearch
 	const FPoseSearchDatabaseAnimationAsset* FMapping::FindDatabaseAnimationAsset(int32 IndexInDatabase, const UPoseSearchDatabase* Database) const
 	{
 		const FPoseSearchDatabaseAnimationAsset* DatabaseAnimationAsset = nullptr;
-		ensure(IndexInDatabase >= 0 && Database);
-		if (const FTableCells* TableCells = Map.Find(Database))
+		if (ensure(IndexInDatabase >= 0 && Database))
 		{
-			if (TableCells->IsValidIndex(IndexInDatabase))
+			if (const FTableCells* TableCells = Map.Find(Database))
 			{
-				const FTableCell& TableCell = (*TableCells)[IndexInDatabase];
-				DatabaseAnimationAsset = &TableCell.Column->GetRowValue(TableCell.RowIndex).Data;
-#if WITH_EDITOR
-				const UChooserTable* Chooser = TableCell.Column->GetChooser();
-				if (!ensure(GetReferencedObject(Chooser, TableCell.RowIndex) == DatabaseAnimationAsset->AnimAsset))
+				if (TableCells->IsValidIndex(IndexInDatabase))
 				{
-					UE_LOG(LogPoseSearch, Error, TEXT("FMapping::FindDatabaseAnimationAsset - Mismatch between chooser table '%s' with root '%s' and Database '%s' for row %d"), *Chooser->GetName(), *Chooser->GetRootChooser()->GetName(), *Database->GetName(), TableCell.RowIndex);
-				}
+					const FTableCell& TableCell = (*TableCells)[IndexInDatabase];
+					DatabaseAnimationAsset = &TableCell.Column->GetRowValue(TableCell.RowIndex).Data;
+#if WITH_EDITOR
+					if (DatabaseAnimationAsset)
+					{
+						if (const UChooserTable* Chooser = TableCell.Column->GetChooser())
+						{
+							if (!ensure(GetReferencedObject(Chooser, TableCell.RowIndex) == DatabaseAnimationAsset->AnimAsset))
+							{
+								UE_LOG(LogPoseSearch, Error, TEXT("FMapping::FindDatabaseAnimationAsset - Mismatch between chooser table '%s' with root '%s' and Database '%s' for row %d"), *Chooser->GetName(), *Chooser->GetRootChooser()->GetName(), *Database->GetName(), TableCell.RowIndex);
+							}
+						}
+					}
 #endif // WITH_EDITOR
+				}
 			}
 		}
 		return DatabaseAnimationAsset;
@@ -302,26 +309,28 @@ namespace UE::PoseSearch
 #if WITH_EDITORONLY_DATA
 	void FMapping::SetAnimationAssetAt(const FPoseSearchDatabaseAnimationAsset& AnimationAsset, int32 IndexInDatabase, const UPoseSearchDatabase* Database)
 	{
-		ensure(IndexInDatabase >= 0 && Database);
-		if (FTableCells* TableCells = Map.Find(Database))
+		if (ensure(IndexInDatabase >= 0 && Database))
 		{
-			if (TableCells->IsValidIndex(IndexInDatabase))
+			if (FTableCells* TableCells = Map.Find(Database))
 			{
-				FTableCell& TableCell = (*TableCells)[IndexInDatabase];
+				if (TableCells->IsValidIndex(IndexInDatabase))
+				{
+					FTableCell& TableCell = (*TableCells)[IndexInDatabase];
 
-				DO_ENSURE_CONSISTENCY(CheckConsistency(TableCell.Column->GetChooser()));
+					DO_ENSURE_CONSISTENCY(CheckConsistency(TableCell.Column->GetChooser()));
 
-				FPoseSearchDatabaseAnimationAsset& DatabaseAnimationAsset = TableCell.Column->EditRowValue(TableCell.RowIndex).Data;
-				DatabaseAnimationAsset = AnimationAsset;
+					FPoseSearchDatabaseAnimationAsset& DatabaseAnimationAsset = TableCell.Column->EditRowValue(TableCell.RowIndex).Data;
+					DatabaseAnimationAsset = AnimationAsset;
 
-				// enforcing RowValue.Data.AnimAsset to be equal to ReferencedObject
-				UObject* ReferencedObject = GetReferencedObject(TableCell.Column->GetChooser(), TableCell.RowIndex);
-				if (ReferencedObject != DatabaseAnimationAsset.AnimAsset)
-				{
-					DatabaseAnimationAsset.AnimAsset = ReferencedObject;
-				}
+					// enforcing RowValue.Data.AnimAsset to be equal to ReferencedObject
+					UObject* ReferencedObject = GetReferencedObject(TableCell.Column->GetChooser(), TableCell.RowIndex);
+					if (ReferencedObject != DatabaseAnimationAsset.AnimAsset)
+					{
+						DatabaseAnimationAsset.AnimAsset = ReferencedObject;
+					}
 
-				DO_ENSURE_CONSISTENCY(CheckConsistency(TableCell.Column->GetChooser()));
+					DO_ENSURE_CONSISTENCY(CheckConsistency(TableCell.Column->GetChooser()));
+				}
 			}
 		}
 	}
@@ -329,8 +338,18 @@ namespace UE::PoseSearch
 
 	int32 FMapping::GetRowIndex(int32 IndexInDatabase, const UPoseSearchDatabase* Database) const
 	{
-		ensure(IndexInDatabase >= 0 && Database);
-		return Map[Database][IndexInDatabase].RowIndex;
+		if (ensure(IndexInDatabase >= 0 && Database))
+		{
+			if (const FTableCells* TableCells = Map.Find(Database))
+			{
+				if (IndexInDatabase < TableCells->Num())
+				{
+					return (*TableCells)[IndexInDatabase].RowIndex;
+				}
+			}
+		}
+		ensure(false);
+		return INDEX_NONE;
 	}
 
 	void FMapping::ApplyMappingToAllColumns()
@@ -1094,9 +1113,9 @@ void FPoseSearchColumn::AddToDetails(FInstancedPropertyBag& PropertyBag, int32 C
 	PropertyBag.AddProperties({PropertyDesc});
 
 	const FPoseSearchColumnRow& RowValue = GetRowValue(RowIndex);
-	if (!ensure(GetReferencedObject(GetChooser(), RowIndex) == RowValue.Data.AnimAsset))
+	if (!ensure(GetReferencedObject(Chooser, RowIndex) == RowValue.Data.AnimAsset))
 	{
-		UE_LOG(LogPoseSearch, Error, TEXT("FPoseSearchColumn::AddToDetails - Mismatch between chooser table '%s' with root '%s' and Database '%s' for row %d"), *GetChooser()->GetName(), *GetChooser()->GetRootChooser()->GetName(), *GetNameSafe(RowValue.Database), RowIndex);
+		UE_LOG(LogPoseSearch, Error, TEXT("FPoseSearchColumn::AddToDetails - Mismatch between chooser table '%s' with root '%s' and Database '%s' for row %d"), *Chooser->GetName(), *Chooser->GetRootChooser()->GetName(), *GetNameSafe(RowValue.Database), RowIndex);
 	}
 
 	FPoseSearchColumnRowReflection RowReflection;
