PoseSearch - fix for PSD editor details not updating the database

#Synced-CL 51222883

diff --git a/Engine/Plugins/Animation/PoseSearch/Source/Editor/Private/PoseSearchDatabaseEditor.cpp b/Engine/Plugins/Animation/PoseSearch/Source/Editor/Private/PoseSearchDatabaseEditor.cpp
--- a/Engine/Plugins/Animation/PoseSearch/Source/Editor/Private/PoseSearchDatabaseEditor.cpp
+++ b/Engine/Plugins/Animation/PoseSearch/Source/Editor/Private/PoseSearchDatabaseEditor.cpp
@@ -2,13 +2,13 @@
 
 #include "PoseSearchDatabaseEditor.h"
 #include "AdvancedPreviewSceneModule.h"
+#include "Animation/AnimSequenceBase.h"
 #include "Framework/Application/SlateApplication.h"
 #include "GameFramework/WorldSettings.h"
 #include "IStructureDetailsView.h"
-#include "StructUtils/InstancedStruct.h"
-#include "PoseSearchDatabaseAssetBrowser.h"
 #include "Modules/ModuleManager.h"
 #include "PoseSearch/PoseSearchDerivedData.h"
+#include "PoseSearchDatabaseAssetBrowser.h"
 #include "PoseSearchDatabaseAssetTree.h"
 #include "PoseSearchDatabaseDataDetails.h"
 #include "PoseSearchDatabaseEditorCommands.h"
@@ -18,6 +18,7 @@
 #include "PoseSearchEditor.h"
 #include "PropertyEditorModule.h"
 #include "SPoseSearchDatabaseViewport.h"
+#include "StructUtils/InstancedStruct.h"
 #include "Styling/AppStyle.h"
 #include "Widgets/Docking/SDockTab.h"
 #include "Widgets/Input/SButton.h"
@@ -541,6 +542,50 @@ namespace UE::PoseSearch
 	
 	void FDatabaseEditor::OnFinishedChangingSelectionProperties(const FPropertyChangedEvent& PropertyChangedEvent)
 	{
+		if (UPoseSearchDatabase* Database = GetPoseSearchDatabase())
+		{
+			if (IsValid(Database))
+			{
+				bool bModified = false;
+				for (TWeakObjectPtr<UObject>& SelectedReflectionPtr : SelectionWidget.SelectedReflections)
+				{
+					if (UPoseSearchDatabaseAnimationAssetReflection* SelectionReflection = Cast<UPoseSearchDatabaseAnimationAssetReflection>(SelectedReflectionPtr.Pin().Get()))
+					{
+						// adding some edit constraints to SelectionReflection
+						if (UObject* AnimationAsset = SelectionReflection->DatabaseAnimationAsset.GetAnimationAsset())
+						{
+							if (AnimationAsset->IsA<UAnimSequenceBase>())
+							{
+								const float PlayLength = SelectionReflection->DatabaseAnimationAsset.GetPlayLength(FVector::ZeroVector);
+								const FFloatInterval ClampedSamplingRange(FMath::Clamp(SelectionReflection->DatabaseAnimationAsset.SamplingRange.Min, 0.0f, PlayLength), FMath::Clamp(SelectionReflection->DatabaseAnimationAsset.SamplingRange.Max, 0.0f, PlayLength));
+								if (ClampedSamplingRange != SelectionReflection->DatabaseAnimationAsset.SamplingRange)
+								{
+									SelectionReflection->DatabaseAnimationAsset.SamplingRange = ClampedSamplingRange;
+								}
+							}
+						}
+
+						if (const FPoseSearchDatabaseAnimationAsset* DatabaseAnimationAssetPtr = Database->GetDatabaseAnimationAsset(SelectionReflection->SourceAssetIdx))
+						{
+							if (*DatabaseAnimationAssetPtr != SelectionReflection->DatabaseAnimationAsset)
+							{
+								if (!bModified)
+								{
+									Database->Modify();
+									bModified = true;
+								}
+								Database->SetAnimationAssetAt(SelectionReflection->DatabaseAnimationAsset, SelectionReflection->SourceAssetIdx);
+							}
+						}
+					}
+				}
+
+				if (bModified)
+				{
+					Database->PostEditChange();
+				}
+			}
+		}
 	}
 
 	void FDatabaseEditor::OnAssetTreeSelectionChanged(
@@ -562,7 +607,7 @@ namespace UE::PoseSearch
 						UPoseSearchDatabaseAnimationAssetReflection* NewSelectionReflection = NewObject<UPoseSearchDatabaseAnimationAssetReflection>();
 						NewSelectionReflection->AddToRoot();
 						NewSelectionReflection->DatabaseAnimationAsset = *DatabaseAnimationAsset;
-						NewSelectionReflection->SetSourceLink(SelectedItem, AssetTreeWidget);
+						NewSelectionReflection->SourceAssetIdx = SelectedItem->SourceAssetIdx;
 						NewSelectionReflection->SetFlags(RF_Transactional);
 
 						GetInitializedSelectionWidget().SelectedReflections.Add(NewSelectionReflection);
diff --git a/Engine/Plugins/Animation/PoseSearch/Source/Editor/Private/PoseSearchDatabaseEditorReflection.cpp b/Engine/Plugins/Animation/PoseSearch/Source/Editor/Private/PoseSearchDatabaseEditorReflection.cpp
--- a/Engine/Plugins/Animation/PoseSearch/Source/Editor/Private/PoseSearchDatabaseEditorReflection.cpp
+++ b/Engine/Plugins/Animation/PoseSearch/Source/Editor/Private/PoseSearchDatabaseEditorReflection.cpp
@@ -1,7 +1,6 @@
 // Copyright Epic Games, Inc. All Rights Reserved.
 
 #include "PoseSearchDatabaseEditorReflection.h"
-#include "Animation/AnimSequenceBase.h"
 #include "PoseSearch/PoseSearchDerivedData.h"
 #include "PoseSearch/PoseSearchSchema.h"
 #include "PoseSearchDatabaseAssetTree.h"
@@ -13,51 +12,6 @@
 
 #define LOCTEXT_NAMESPACE "UPoseSearchDatabaseReflection"
 
-void UPoseSearchDatabaseReflectionBase::SetSourceLink(
-	const TWeakPtr<UE::PoseSearch::FDatabaseAssetTreeNode>& InWeakAssetTreeNode,
-	const TSharedPtr<UE::PoseSearch::SDatabaseAssetTree>& InAssetTreeWidget)
-{
-	WeakAssetTreeNode = InWeakAssetTreeNode;
-	AssetTreeWidget = InAssetTreeWidget;
-}
-
-bool UPoseSearchDatabaseAnimationAssetReflection::ApplyChanges() 
-{
-	if (UObject* AnimationAsset = DatabaseAnimationAsset.GetAnimationAsset())
-	{
-		if (AnimationAsset->IsA<UAnimSequenceBase>())
-		{
-			const float PlayLength = DatabaseAnimationAsset.GetPlayLength(FVector::ZeroVector);
-			const FFloatInterval ClampedSamplingRange(FMath::Clamp(DatabaseAnimationAsset.SamplingRange.Min, 0.0f, PlayLength), FMath::Clamp(DatabaseAnimationAsset.SamplingRange.Max, 0.0f, PlayLength));
-			if (ClampedSamplingRange != DatabaseAnimationAsset.SamplingRange)
-			{
-				DatabaseAnimationAsset.SamplingRange = ClampedSamplingRange;
-			}
-		}
-	}
-
-	if (const TSharedPtr<UE::PoseSearch::FDatabaseAssetTreeNode> AssetTreeNode = WeakAssetTreeNode.Pin())
-	{
-		if (const TSharedPtr<UE::PoseSearch::FDatabaseViewModel> ViewModel = AssetTreeNode->EditorViewModel.Pin())
-		{
-			UPoseSearchDatabase* Database = ViewModel->GetPoseSearchDatabase();
-			if (IsValid(Database))
-			{
-				if (const FPoseSearchDatabaseAnimationAsset* DatabaseAnimationAssetPtr = Database->GetDatabaseAnimationAsset(AssetTreeNode->SourceAssetIdx))
-				{
-					if (*DatabaseAnimationAssetPtr != DatabaseAnimationAsset)
-					{
-						Database->SetAnimationAssetAt(DatabaseAnimationAsset, AssetTreeNode->SourceAssetIdx);
-						return true;
-					}
-				}
-			}
-		}
-	}
-
-	return false;
-}
-
 void UPoseSearchDatabaseStatistics::Initialize(const UPoseSearchDatabase* PoseSearchDatabase)
 {
 	static FText TimeFormat = LOCTEXT("TimeFormat", "{0} {0}|plural(one=Second,other=Seconds)");
diff --git a/Engine/Plugins/Animation/PoseSearch/Source/Editor/Public/PoseSearchDatabaseEditorReflection.h b/Engine/Plugins/Animation/PoseSearch/Source/Editor/Public/PoseSearchDatabaseEditorReflection.h
--- a/Engine/Plugins/Animation/PoseSearch/Source/Editor/Public/PoseSearchDatabaseEditorReflection.h
+++ b/Engine/Plugins/Animation/PoseSearch/Source/Editor/Public/PoseSearchDatabaseEditorReflection.h
@@ -20,15 +20,7 @@ class UPoseSearchDatabaseReflectionBase : public UObject
 	GENERATED_BODY()
 
 public:
-	void SetSourceLink(
-		const TWeakPtr<UE::PoseSearch::FDatabaseAssetTreeNode>& InWeakAssetTreeNode,
-		const TSharedPtr<UE::PoseSearch::SDatabaseAssetTree>& InAssetTreeWidget);
-	
-	virtual bool ApplyChanges() { return false; };
-	
-protected:
-	TWeakPtr<UE::PoseSearch::FDatabaseAssetTreeNode> WeakAssetTreeNode;
-	TSharedPtr<UE::PoseSearch::SDatabaseAssetTree> AssetTreeWidget;
+	int32 SourceAssetIdx = INDEX_NONE;
 };
 
 UCLASS()
@@ -37,10 +29,8 @@ class UPoseSearchDatabaseAnimationAssetReflection : public UPoseSearchDatabaseRe
 	GENERATED_BODY()
 
 public:
-	UPROPERTY(EditAnywhere, Category = "Selected Animation Asset")
+	UPROPERTY(Transient, EditAnywhere, Category = "Selected Animation Asset")
 	FPoseSearchDatabaseAnimationAsset DatabaseAnimationAsset;
-	
-	virtual bool ApplyChanges() override;
 };
 
 USTRUCT()
