Fix possible null pointer

#rnx

#Synced-CL 51237423

diff --git a/Engine/Plugins/Animation/PoseSearch/Source/Runtime/Private/Chooser/PoseSearchChooserColumn.cpp b/Engine/Plugins/Animation/PoseSearch/Source/Runtime/Private/Chooser/PoseSearchChooserColumn.cpp
--- a/Engine/Plugins/Animation/PoseSearch/Source/Runtime/Private/Chooser/PoseSearchChooserColumn.cpp
+++ b/Engine/Plugins/Animation/PoseSearch/Source/Runtime/Private/Chooser/PoseSearchChooserColumn.cpp
@@ -117,7 +117,10 @@ namespace UE::PoseSearch
 
 	static void IterateChooserRows(UChooserTable* Chooser, FVisitedChoosers& VisitedChoosers, const FProcessRow& ProcessRow)
 	{
-		ensure(Chooser);
+		if (!ensure(Chooser))
+		{
+			return;
+		}
 
 		bool bAlreadyVisited = false;
 		VisitedChoosers.Add(Chooser, &bAlreadyVisited);
