PoseSearch - demoting check to ensure + LOG in FPoseHistory::EvaluateComponentSpace_AnyThread

#Synced-CL 51342120

diff --git a/Engine/Plugins/Animation/PoseSearch/Source/Runtime/Private/PoseSearchHistory.cpp b/Engine/Plugins/Animation/PoseSearch/Source/Runtime/Private/PoseSearchHistory.cpp
--- a/Engine/Plugins/Animation/PoseSearch/Source/Runtime/Private/PoseSearchHistory.cpp
+++ b/Engine/Plugins/Animation/PoseSearch/Source/Runtime/Private/PoseSearchHistory.cpp
@@ -422,7 +422,7 @@ FTransform FPoseHistoryEntry::GetComponentSpaceTransform(int32 Index) const
 		return FTransform(Quat, ComponentSpacePositions[Index], Scale);
 	}
 
-	UE_LOG(LogPoseSearch, Error, TEXT("FPoseHistoryEntry::GetComponentSpaceTransform - Index %d out of bound [0, %d)"), Index, ComponentSpaceRotations.Num());
+	UE_LOGF(LogPoseSearch, Error, "FPoseHistoryEntry::GetComponentSpaceTransform - Index %d out of bound [0, %d)", Index, ComponentSpaceRotations.Num());
 	return FTransform::Identity;
 }
 
@@ -433,7 +433,7 @@ float FPoseHistoryEntry::GetCurveValue(int32 Index) const
 		return CurveValues[Index];
 	}
 	
-	UE_LOG(LogPoseSearch, Error, TEXT("FPoseHistoryEntry::GetCurveValue - Index %d out of bound [0, %d)"), Index, CurveValues.Num());
+	UE_LOGF(LogPoseSearch, Error, "FPoseHistoryEntry::GetCurveValue - Index %d out of bound [0, %d)", Index, CurveValues.Num());
 	return 0.0f;
 }
 
@@ -837,7 +837,7 @@ bool FPoseHistory::GetTransformAtTime(float Time, FTransform& OutBoneTransform,
 //	// and we ultimately want to catch errors of missing bone transforms only checking valid USkeleton(s) against valid PoseSearchSchema(s)
 //	if (!bSuccess && BoneIndexSkeleton)
 //	{
-//		UE_LOG(LogPoseSearch, Error, TEXT("FPoseHistory::GetTransformAtTime - Couldn't find root bone transform at time %f for Skeleton %s!"), Time, *BoneIndexSkeleton->GetName());
+//		UE_LOGF(LogPoseSearch, Error, "FPoseHistory::GetTransformAtTime - Couldn't find root bone transform at time %f for Skeleton %ls!", Time, *BoneIndexSkeleton->GetName());
 //	}
 //#endif // WITH_EDITOR
 
@@ -982,9 +982,13 @@ void FPoseHistory::EvaluateComponentSpace_AnyThread(float DeltaTime, IComponentS
 #endif //WITH_EDITORONLY_DATA
 	)
 {
-	UE_MT_SCOPED_WRITE_ACCESS(PoseDataThreadSafeCounter);
+	if (!ensure(MaxNumPoses >= 2))
+	{
+		UE_LOGF(LogPoseSearch, Error, "FPoseHistory::EvaluateComponentSpace_AnyThread called before FPoseHistory::Initialize_AnyThread! Please double check the implementation of any FAnimNode_Base calling this API, correctly propagating Initialize_AnyThread to their FPoseLink(s)! To mitigate issues FPoseHistory will be set to to default values");
+		Initialize_AnyThread(2, 0.f);
+	}
 
-	check(MaxNumPoses >= 2);
+	UE_MT_SCOPED_WRITE_ACCESS(PoseDataThreadSafeCounter);
 
 	const USkeleton* Skeleton = ComponentSpacePoseProvider.GetSkeletonAsset();
 
@@ -1269,14 +1273,14 @@ void FMemStackPoseHistory::ExtractAndAddFuturePoses(const UAnimationAsset* Anima
 	check(FiniteDeltaTime >= 0.f);
 	if (AnimationTime < FiniteDeltaTime)
 	{
-		UE_LOG(LogPoseSearch, Error, TEXT("FMemStackPoseHistory::ExtractAndAddFuturePose - provided AnimationTime (%f) is too small. Clamping it to minimum value of %f"), AnimationTime, FiniteDeltaTime);
+		UE_LOGF(LogPoseSearch, Error, "FMemStackPoseHistory::ExtractAndAddFuturePose - provided AnimationTime (%f) is too small. Clamping it to minimum value of %f", AnimationTime, FiniteDeltaTime);
 		AnimationTime = FiniteDeltaTime;
 	}
 
 	if (IntervalTime < FiniteDeltaTime)
 	{
-		UE_LOG(LogPoseSearch, Error, TEXT("FMemStackPoseHistory::ExtractAndAddFuturePose - provided IntervalTime (%f) is too small. Clamping it to minimum value of %f"), IntervalTime, FiniteDeltaTime);
-		AnimationTime = FiniteDeltaTime;
+		UE_LOGF(LogPoseSearch, Error, "FMemStackPoseHistory::ExtractAndAddFuturePose - provided IntervalTime (%f) is too small. Clamping it to minimum value of %f", IntervalTime, FiniteDeltaTime);
+		IntervalTime = FiniteDeltaTime;
 	}
 
 	const USkeleton* Skeleton = OverrideSkeleton ? OverrideSkeleton : AnimationAsset->GetSkeleton();
