PoseSearch - PSIA and asset sampler library improvements to be able to edit MMI assets via editor utility BPs

#Synced-CL 51592296

diff --git a/Engine/Plugins/Animation/PoseSearch/Source/Runtime/Private/PoseSearchAssetSamplerLibrary.cpp b/Engine/Plugins/Animation/PoseSearch/Source/Runtime/Private/PoseSearchAssetSamplerLibrary.cpp
--- a/Engine/Plugins/Animation/PoseSearch/Source/Runtime/Private/PoseSearchAssetSamplerLibrary.cpp
+++ b/Engine/Plugins/Animation/PoseSearch/Source/Runtime/Private/PoseSearchAssetSamplerLibrary.cpp
@@ -10,48 +10,71 @@
 
 #include UE_INLINE_GENERATED_CPP_BY_NAME(PoseSearchAssetSamplerLibrary)
 
-FPoseSearchAssetSamplerPose UPoseSearchAssetSamplerLibrary::SamplePose(const UAnimInstance* AnimInstance, const FPoseSearchAssetSamplerInput Input)
+FPoseSearchAssetSamplerPose UPoseSearchAssetSamplerLibrary::SamplePose(const FPoseSearchAssetSamplerInput Input, const UAnimInstance* AnimInstance)
 {
-	FPoseSearchAssetSamplerPose AssetSamplerPose;
 	using namespace UE::PoseSearch;
 	if (!Input.Animation)
 	{
 		UE_LOG(LogPoseSearch, Error, TEXT("UPoseSearchAssetSamplerLibrary::SamplePose invalid Input.Animation"));
+		return FPoseSearchAssetSamplerPose();
 	}
-	else if (!AnimInstance)
+	
+	if (Input.bMirrored && !Input.MirrorDataTable)
 	{
-		UE_LOG(LogPoseSearch, Error, TEXT("UPoseSearchAssetSamplerLibrary::SamplePose invalid AnimInstance"));
+		UE_LOG(LogPoseSearch, Error, TEXT("UPoseSearchAssetSamplerLibrary::SamplePose unable to mirror the pose from %s at time %f because of invalid MirrorDataTable"), *Input.Animation->GetName(), Input.AnimationTime);
+		return FPoseSearchAssetSamplerPose();
 	}
-	else if (Input.bMirrored && !Input.MirrorDataTable)
+
+	if (AnimInstance)
 	{
-		UE_LOG(LogPoseSearch, Error, TEXT("UPoseSearchAssetSamplerLibrary::SamplePose unable to mirror the pose from %s at time %f because of invalid MirrorDataTable"), *Input.Animation->GetName(), Input.AnimationTime);
+		return SamplePose(AnimInstance->GetRequiredBonesOnAnyThread(), Input);
 	}
-	else
+	
+	const USkeleton* Skeleton = Input.Animation->GetSkeleton();
+	TArray<uint16> BoneIndices;
+
+	const int32 NumBones = Skeleton->GetReferenceSkeleton().GetNum();
+	BoneIndices.SetNum(NumBones);
+	for (int32 BoneIndex = 0; BoneIndex < NumBones; ++BoneIndex)
 	{
-		const FBoneContainer& BoneContainer = AnimInstance->GetRequiredBonesOnAnyThread();
+		BoneIndices[BoneIndex] = BoneIndex;
+	}
 
-		FMemMark Mark(FMemStack::Get());
+	FBoneContainer BoneContainer;
+	BoneContainer.InitializeTo(BoneIndices, UE::Anim::FCurveFilterSettings(UE::Anim::ECurveFilterMode::DisallowAll), *Skeleton);
 
-		bool bPreProcessRootTransform = true;
-		const FAnimationAssetSampler Sampler(Input.Animation, Input.RootTransformOrigin, Input.BlendParameters, Input.RootTransformSamplingRate, bPreProcessRootTransform);
+	return SamplePose(BoneContainer, Input);
+}
 
-		FBlendedCurve Curve;
-		FCompactPose Pose;
-		Pose.SetBoneContainer(&BoneContainer);
+FPoseSearchAssetSamplerPose UPoseSearchAssetSamplerLibrary::SamplePose(const FBoneContainer& BoneContainer, const FPoseSearchAssetSamplerInput Input)
+{
+	using namespace UE::PoseSearch;
 
-		Sampler.ExtractPose(Input.AnimationTime, Pose, Curve);
-		AssetSamplerPose.RootTransform = Sampler.ExtractRootTransform(Input.AnimationTime);
+	FMemMark Mark(FMemStack::Get());
 
-		if (Input.bMirrored)
-		{
-			const FMirrorDataCache MirrorDataCache(Input.MirrorDataTable, BoneContainer);
-			MirrorDataCache.MirrorPose(Pose);
-			AssetSamplerPose.RootTransform = MirrorDataCache.MirrorTransform(AssetSamplerPose.RootTransform);
-		}
+	bool bPreProcessRootTransform = true;
+	const FAnimationAssetSampler Sampler(Input.Animation, Input.RootTransformOrigin, Input.BlendParameters, Input.RootTransformSamplingRate, bPreProcessRootTransform, false);
+
+	FPoseSearchAssetSamplerPose AssetSamplerPose;
+	AssetSamplerPose.BoneContainer = BoneContainer;
+
+	FBlendedCurve Curve;
+	FCompactPose Pose;
+	Pose.SetBoneContainer(&AssetSamplerPose.BoneContainer);
+	Sampler.ExtractPose(Input.AnimationTime, Pose, Curve);
+	AssetSamplerPose.RootTransform = Sampler.ExtractRootTransform(Input.AnimationTime);
 
-		AssetSamplerPose.Pose.CopyBonesFrom(Pose);
-		AssetSamplerPose.ComponentSpacePose.InitPose(AssetSamplerPose.Pose);
+	if (Input.bMirrored)
+	{
+		const FMirrorDataCache MirrorDataCache(Input.MirrorDataTable, AssetSamplerPose.BoneContainer);
+		MirrorDataCache.MirrorPose(Pose);
+		AssetSamplerPose.RootTransform = MirrorDataCache.MirrorTransform(AssetSamplerPose.RootTransform);
 	}
+
+	AssetSamplerPose.Pose.CopyBonesFrom(Pose);
+	AssetSamplerPose.ComponentSpacePose.InitPose(AssetSamplerPose.Pose);
+	check(AssetSamplerPose.ComponentSpacePose.GetPose().IsValid());
+
 	return AssetSamplerPose;
 }
 
diff --git a/Engine/Plugins/Animation/PoseSearch/Source/Runtime/Public/PoseSearch/PoseSearchAssetSamplerLibrary.h b/Engine/Plugins/Animation/PoseSearch/Source/Runtime/Public/PoseSearch/PoseSearchAssetSamplerLibrary.h
--- a/Engine/Plugins/Animation/PoseSearch/Source/Runtime/Public/PoseSearch/PoseSearchAssetSamplerLibrary.h
+++ b/Engine/Plugins/Animation/PoseSearch/Source/Runtime/Public/PoseSearch/PoseSearchAssetSamplerLibrary.h
@@ -49,7 +49,31 @@ struct FPoseSearchAssetSamplerPose
 {
 	GENERATED_BODY()
 
+	FPoseSearchAssetSamplerPose() = default;
+
+	FPoseSearchAssetSamplerPose(const FPoseSearchAssetSamplerPose& Other)
+	: RootTransform(Other.RootTransform)
+	, BoneContainer(Other.BoneContainer)
+	, Pose(Other.Pose)
+	, Curve(Other.Curve)
+	, ComponentSpacePose(Other.ComponentSpacePose)
+	{
+		Pose.SetBoneContainer(&BoneContainer);
+		const_cast<FCompactHeapPose&>(ComponentSpacePose.GetPose()).SetBoneContainer(&BoneContainer);
+	}
+
+	FPoseSearchAssetSamplerPose& operator=(const FPoseSearchAssetSamplerPose& Other)
+	{
+		if (this != &Other)
+		{
+			this->~FPoseSearchAssetSamplerPose();
+			new(this) FPoseSearchAssetSamplerPose(Other);
+		}
+		return *this;
+	}
+
 	FTransform RootTransform = FTransform::Identity;
+	FBoneContainer BoneContainer;
 	FCompactHeapPose Pose;
 	FBlendedHeapCurve Curve;
 	// @todo: add Attribute(s)
@@ -72,8 +96,10 @@ class UPoseSearchAssetSamplerLibrary : public UBlueprintFunctionLibrary
 	GENERATED_BODY()
 
 public:
-	UFUNCTION(BlueprintPure, Category = "Animation|Pose Search|Experimental", meta = (BlueprintThreadSafe))
-	static UE_API FPoseSearchAssetSamplerPose SamplePose(const UAnimInstance* AnimInstance, const FPoseSearchAssetSamplerInput Input);
+	UFUNCTION(BlueprintPure, Category = "Animation|Pose Search|Experimental", meta = (BlueprintThreadSafe, AdvancedDisplay = "AnimInstance"))
+	static UE_API FPoseSearchAssetSamplerPose SamplePose(const FPoseSearchAssetSamplerInput Input, const UAnimInstance* AnimInstance = nullptr);
+
+	static UE_API FPoseSearchAssetSamplerPose SamplePose(const FBoneContainer& BoneContainer, const FPoseSearchAssetSamplerInput Input);
 
 	UFUNCTION(BlueprintPure, Category = "Animation|Pose Search|Experimental", meta = (BlueprintThreadSafe))
 	static UE_API FTransform GetTransformByName(UPARAM(ref) FPoseSearchAssetSamplerPose& AssetSamplerPose, FName BoneName, EPoseSearchAssetSamplerSpace Space = EPoseSearchAssetSamplerSpace::World);
diff --git a/Engine/Plugins/Animation/PoseSearch/Source/Runtime/Public/PoseSearch/PoseSearchInteractionAsset.h b/Engine/Plugins/Animation/PoseSearch/Source/Runtime/Public/PoseSearch/PoseSearchInteractionAsset.h
--- a/Engine/Plugins/Animation/PoseSearch/Source/Runtime/Public/PoseSearch/PoseSearchInteractionAsset.h
+++ b/Engine/Plugins/Animation/PoseSearch/Source/Runtime/Public/PoseSearch/PoseSearchInteractionAsset.h
@@ -9,33 +9,33 @@
 
 #define UE_API POSESEARCH_API
 
-USTRUCT(Experimental)
+USTRUCT(Experimental, BlueprintType)
 struct FPoseSearchInteractionAssetItem
 {
 	GENERATED_BODY()
 
 	// associated aniamtion for this FPoseSearchInteractionAssetItem
-	UPROPERTY(EditAnywhere, Category = "Settings")
+	UPROPERTY(EditAnywhere, BlueprintReadWrite, Category = "Settings")
 	TObjectPtr<UAnimationAsset> Animation;
 
 	// associated role for this FPoseSearchInteractionAssetItem
-	UPROPERTY(EditAnywhere, Category = "Settings")
+	UPROPERTY(EditAnywhere, BlueprintReadWrite, Category = "Settings")
 	FName Role;
 
 	// relative weight to the other FPoseSearchInteractionAssetItem::WarpingWeightRotation(s) defining which character will be rotated while warping
 	// 0 - the associated character to this item will move fully to compensate the warping errors
 	// > 0 && all the other FPoseSearchInteractionAssetItem::WarpingWeightTranslation as zero, and the associated character will not move
-	UPROPERTY(EditAnywhere, Category = "Warping", meta = (ClampMin = "0", ClampMax = "1"))
+	UPROPERTY(EditAnywhere, BlueprintReadWrite, Category = "Warping", meta = (ClampMin = "0", ClampMax = "1"))
 	float WarpingWeightRotation = 0.5f;
 
 	// relative weight to the other FPoseSearchInteractionAssetItem::WarpingWeightTranslation(s) defining which character will be translated while warping
 	// 0 - the associated character to this item will move fully to compensate the warping errors
 	// > 0 && all the other FPoseSearchInteractionAssetItem::WarpingWeightTranslation as zero, and the associated character will not move
-	UPROPERTY(EditAnywhere, Category = "Warping", meta = (ClampMin = "0", ClampMax = "1"))
+	UPROPERTY(EditAnywhere, BlueprintReadWrite, Category = "Warping", meta = (ClampMin = "0", ClampMax = "1"))
 	float WarpingWeightTranslation = 0.5f;
 
 	// offset from the origin
-	UPROPERTY(EditAnywhere, Category = "Settings")
+	UPROPERTY(EditAnywhere, BlueprintReadWrite, Category = "Settings")
 	FTransform Origin = FTransform::Identity;
 
 #if WITH_EDITORONLY_DATA
@@ -50,19 +50,19 @@ class UPoseSearchInteractionAsset : public UMultiAnimAsset
 {
 	GENERATED_BODY()
 
-	UPROPERTY(EditAnywhere, Category = "Settings")
+public:
+	UPROPERTY(EditAnywhere, BlueprintReadWrite, Category = "Settings")
 	TArray<FPoseSearchInteractionAssetItem> Items;
 
-	UPROPERTY(EditAnywhere, Category = "Settings", meta = (ClampMin = "0", ClampMax = "1"))
+	UPROPERTY(EditAnywhere, BlueprintReadWrite, Category = "Settings", meta = (ClampMin = "0", ClampMax = "1"))
 	float WarpingBankingWeight = 0.f;
 
-	UPROPERTY(EditAnywhere, Category = "Settings", meta = (ClampMin = "0", ClampMax = "1"))
+	UPROPERTY(EditAnywhere, BlueprintReadWrite, Category = "Settings", meta = (ClampMin = "0", ClampMax = "1"))
 	float MinimalTranslationWeight = 0.f;
 
-public:
 #if WITH_EDITORONLY_DATA
 
-	UPROPERTY(Transient, EditAnywhere, Category = "Debug", meta = (EditCondition=bEnableDebugWarp, EditConditionHides))
+	UPROPERTY(Transient, BlueprintReadWrite, Category = "Debug", meta = (EditCondition=bEnableDebugWarp, EditConditionHides))
 	TArray<FTransform> DebugWarpOffsets;
 
 	// used to test warping: 0 - no warping applied, 1 - full warping/alignment applied
@@ -70,14 +70,14 @@ public:
 	// UMultiAnimAsset::GetOrigin() definition and warped accordingly with CalculateWarpTransforms
 	// following the rotation and translation weights defined in Items::WarpingWeightRotation and
 	// Items::WarpingWeightTranslation as relative weights between the Items (they'll be normalized at runtime)
-	UPROPERTY(Transient, EditAnywhere, Category = "Debug", meta = (ClampMin = "0", ClampMax = "1", EditCondition=bEnableDebugWarp, EditConditionHides))
+	UPROPERTY(Transient, BlueprintReadWrite, Category = "Debug", meta = (ClampMin = "0", ClampMax = "1", EditCondition=bEnableDebugWarp, EditConditionHides))
 	float DebugWarpAmount = 0.f;
 
 	// useful to test TimeOffset
-	UPROPERTY(Transient, EditAnywhere, Category = "Debug", meta = (ClampMin = "0", EditCondition=bEnableDebugWarp, EditConditionHides))
+	UPROPERTY(Transient, BlueprintReadWrite, Category = "Debug", meta = (ClampMin = "0", EditCondition=bEnableDebugWarp, EditConditionHides))
 	float DebugTimeOffset = 0.f;
 
-	UPROPERTY(Transient, EditAnywhere, Category = "Debug")
+	UPROPERTY(Transient, BlueprintReadWrite, Category = "Debug")
 	bool bEnableDebugWarp = false;
 #endif // WITH_EDITORONLY_DATA
 
