PoseSearch - fix for missing FMemMark(s) in FPoseSearchColumn::Filter and FPoseSearchColumn::SetOutputs

#Synced-CL 51609489

diff --git a/Engine/Plugins/Animation/PoseSearch/Source/Runtime/Private/Chooser/PoseSearchChooserColumn.cpp b/Engine/Plugins/Animation/PoseSearch/Source/Runtime/Private/Chooser/PoseSearchChooserColumn.cpp
--- a/Engine/Plugins/Animation/PoseSearch/Source/Runtime/Private/Chooser/PoseSearchChooserColumn.cpp
+++ b/Engine/Plugins/Animation/PoseSearch/Source/Runtime/Private/Chooser/PoseSearchChooserColumn.cpp
@@ -685,6 +685,9 @@ void FPoseSearchColumn::Filter(FChooserEvaluationContext& Context, const FChoose
 	ENSURE_MAPPING_VALIDATION(Mapping.Equals(FMapping(Chooser->GetRootChooser())));
 
 	ensure(ScratchArea.Num() == GetScratchAreaSize());
+
+	FMemMark Mark(FMemStack::Get());
+
 	FPoseSearchColumnScratchArea* PoseSearchColumnScratchArea = reinterpret_cast<FPoseSearchColumnScratchArea*>(ScratchArea.begin());
 #if DO_CHECK
 	ensure(PoseSearchColumnScratchArea->DebugOwner == this);
@@ -887,6 +890,8 @@ void FPoseSearchColumn::SetOutputs(FChooserEvaluationContext& Context, int RowIn
 
 	DO_ENSURE_CONSISTENCY(CheckConsistency(Chooser));
 
+	FMemMark Mark(FMemStack::Get());
+
 	// set an arbitrary "high" cost value so that cost threshold implementations can still wait for the pose search to find a good match
 	float CostValue = 100.f;
 	float StartTimeValue = 0.f;
