PoseSearch - PSD indexing improvements

#Synced-CL 51609489

diff --git a/Engine/Plugins/Animation/PoseSearch/Source/Runtime/Private/PoseSearchDatabase.cpp b/Engine/Plugins/Animation/PoseSearch/Source/Runtime/Private/PoseSearchDatabase.cpp
--- a/Engine/Plugins/Animation/PoseSearch/Source/Runtime/Private/PoseSearchDatabase.cpp
+++ b/Engine/Plugins/Animation/PoseSearch/Source/Runtime/Private/PoseSearchDatabase.cpp
@@ -54,12 +54,16 @@ static FAutoConsoleVariableRef CVarMotionMatchCompareAgainstBruteForce(TEXT("a.M
 static bool GVarMotionMatchValidateKNNSearch = false;
 static FAutoConsoleVariableRef CVarMotionMatchValidateKNNSearch(TEXT("a.MotionMatch.ValidateKNNSearch"), GVarMotionMatchValidateKNNSearch, TEXT("Validate KNN search"));
 
-// Experimental, this feature might be removed without warning, not for production use
 static bool GVarMotionMatchProfileMaxHeapKNNSearch = false;
 static FAutoConsoleVariableRef CVarMotionMatchProfileMaxHeapKNNSearch(TEXT("a.MotionMatch.ProfileMaxHeapKNNSearch"), GVarMotionMatchProfileMaxHeapKNNSearch, TEXT("Profile MaxHeap KNN Search"));
 
 #endif
 
+#if WITH_EDITOR
+static bool GVarMotionMatchRequestAsyncBuildIndexOnPostLoad = true;
+static FAutoConsoleVariableRef CVarMotionMatchRequestAsyncBuildIndexOnPostLoad(TEXT("a.MotionMatch.RequestAsyncBuildIndexOnPostLoad"), GVarMotionMatchRequestAsyncBuildIndexOnPostLoad, TEXT("Request AsyncBuildIndex On Database PostLoad. If false "));
+#endif // WITH_EDITOR
+
 // Experimental, this feature might be removed without warning, not for production use
 static int32 GVarMotionMatchDebugWeightGroupID = 0;
 static FAutoConsoleVariableRef CVarMotionMatchDebugWeightGroupID(TEXT("a.MotionMatch.DebugWeightGroupID"), GVarMotionMatchDebugWeightGroupID, TEXT("Only the channels with this or negative DebugWeightGroupID will have non zero weights"));
@@ -1134,8 +1138,9 @@ void UPoseSearchDatabase::PostLoad()
 	UpdateCachedProperties();
 
 #if WITH_EDITOR
+
 	// Delay indexing until BeginCacheForCookedPlatformData when running a CookCommandlet
-	if (!IsRunningCookCommandlet())
+	if (GVarMotionMatchRequestAsyncBuildIndexOnPostLoad && !IsRunningCookCommandlet())
 	{
 		ERequestAsyncBuildFlag Flag = ERequestAsyncBuildFlag::NewRequest;
 #if WITH_ENGINE
diff --git a/Engine/Plugins/Animation/PoseSearch/Source/Runtime/Private/PoseSearchDerivedData.cpp b/Engine/Plugins/Animation/PoseSearch/Source/Runtime/Private/PoseSearchDerivedData.cpp
--- a/Engine/Plugins/Animation/PoseSearch/Source/Runtime/Private/PoseSearchDerivedData.cpp
+++ b/Engine/Plugins/Animation/PoseSearch/Source/Runtime/Private/PoseSearchDerivedData.cpp
@@ -126,7 +126,7 @@ struct FSerializablePartialKeyHashes
 				FString ObjectName = Object->GetFullName();
 				if (Entries.Find(ObjectName))
 				{
-					UE_LOG(LogPoseSearch, Error, TEXT("FSerializablePartialKeyHashes duplicate asset name '%s'"), *ObjectName);
+					UE_LOGF(LogPoseSearch, Error, "FSerializablePartialKeyHashes duplicate asset name '%ls'", *ObjectName);
 				}
 				else
 				{
@@ -162,13 +162,13 @@ struct FSerializablePartialKeyHashes
 			{
 				if (Pair.Value.Hash != OtherEntry->Hash)
 				{
-					UE_LOG(LogPoseSearch, Error, TEXT("FSerializablePartialKeyHashes hash mismatch for asset '%s'"), *Pair.Key);
+					UE_LOGF(LogPoseSearch, Error, "FSerializablePartialKeyHashes hash mismatch for asset '%ls'", *Pair.Key);
 					bSuccess = false;
 				}
 
 				if (Pair.Value.Dependencies != OtherEntry->Dependencies)
 				{
-					UE_LOG(LogPoseSearch, Warning, TEXT("FSerializablePartialKeyHashes Dependencies mismatch for asset '%s'"), *Pair.Key);
+					UE_LOGF(LogPoseSearch, Warning, "FSerializablePartialKeyHashes Dependencies mismatch for asset '%ls'", *Pair.Key);
 					bSuccess = false;
 				}
 			}
@@ -505,7 +505,7 @@ static bool InitSearchIndexAssets(FSearchIndexBase& SearchIndex, const UPoseSear
 				DatabaseAssetRoles.Add(Role, &bIsAlreadyInSet);
 				if (bIsAlreadyInSet)
 				{
-					UE_LOG(LogPoseSearch, Error, TEXT("DatabaseAsset '%s' contains duplicate Role '%s'"), *DatabaseAsset->GetAnimationAsset()->GetName(), *Role.ToString());
+					UE_LOGF(LogPoseSearch, Error, "DatabaseAsset '%ls' contains duplicate Role '%ls'", *DatabaseAsset->GetAnimationAsset()->GetName(), *Role.ToString());
 					bAnyErrors = true;
 				}
 			}
@@ -516,7 +516,7 @@ static bool InitSearchIndexAssets(FSearchIndexBase& SearchIndex, const UPoseSear
 			{
 				if (!DatabaseAsset->GetAnimationAssetForRole(RoledSkeleton.Role))
 				{
-					UE_LOG(LogPoseSearch, Error, TEXT("DatabaseAsset '%s' doesn't support Role '%s' required by Schema '%s' Skeletons"), *DatabaseAsset->GetAnimationAsset()->GetName(), *RoledSkeleton.Role.ToString(), *Schema->GetName());
+					UE_LOGF(LogPoseSearch, Error, "DatabaseAsset '%ls' doesn't support Role '%ls' required by Schema '%ls' Skeletons", *DatabaseAsset->GetAnimationAsset()->GetName(), *RoledSkeleton.Role.ToString(), *Schema->GetName());
 					bAreAllRolesSupported = false;
 				}
 			}
@@ -530,7 +530,7 @@ static bool InitSearchIndexAssets(FSearchIndexBase& SearchIndex, const UPoseSear
 			const bool bAddMirrored = DatabaseAsset->GetMirrorOption() == EPoseSearchMirrorOption::MirroredOnly || DatabaseAsset->GetMirrorOption() == EPoseSearchMirrorOption::UnmirroredAndMirrored;
 			const bool bIsLooping = DatabaseAsset->IsLooping();
 			const bool bDisableReselection = DatabaseAsset->IsDisableReselection();
-			
+
 			// @todo: add better support for IMultiAnimAsset: currently we fix blend space parameters for only one role, 
 			//        that implies having homogeneous blendspaces for ALL the roles, and having either ONLY blend spaces or only NOT blendspaces..
 			bool bIsBlendSpace = false;
@@ -612,7 +612,7 @@ static bool PreprocessSearchIndexWeights(FSearchIndex& SearchIndex, const UPoseS
 		{
 			if (SearchIndex.WeightsSqrt[ChannelDataIndex] < 0.f)
 			{
-				UE_LOG(LogPoseSearch, Error, TEXT("Error in PreprocessSearchIndexWeights: channel '%s' in schema '%s' has a negative weight!"), *ChannelPtr->GetName(), *Schema->GetName());
+				UE_LOGF(LogPoseSearch, Error, "Error in PreprocessSearchIndexWeights: channel '%ls' in schema '%ls' has a negative weight!", *ChannelPtr->GetName(), *Schema->GetName());
 				return false;
 			}
 		}
@@ -835,7 +835,7 @@ static void PreprocessSearchIndexKDTree(FSearchIndex& SearchIndex, const UPoseSe
 
 				if (KDTreeTest != SearchIndex.KDTree)
 				{
-					UE_LOG(LogPoseSearch, Warning, TEXT("PreprocessSearchIndexKDTree - FKDTree::Construct is not deterministic"));
+					UE_LOGF(LogPoseSearch, Warning, "PreprocessSearchIndexKDTree - FKDTree::Construct is not deterministic");
 				}
 			}
 		}
@@ -870,13 +870,13 @@ static void PreprocessSearchIndexKDTree(FSearchIndex& SearchIndex, const UPoseSe
 				
 				if (!bFound)
 				{
-					UE_LOG(LogPoseSearch, Error, TEXT("PreprocessSearchIndexKDTree - kdtree for %s is not properly constructed! Couldn't find the Point %d in it"), *Database->GetName(), PointIndex);
+					UE_LOGF(LogPoseSearch, Error, "PreprocessSearchIndexKDTree - kdtree for %ls is not properly constructed! Couldn't find the Point %d in it", *Database->GetName(), PointIndex);
 				}
 			}
 
 			if (MaxNumNeighborToFindAPoint >= KDTreeQueryNumNeighbors)
 			{
-				UE_LOG(LogPoseSearch, Warning, TEXT("Not enough 'KDTreeQueryNumNeighbors' (%d) for database '%s'. Pose values projected in PCA space have too many duplicates, so try to prune duplicates by tuning 'PCAValuesPruningSimilarityThreshold' or increase 'KDTreeQueryNumNeighbors' at least to %d"), KDTreeQueryNumNeighbors, *Database->GetName(), MaxNumNeighborToFindAPoint);
+				UE_LOGF(LogPoseSearch, Warning, "Not enough 'KDTreeQueryNumNeighbors' (%d) for database '%ls'. Pose values projected in PCA space have too many duplicates, so try to prune duplicates by tuning 'PCAValuesPruningSimilarityThreshold' or increase 'KDTreeQueryNumNeighbors' at least to %d", KDTreeQueryNumNeighbors, *Database->GetName(), MaxNumNeighborToFindAPoint);
 			}
 
 			// if bArePCAValuesPruned PointIndex is the index of the point in the kdtree, NOT necessary the pose index, so doing the PCAProject would lead to the wrong data
@@ -897,14 +897,14 @@ static void PreprocessSearchIndexKDTree(FSearchIndex& SearchIndex, const UPoseSe
 						{
 							if (Results[ResultIndex].Distance > UE_KINDA_SMALL_NUMBER)
 							{
-								UE_LOG(LogPoseSearch, Error, TEXT("PreprocessSearchIndexKDTree - kdtree for %s is not properly constructed! Couldn't find the Point %d in it within UE_KINDA_SMALL_NUMBER tolerance, after PCA projection"), *Database->GetName(), PointIndex);
+								UE_LOGF(LogPoseSearch, Error, "PreprocessSearchIndexKDTree - kdtree for %ls is not properly constructed! Couldn't find the Point %d in it within UE_KINDA_SMALL_NUMBER tolerance, after PCA projection", *Database->GetName(), PointIndex);
 							}
 							break;
 						}
 					}
 					if (ResultIndex == NumResults)
 					{
-						UE_LOG(LogPoseSearch, Error, TEXT("PreprocessSearchIndexKDTree - kdtree for %s is not properly constructed! Couldn't find the Point %d in it, after PCA projection"), *Database->GetName(), PointIndex);
+						UE_LOGF(LogPoseSearch, Error, "PreprocessSearchIndexKDTree - kdtree for %ls is not properly constructed! Couldn't find the Point %d in it, after PCA projection", *Database->GetName(), PointIndex);
 					}
 				}
 			}
@@ -933,7 +933,7 @@ static void PreprocessSearchIndexVPTree(FSearchIndex& SearchIndex, const UPoseSe
 		{
 			if (!SearchIndex.VPTree.TestConstruct(DataSource))
 			{
-				UE_LOG(LogPoseSearch, Error, TEXT("PreprocessSearchIndexVPTree - FVPTree construction failed"));
+				UE_LOGF(LogPoseSearch, Error, "PreprocessSearchIndexVPTree - FVPTree construction failed");
 			}
 
 			const int32 NumIterations = GVarMotionMatchTestNumIterations;
@@ -945,7 +945,7 @@ static void PreprocessSearchIndexVPTree(FSearchIndex& SearchIndex, const UPoseSe
 
 				if (VPTreeTest != SearchIndex.VPTree)
 				{
-					UE_LOG(LogPoseSearch, Warning, TEXT("PreprocessSearchIndexVPTree - FVPTree construction is not deterministic"));
+					UE_LOGF(LogPoseSearch, Warning, "PreprocessSearchIndexVPTree - FVPTree construction is not deterministic");
 				}
 			}
 		}
@@ -968,7 +968,7 @@ static void PreprocessSearchIndexVPTree(FSearchIndex& SearchIndex, const UPoseSe
 						{
 							if (!FMath::IsNearlyZero(Result.Distance))
 							{
-								UE_LOG(LogPoseSearch, Error, TEXT("PreprocessSearchIndexVPTree - VPTree for %s is malformed because foud PointIndex %d distance %f from itself (distance should be zero)!"), *Database->GetName(), PointIndex, Result.Distance);
+								UE_LOGF(LogPoseSearch, Error, "PreprocessSearchIndexVPTree - VPTree for %ls is malformed because foud PointIndex %d distance %f from itself (distance should be zero)!", *Database->GetName(), PointIndex, Result.Distance);
 							}
 
 							bFound = true;
@@ -978,13 +978,13 @@ static void PreprocessSearchIndexVPTree(FSearchIndex& SearchIndex, const UPoseSe
 
 					if (!bFound)
 					{
-						UE_LOG(LogPoseSearch, Error, TEXT("PreprocessSearchIndexVPTree - VPTree for %s is malformed and couldn't find PointIndex %d"), *Database->GetName(), PointIndex);
+						UE_LOGF(LogPoseSearch, Error, "PreprocessSearchIndexVPTree - VPTree for %ls is malformed and couldn't find PointIndex %d", *Database->GetName(), PointIndex);
 					}
 				}
 			}
 			else
 			{
-				UE_LOG(LogPoseSearch, Warning, TEXT("PreprocessSearchIndexVPTree - cannot VPTreeConstruct for %s if there could be potential duplicate poses: set PosePruningSimilarityThreshold > 0 to enforce duplicate pruning"), *Database->GetName());
+				UE_LOGF(LogPoseSearch, Warning, "PreprocessSearchIndexVPTree - cannot VPTreeConstruct for %ls if there could be potential duplicate poses: set PosePruningSimilarityThreshold > 0 to enforce duplicate pruning", *Database->GetName());
 			}
 		}
 #endif // ENABLE_ANIM_DEBUG
@@ -1040,7 +1040,7 @@ static void InitRoledBoneContainers(TMap<FRole, FBoneContainer>& RoledBoneContai
 			}
 			else
 			{
-				UE_LOG(LogPoseSearch, Warning, TEXT("InitBoneContainersFromRoledSkeleton: MirrorDataTable %s doesn't contain bone with index %d."), *MirrorDataTable->GetName(), BoneIndex);
+				UE_LOGF(LogPoseSearch, Warning, "InitBoneContainersFromRoledSkeleton: MirrorDataTable %ls doesn't contain bone with index %d.", *MirrorDataTable->GetName(), BoneIndex);
 			}
 		}
 	};
@@ -1101,7 +1101,7 @@ static void InitRoledBoneContainers(TMap<FRole, FBoneContainer>& RoledBoneContai
 	}
 }
 
-static bool IndexDatabase(FSearchIndexBase& SearchIndexBase, const UPoseSearchDatabase* DatabaseToLookForAssets, const UPoseSearchSchema* Schema, const FAssetSamplingContext& SamplingContext, const FFloatInterval& AdditionalExtrapolationTime, UE::DerivedData::FRequestOwner& Owner)
+static bool IndexDatabase(FSearchIndexBase& SearchIndexBase, const UPoseSearchDatabase* DatabaseToLookForAssets, const UPoseSearchSchema* Schema, const FAssetSamplingContext& SamplingContext, const FFloatInterval& AdditionalExtrapolationTime)
 {
 	if (!ensure(Schema))
 	{
@@ -1140,10 +1140,6 @@ static bool IndexDatabase(FSearchIndexBase& SearchIndexBase, const UPoseSearchDa
 	}
 
 	ParallelFor(Samplers.Num(), [&Samplers](int32 SamplerIdx) { Samplers[SamplerIdx].Process(); }, ParallelForFlags);
-	if (Owner.IsCanceled())
-	{
-		return false;
-	}
 
 	// prepare indexers
 	TArray<FAssetIndexer> Indexers;
@@ -1214,11 +1210,6 @@ static bool IndexDatabase(FSearchIndexBase& SearchIndexBase, const UPoseSearchDa
 		TotalPoses += Indexers[AssetIdx].GetNumIndexedPoses();
 	}
 
-	if (Owner.IsCanceled())
-	{
-		return false;
-	}
-	
 	// Index asset data
 	ParallelFor(Indexers.Num(), [&Indexers](int32 AssetIdx) { Indexers[AssetIdx].Process(AssetIdx); }, ParallelForFlags);
 
@@ -1230,11 +1221,6 @@ static bool IndexDatabase(FSearchIndexBase& SearchIndexBase, const UPoseSearchDa
 		}
 	}
 
-	if (Owner.IsCanceled())
-	{
-		return false;
-	}
-
 	// Joining EventData
 	FEventDataCollector EventDataCollector;
 	for (int32 AssetIdx = 0; AssetIdx != SearchIndexBase.Assets.Num(); ++AssetIdx)
@@ -1292,11 +1278,6 @@ static bool IndexDatabase(FSearchIndexBase& SearchIndexBase, const UPoseSearchDa
 		}
 	}
 
-	if (Owner.IsCanceled())
-	{
-		return false;
-	}
-
 	return true;
 }
 
@@ -1323,7 +1304,7 @@ static bool ValidateSearchInfoAgainstDatabase(const FSearchIndex& SearchIndex, c
 
 	if (SearchIndex.GetNumDimensions() != Database->Schema->SchemaCardinality)
 	{
-		UE_LOG(LogPoseSearch, Warning, TEXT("%s - %s BuildIndex SearchIndex Corrupted! SchemaCardinality mismatch %d vs %d"), *LexToString(FullIndexKey.Hash), *GetDatabaseAndOuterName(Database), SearchIndex.GetNumDimensions(), Database->Schema->SchemaCardinality);
+		UE_LOGF(LogPoseSearch, Warning, "%ls - %ls BuildIndex SearchIndex Corrupted! SchemaCardinality mismatch %d vs %d", *LexToString(FullIndexKey.Hash), *GetDatabaseAndOuterName(Database), SearchIndex.GetNumDimensions(), Database->Schema->SchemaCardinality);
 		return false;
 	}
 
@@ -1334,7 +1315,7 @@ static bool ValidateSearchInfoAgainstDatabase(const FSearchIndex& SearchIndex, c
 		const int32 SourceAssetIdx = SearchIndexAsset.GetSourceAssetIdx();
 		if (SourceAssetIdx < 0 || SourceAssetIdx >= NumAnimationAssets)
 		{
-			UE_LOG(LogPoseSearch, Warning, TEXT("%s - %s BuildIndex SearchIndex Corrupted! SearchIndex.Assets referencing missing asset with index %d"), *LexToString(FullIndexKey.Hash), *GetDatabaseAndOuterName(Database), SourceAssetIdx);
+			UE_LOGF(LogPoseSearch, Warning, "%ls - %ls BuildIndex SearchIndex Corrupted! SearchIndex.Assets referencing missing asset with index %d", *LexToString(FullIndexKey.Hash), *GetDatabaseAndOuterName(Database), SourceAssetIdx);
 			return false;
 		}
 
@@ -1344,14 +1325,14 @@ static bool ValidateSearchInfoAgainstDatabase(const FSearchIndex& SearchIndex, c
 			{
 				if (!DatabaseAsset->IsEnabled())
 				{
-					UE_LOG(LogPoseSearch, Warning, TEXT("%s - %s BuildIndex SearchIndex Corrupted! Found references to disabled asset %s in the SearchIndex"), *LexToString(FullIndexKey.Hash), *GetDatabaseAndOuterName(Database), *AnimationAsset->GetName());
+					UE_LOGF(LogPoseSearch, Warning, "%ls - %ls BuildIndex SearchIndex Corrupted! Found references to disabled asset %ls in the SearchIndex", *LexToString(FullIndexKey.Hash), *GetDatabaseAndOuterName(Database), *AnimationAsset->GetName());
 					return false;
 				}
 			}
 		}
 		else
 		{
-			UE_LOG(LogPoseSearch, Warning, TEXT("%s - %s BuildIndex SearchIndex Corrupted! null FPoseSearchDatabaseAnimationAssetBase asset at index %d!?"), *LexToString(FullIndexKey.Hash), *GetDatabaseAndOuterName(Database), SourceAssetIdx);
+			UE_LOGF(LogPoseSearch, Warning, "%ls - %ls BuildIndex SearchIndex Corrupted! null FPoseSearchDatabaseAnimationAssetBase asset at index %d!?", *LexToString(FullIndexKey.Hash), *GetDatabaseAndOuterName(Database), SourceAssetIdx);
 			return false;
 		}
 	}
@@ -1507,6 +1488,168 @@ static void CompareSearchIndex(const FSearchIndex& A, const FSearchIndex& B, con
 
 #endif // ENABLE_ANIM_DEBUG
 
+static bool IndexDatabase(const UPoseSearchDatabase* DependentDatabase, const UPoseSearchSchema* DependentDatabaseSchema, FSearchIndexBase& SearchIndexBase, 
+	const UPoseSearchDatabase* MainDatabase, const UE::DerivedData::FCacheKey& FullIndexKey)
+{
+	check(DependentDatabase);
+
+	const FString DependentDatabaseName = GetDatabaseAndOuterName(DependentDatabase);
+
+	const FAssetSamplingContext DependentSamplingContext(*DependentDatabase);
+	const FFloatInterval& DependentExcludeFromDatabaseParameters = DependentDatabase->ExcludeFromDatabaseParameters;
+	const FFloatInterval& DependentAdditionalExtrapolationTime = DependentDatabase->AdditionalExtrapolationTime;
+
+	// early out for invalid indexing conditions
+	if (!DependentDatabaseSchema || DependentDatabaseSchema->SchemaCardinality <= 0)
+	{
+		if (DependentDatabase == MainDatabase)
+		{
+			UE_LOGF(LogPoseSearch, Error, "%ls - %ls BuildIndex Failed because of invalid Schema", *LexToString(FullIndexKey.Hash), *GetDatabaseAndOuterName(MainDatabase));
+		}
+		else
+		{
+			UE_LOGF(LogPoseSearch, Error, "%ls - %ls BuildIndex Failed because dependent database '%ls' has an invalid Schema", *LexToString(FullIndexKey.Hash), *GetDatabaseAndOuterName(MainDatabase), *DependentDatabaseName);
+		}
+		return false;
+	}
+
+	// validating that the missing MirrorDataTable(s) are not necessary 
+	bool bAllRoledSkeletonHaveMirrorDataTable = true;
+	for (const FPoseSearchRoledSkeleton& RoledSkeleton : DependentDatabaseSchema->GetRoledSkeletons())
+	{
+		if (RoledSkeleton.MirrorDataTable)
+		{
+			if (!RoledSkeleton.MirrorDataTable->Skeleton)
+			{
+				UE_LOGF(LogPoseSearch, Error, "%ls - %ls BuildIndex Failed because '%ls' schema MirrorDataTable Skeleton is not set for Role '%ls' ", *LexToString(FullIndexKey.Hash), *GetDatabaseAndOuterName(MainDatabase), *DependentDatabaseName, *RoledSkeleton.Role.ToString());
+				return false;
+			}
+		}
+		else
+		{
+			bAllRoledSkeletonHaveMirrorDataTable = false;
+		}
+	}
+
+	if (!bAllRoledSkeletonHaveMirrorDataTable)
+	{
+		for (int32 AnimationAssetIndex = 0; AnimationAssetIndex < DependentDatabase->GetNumAnimationAssets(); ++AnimationAssetIndex)
+		{
+			if (const FPoseSearchDatabaseAnimationAssetBase* DatabaseAsset = DependentDatabase->GetDatabaseAnimationAsset(AnimationAssetIndex))
+			{
+				if (DatabaseAsset->GetMirrorOption() == EPoseSearchMirrorOption::MirroredOnly || DatabaseAsset->GetMirrorOption() == EPoseSearchMirrorOption::UnmirroredAndMirrored)
+				{
+					// want to sample a mirrored asset
+					UE_LOGF(LogPoseSearch, Error, "%ls - %ls BuildIndex Failed because '%ls' schema requires MirrorDataTable(s) to sample mirrored animation assets", *LexToString(FullIndexKey.Hash), *GetDatabaseAndOuterName(MainDatabase), *DependentDatabaseName);
+					return false;
+				}
+			}
+		}
+	}
+
+	for (int32 AnimationAssetIndex = 0; AnimationAssetIndex < DependentDatabase->GetNumAnimationAssets(); ++AnimationAssetIndex)
+	{
+		if (const FPoseSearchDatabaseAnimationAssetBase* DatabaseAsset = DependentDatabase->GetDatabaseAnimationAsset(AnimationAssetIndex))
+		{
+			const UAnimationAsset* AnimationAsset = Cast<const UAnimationAsset>(DatabaseAsset->GetAnimationAsset());
+			if (DatabaseAsset->GetAnimationAsset() == nullptr)
+			{
+				UE_LOGF(LogPoseSearch, Warning, "OnGetComplete - Database %ls entry %d has no associated animation asset.", *DatabaseAsset->GetName(), AnimationAssetIndex);
+			}
+			else if (!DatabaseAsset->GetAnimationAsset()->GetPackage() || DatabaseAsset->GetAnimationAsset()->HasAnyFlags(RF_MirroredGarbage))
+			{
+				UE_LOGF(LogPoseSearch, Error, "%ls - %ls BuildIndex Failed because '%ls' referenced asset '%ls' is garbage or is not contained in a package", *LexToString(FullIndexKey.Hash), *GetDatabaseAndOuterName(MainDatabase), *DependentDatabaseName, *DatabaseAsset->GetAnimationAsset()->GetName());
+				return false;
+			}
+			else if (AnimationAsset && AnimationAsset->GetSkeleton() == nullptr)
+			{
+				UE_LOGF(LogPoseSearch, Error, "%ls - %ls BuildIndex Failed because '%ls' referenced asset '%ls' which has an invalid (null) Skeleton", *LexToString(FullIndexKey.Hash), *GetDatabaseAndOuterName(MainDatabase), *DependentDatabaseName, *DatabaseAsset->GetAnimationAsset()->GetName());
+				return false;
+			}
+			else if (!IsSchemaCompatible(DatabaseAsset, MainDatabase->Schema))
+			{
+				UE_LOGF(LogPoseSearch, Warning, "OnGetComplete - %ls is not compatible with the schema %ls", *DatabaseAsset->GetName(), *GetNameSafe(MainDatabase->Schema));
+			}
+			else if (const UBlendSpace* BlendSpace = Cast<UBlendSpace>(DatabaseAsset->GetAnimationAsset()))
+			{
+				if (!BlendSpace->bShouldMatchSyncPhases)
+				{
+					UE_LOGF(LogPoseSearch, Warning, "OnGetComplete - %ls's bShouldMatchSyncPhases flag is not enabled. This is required for properly pose matching blendspaces.", *DatabaseAsset->GetName());
+				}
+				else
+				{
+					const TConstArrayView<FBlendSample> BlendSamples = BlendSpace->GetBlendSamples();
+					for (int i = 0; i < BlendSamples.Num() - 1; ++i)
+					{
+						const FBlendSample& CurrSample = BlendSamples[i];
+						const FBlendSample& NextSample = BlendSamples[i + 1];
+
+						if (CurrSample.Animation && NextSample.Animation)
+						{
+							bool bWarning = false;
+
+							if (CurrSample.Animation->AuthoredSyncMarkers.Num() == NextSample.Animation->AuthoredSyncMarkers.Num())
+							{
+								for (int j = 0; j < CurrSample.Animation->AuthoredSyncMarkers.Num(); ++j)
+								{
+									if (CurrSample.Animation->AuthoredSyncMarkers[j].MarkerName != NextSample.Animation->AuthoredSyncMarkers[j].MarkerName)
+									{
+										bWarning = true;
+										break;
+									}
+								}
+							}
+							else
+							{
+								bWarning = true;
+							}
+
+							if (bWarning)
+							{
+								UE_LOGF(LogPoseSearch, Warning, "OnGetComplete - %ls's samples don't share the same layout of sync markers. This is required for properly pose matching blendspaces.", *DatabaseAsset->GetName());
+							}
+						}
+					}
+				}
+			}
+		}
+	}
+
+	// Building all the related FPoseSearchBaseIndex first
+	if (!InitSearchIndexAssets(SearchIndexBase, DependentDatabase, DependentDatabaseSchema, DependentExcludeFromDatabaseParameters))
+	{
+		UE_LOGF(LogPoseSearch, Error, "%ls - %ls BuildIndex Failed because of invalid assets", *LexToString(FullIndexKey.Hash), *GetDatabaseAndOuterName(MainDatabase));
+		return false;
+	}
+
+	if (!IndexDatabase(SearchIndexBase, DependentDatabase, DependentDatabaseSchema, DependentSamplingContext, DependentAdditionalExtrapolationTime))
+	{
+		UE_LOGF(LogPoseSearch, Log, "%ls - %ls BuildIndex Failed in IndexDatabase", *LexToString(FullIndexKey.Hash), *GetDatabaseAndOuterName(MainDatabase));
+		return false;
+	}
+
+#if ENABLE_ANIM_DEBUG
+	if (GVarMotionMatchTestDeterminismIndexDatabase)
+	{
+		for (int32 Iteration = 0; Iteration < GVarMotionMatchTestNumIterations; ++Iteration)
+		{
+			FSearchIndexBase TestSearchIndexBase = SearchIndexBase;
+			if (IndexDatabase(TestSearchIndexBase, DependentDatabase, DependentDatabaseSchema, DependentSamplingContext, DependentAdditionalExtrapolationTime))
+			{
+				if (TestSearchIndexBase != SearchIndexBase)
+				{
+					TStringBuilder<0> Message;
+					CompareSearchIndexBase(TestSearchIndexBase, SearchIndexBase, DependentDatabaseSchema, Message);
+					UE_LOGF(LogPoseSearch, Warning, "OnGetComplete - IndexDatabase is not deterministic\n%ls", *Message);
+				}
+			}
+		}
+	}
+#endif //ENABLE_ANIM_DEBUG
+
+	return true;
+}
+
 //////////////////////////////////////////////////////////////////////////
 // FPoseSearchDatabaseAsyncCacheTask
 struct FPoseSearchDatabaseAsyncCacheTask
@@ -1581,7 +1724,7 @@ FPoseSearchDatabaseAsyncCacheTask::FPoseSearchDatabaseAsyncCacheTask(UPoseSearch
 	}
 	else
 	{
-		UE_LOG(LogPoseSearch, Log, TEXT("Delaying DDC until on the game thread    - %s"), *Database->GetName());
+		UE_LOGF(LogPoseSearch, Log, "Delaying DDC until on the game thread    - %ls", *Database->GetName());
 	}
 }
 
@@ -1659,7 +1802,7 @@ void FPoseSearchDatabaseAsyncCacheTask::StartNewRequestIfNeeded(bool bPerformCon
 		DerivedDataKey = FIoHash::Zero;
 		SetState(EState::Notstarted);
 	
-		UE_LOG(LogPoseSearch, Log, TEXT("Delaying DDC until dependents are fully loaded - %s"), *Database->GetName());
+		UE_LOGF(LogPoseSearch, Log, "Delaying DDC until dependents are fully loaded - %ls", *Database->GetName());
 	}
 	else
 	{
@@ -1680,7 +1823,7 @@ void FPoseSearchDatabaseAsyncCacheTask::StartNewRequestIfNeeded(bool bPerformCon
 
 				SetState(EState::Prestarted);
 				
-				UE_LOG(LogPoseSearch, Log, TEXT("%s - %s BeginCache"), *LexToString(DerivedDataKey), *GetDatabaseAndOuterName(Database.Get()));
+				UE_LOGF(LogPoseSearch, Log, "%ls - %ls BeginCache", *LexToString(DerivedDataKey), *GetDatabaseAndOuterName(Database.Get()));
 
 				const FCacheKey CacheKey{ Bucket, DerivedDataKey };
 				const FCacheGetRequest CacheRequest = { { Database->GetPathName() }, CacheKey, ECachePolicy::Default };
@@ -1695,7 +1838,7 @@ void FPoseSearchDatabaseAsyncCacheTask::StartNewRequestIfNeeded(bool bPerformCon
 			{
 				DerivedDataKey = FIoHash::Zero;
 				SetState(EState::Notstarted);
-				UE_LOG(LogPoseSearch, Log, TEXT("Delaying DDC until dependents are ready - %s"), *Database->GetName());
+				UE_LOGF(LogPoseSearch, Log, "Delaying DDC until dependents are ready - %ls", *Database->GetName());
 			}
 			
 			bResidencyCleared = false;
@@ -1715,7 +1858,7 @@ void FPoseSearchDatabaseAsyncCacheTask::PreCancelIfDependsOn(const UObject* Obje
 		{
 			// Database can be null if the task was Ended/Failed and Database was already garbage collected, but Tick hasn't been called yet
 			FString DatabaseName = IsValid() ? *GetDatabaseAndOuterName(Database.Get()) : TEXT("Garbage Collected Database");
-			UE_LOG(LogPoseSearch, Log, TEXT("%s - %s PreCancelled because of %s"), *LexToString(DerivedDataKey), *DatabaseName, *Object->GetName());
+			UE_LOGF(LogPoseSearch, Log, "%ls - %ls PreCancelled because of %ls", *LexToString(DerivedDataKey), *DatabaseName, *Object->GetName());
 
 			// Owner.Cancel must be performed before SearchIndex.Reset() in case any task is flying (launched by Owner.LaunchTask)
 			Owner.Cancel();
@@ -1731,7 +1874,7 @@ void FPoseSearchDatabaseAsyncCacheTask::Cancel()
 	check(IsInGameThread());
 
 	FString DatabaseName = IsValid() ? *GetDatabaseAndOuterName(Database.Get()) : TEXT("Garbage Collected Database");
-	UE_LOG(LogPoseSearch, Log, TEXT("%s - %s Cancelled"), *LexToString(DerivedDataKey), *DatabaseName);
+	UE_LOGF(LogPoseSearch, Log, "%ls - %ls Cancelled", *LexToString(DerivedDataKey), *DatabaseName);
 
 	// Owner.Cancel must be performed before SearchIndex.Reset() in case any task is flying (launched by Owner.LaunchTask)
 	Owner.Cancel();
@@ -1837,7 +1980,7 @@ void FPoseSearchDatabaseAsyncCacheTask::OnGetComplete(UE::DerivedData::FCacheGet
 		// so to account for this, we just reindex the database and update the associated DDC 
 		if (ValidateSearchInfoAgainstDatabase(SearchIndex, Database.Get(), FullIndexKey))
 		{
-			UE_LOG(LogPoseSearch, Log, TEXT("%s - %s BuildIndex From Cache"), *LexToString(FullIndexKey.Hash), *GetDatabaseAndOuterName(Database.Get()));
+			UE_LOGF(LogPoseSearch, Log, "%ls - %ls BuildIndex From Cache", *LexToString(FullIndexKey.Hash), *GetDatabaseAndOuterName(Database.Get()));
 		}
 		else
 		{
@@ -1850,7 +1993,7 @@ void FPoseSearchDatabaseAsyncCacheTask::OnGetComplete(UE::DerivedData::FCacheGet
 	if (Response.Status == EStatus::Canceled)
 	{
 		ResetSearchIndex();
-		UE_LOG(LogPoseSearch, Log, TEXT("%s - %s BuildIndex Cancelled"), *LexToString(FullIndexKey.Hash), *GetDatabaseAndOuterName(Database.Get()));
+		UE_LOGF(LogPoseSearch, Log, "%ls - %ls BuildIndex Cancelled", *LexToString(FullIndexKey.Hash), *GetDatabaseAndOuterName(Database.Get()));
 	}
 	
 	bool bForceBuildIndex = false;
@@ -1883,7 +2026,7 @@ void FPoseSearchDatabaseAsyncCacheTask::OnGetComplete(UE::DerivedData::FCacheGet
 				IndexBaseDatabases.Add(MainDatabase);
 				if (!MainDatabase)
 				{
-					UE_LOG(LogPoseSearch, Log, TEXT("%s - BuildIndex Cancelled because associated Database weak pointer has been released."), *LexToString(FullIndexKey.Hash));
+					UE_LOGF(LogPoseSearch, Log, "%ls - BuildIndex Cancelled because associated Database weak pointer has been released.", *LexToString(FullIndexKey.Hash));
 					ResetSearchIndex();
 					return;
 				}
@@ -1893,11 +2036,10 @@ void FPoseSearchDatabaseAsyncCacheTask::OnGetComplete(UE::DerivedData::FCacheGet
 					MainDatabase->NormalizationSet->AddUniqueDatabases(IndexBaseDatabases);
 				}
 
-				const FString MainDatabaseName = GetDatabaseAndOuterName(MainDatabase);
 				const UPoseSearchSchema* MainDatabaseSchema = MainDatabase->Schema;
 				if (!MainDatabaseSchema || MainDatabaseSchema->SchemaCardinality <= 0)
 				{
-					UE_LOG(LogPoseSearch, Error, TEXT("%s - %s BuildIndex Failed because of invalid Schema"), *LexToString(FullIndexKey.Hash), *MainDatabaseName);
+					UE_LOGF(LogPoseSearch, Error, "%ls - %ls BuildIndex Failed because of invalid Schema", *LexToString(FullIndexKey.Hash), *GetDatabaseAndOuterName(MainDatabase));
 					ResetSearchIndex();
 					return;
 				}
@@ -1914,184 +2056,13 @@ void FPoseSearchDatabaseAsyncCacheTask::OnGetComplete(UE::DerivedData::FCacheGet
 					const UPoseSearchDatabase* DependentDatabase = IndexBaseDatabases[IndexBaseIdx];
 					check(DependentDatabase);
 
-					const UPoseSearchSchema* DependentDatabaseSchema = bNormalizeWithCommonSchema ? MainDatabaseSchema : DependentDatabase->Schema.Get();
-					const FString DependentDatabaseName = GetDatabaseAndOuterName(DependentDatabase);
+					Schemas[IndexBaseIdx] = bNormalizeWithCommonSchema ? MainDatabaseSchema : DependentDatabase->Schema.Get();
 
-					FSearchIndexBase& SearchIndexBase = SearchIndexBases[IndexBaseIdx];
-
-					Schemas[IndexBaseIdx] = DependentDatabaseSchema;
-					const FAssetSamplingContext DependentSamplingContext(*DependentDatabase);
-					const FFloatInterval& DependentExcludeFromDatabaseParameters = DependentDatabase->ExcludeFromDatabaseParameters;
-					const FFloatInterval& DependentAdditionalExtrapolationTime = DependentDatabase->AdditionalExtrapolationTime;
-
-					// early out for invalid indexing conditions
-					if (!DependentDatabaseSchema || DependentDatabaseSchema->SchemaCardinality <= 0)
+					if (!IndexDatabase(DependentDatabase, Schemas[IndexBaseIdx], SearchIndexBases[IndexBaseIdx], MainDatabase, FullIndexKey))
 					{
-						if (IndexBaseIdx == 0)
-						{
-							UE_LOG(LogPoseSearch, Error, TEXT("%s - %s BuildIndex Failed because of invalid Schema"), *LexToString(FullIndexKey.Hash), *MainDatabaseName);
-						}
-						else
-						{
-							UE_LOG(LogPoseSearch, Error, TEXT("%s - %s BuildIndex Failed because dependent database '%s' has an invalid Schema"), *LexToString(FullIndexKey.Hash), *MainDatabaseName, *DependentDatabaseName);
-						}
 						ResetSearchIndex();
 						return;
 					}
-
-					// validating that the missing MirrorDataTable(s) are not necessary 
-					bool bAllRoledSkeletonHaveMirrorDataTable = true;
-					for (const FPoseSearchRoledSkeleton& RoledSkeleton : DependentDatabaseSchema->GetRoledSkeletons())
-					{
-						if (RoledSkeleton.MirrorDataTable)
-						{
-							if (!RoledSkeleton.MirrorDataTable->Skeleton)
-							{
-								UE_LOG(LogPoseSearch, Error, TEXT("%s - %s BuildIndex Failed because '%s' schema MirrorDataTable Skeleton is not set for Role '%s' "), *LexToString(FullIndexKey.Hash), *MainDatabaseName, *DependentDatabaseName, *RoledSkeleton.Role.ToString());
-								ResetSearchIndex();
-								return;
-							}
-						}
-						else
-						{
-							bAllRoledSkeletonHaveMirrorDataTable = false;
-						}
-					}
-
-					if (!bAllRoledSkeletonHaveMirrorDataTable)
-					{
-						for (int32 AnimationAssetIndex = 0; AnimationAssetIndex < DependentDatabase->GetNumAnimationAssets(); ++AnimationAssetIndex)
-						{
-							if (const FPoseSearchDatabaseAnimationAssetBase* DatabaseAsset = DependentDatabase->GetDatabaseAnimationAsset(AnimationAssetIndex))
-							{
-								if (DatabaseAsset->GetMirrorOption() == EPoseSearchMirrorOption::MirroredOnly || DatabaseAsset->GetMirrorOption() == EPoseSearchMirrorOption::UnmirroredAndMirrored)
-								{
-									// want to sample a mirrored asset
-									UE_LOG(LogPoseSearch, Error, TEXT("%s - %s BuildIndex Failed because '%s' schema requires MirrorDataTable(s) to sample mirrored animation assets"), *LexToString(FullIndexKey.Hash), *MainDatabaseName, *DependentDatabaseName);
-									ResetSearchIndex();
-									return;
-								}
-							}
-						}
-					}
-					
-					for (int32 AnimationAssetIndex = 0; AnimationAssetIndex < DependentDatabase->GetNumAnimationAssets(); ++AnimationAssetIndex)
-					{
-						if (const FPoseSearchDatabaseAnimationAssetBase* DatabaseAsset = DependentDatabase->GetDatabaseAnimationAsset(AnimationAssetIndex))
-						{
-							const UAnimationAsset* AnimationAsset = Cast<const UAnimationAsset>(DatabaseAsset->GetAnimationAsset());
-							if (DatabaseAsset->GetAnimationAsset() == nullptr)
-							{
-								UE_LOG(LogPoseSearch, Warning, TEXT("OnGetComplete - Database %s entry %d has no associated animation asset."), *DatabaseAsset->GetName(), AnimationAssetIndex);
-							}
-							else if (!DatabaseAsset->GetAnimationAsset()->GetPackage() || DatabaseAsset->GetAnimationAsset()->HasAnyFlags(RF_MirroredGarbage))
-							{
-								UE_LOG(LogPoseSearch, Error, TEXT("%s - %s BuildIndex Failed because '%s' referenced asset '%s' is garbage or is not contained in a package"), *LexToString(FullIndexKey.Hash), *MainDatabaseName, *DependentDatabaseName, *DatabaseAsset->GetAnimationAsset()->GetName());
-								ResetSearchIndex();
-								return;
-							}
-							else if (AnimationAsset && AnimationAsset->GetSkeleton() == nullptr)
-							{
-								UE_LOG(LogPoseSearch, Error, TEXT("%s - %s BuildIndex Failed because '%s' referenced asset '%s' which has an invalid (null) Skeleton"), *LexToString(FullIndexKey.Hash), *MainDatabaseName, *DependentDatabaseName, *DatabaseAsset->GetAnimationAsset()->GetName());
-								ResetSearchIndex();
-								return;
-							}
-							else if (!IsSchemaCompatible(DatabaseAsset, Database->Schema))
-							{
-								UE_LOG(LogPoseSearch, Warning, TEXT("OnGetComplete - %s is not compatible with the schema %s"), *DatabaseAsset->GetName(), *GetNameSafe(Database->Schema));
-							}
-							else if (const UBlendSpace* BlendSpace = Cast<UBlendSpace>(DatabaseAsset->GetAnimationAsset()))
-							{
-								if (!BlendSpace->bShouldMatchSyncPhases)
-								{
-									UE_LOG(LogPoseSearch, Warning, TEXT("OnGetComplete - %s's bShouldMatchSyncPhases flag is not enabled. This is required for properly pose matching blendspaces."), *DatabaseAsset->GetName());
-								}
-								else
-								{
-									const TConstArrayView<FBlendSample> BlendSamples = BlendSpace->GetBlendSamples();
-									for (int i = 0; i < BlendSamples.Num() - 1; ++i)
-									{
-										const FBlendSample& CurrSample = BlendSamples[i];
-										const FBlendSample& NextSample = BlendSamples[i + 1];
-										
-										if (CurrSample.Animation && NextSample.Animation)
-										{
-											bool bWarning = false;
-											
-											if (CurrSample.Animation->AuthoredSyncMarkers.Num() == NextSample.Animation->AuthoredSyncMarkers.Num())
-											{
-												for (int j = 0; j < CurrSample.Animation->AuthoredSyncMarkers.Num(); ++j)
-												{
-													if (CurrSample.Animation->AuthoredSyncMarkers[j].MarkerName != NextSample.Animation->AuthoredSyncMarkers[j].MarkerName)
-													{
-														bWarning = true;
-														break;
-													}
-												}
-											}
-											else
-											{
-												bWarning = true;
-											}
-											
-											if (bWarning)
-											{
-												UE_LOG(LogPoseSearch, Warning, TEXT("OnGetComplete - %s's samples don't share the same layout of sync markers. This is required for properly pose matching blendspaces."), *DatabaseAsset->GetName());
-											}
-										}
-									}
-								}
-							}
-						}
-					}
-					
-					if (Owner.IsCanceled())
-					{
-						UE_LOG(LogPoseSearch, Log, TEXT("%s - %s BuildIndex Cancelled"), *LexToString(FullIndexKey.Hash), *MainDatabaseName);
-						ResetSearchIndex();
-						return;
-					}
-
-					// Building all the related FPoseSearchBaseIndex first
-					if (!InitSearchIndexAssets(SearchIndexBase, DependentDatabase, DependentDatabaseSchema, DependentExcludeFromDatabaseParameters))
-					{
-						UE_LOG(LogPoseSearch, Error, TEXT("%s - %s BuildIndex Failed because of invalid assets"), *LexToString(FullIndexKey.Hash), *MainDatabaseName);
-						ResetSearchIndex();
-						return;
-					}
-
-					if (Owner.IsCanceled())
-					{
-						UE_LOG(LogPoseSearch, Log, TEXT("%s - %s BuildIndex Cancelled"), *LexToString(FullIndexKey.Hash), *MainDatabaseName);
-						ResetSearchIndex();
-						return;
-					}
-
-					if (!IndexDatabase(SearchIndexBase, DependentDatabase, DependentDatabaseSchema, DependentSamplingContext, DependentAdditionalExtrapolationTime, Owner))
-					{
-						UE_LOG(LogPoseSearch, Log, TEXT("%s - %s BuildIndex Failed in IndexDatabase"), *LexToString(FullIndexKey.Hash), *MainDatabaseName);
-						ResetSearchIndex();
-						return;
-					}
-
-#if ENABLE_ANIM_DEBUG
-					if (GVarMotionMatchTestDeterminismIndexDatabase)
-					{
-						for (int32 Iteration = 0; Iteration < GVarMotionMatchTestNumIterations; ++Iteration)
-						{
-							FSearchIndexBase TestSearchIndexBase = SearchIndexBase;
-							if (IndexDatabase(TestSearchIndexBase, DependentDatabase, DependentDatabaseSchema, DependentSamplingContext, DependentAdditionalExtrapolationTime, Owner))
-							{
-								if (TestSearchIndexBase != SearchIndexBase)
-								{
-									TStringBuilder<0> Message;
-									CompareSearchIndexBase(TestSearchIndexBase, SearchIndexBase, DependentDatabaseSchema, Message);
-									UE_LOG(LogPoseSearch, Warning, TEXT("OnGetComplete - IndexDatabase is not deterministic\n%s"), *Message);
-								}
-							}
-						}
-					}
-#endif //ENABLE_ANIM_DEBUG
 				}
 
 				static_cast<FSearchIndexBase&>(SearchIndex) = SearchIndexBases[0];
@@ -2111,12 +2082,12 @@ void FPoseSearchDatabaseAsyncCacheTask::OnGetComplete(UE::DerivedData::FCacheGet
 
 						if (TestSearchIndexA.Values != TestSearchIndexB.Values)
 						{
-							UE_LOG(LogPoseSearch, Warning, TEXT("OnGetComplete - PruneDuplicateValues is not deterministic"));
+							UE_LOGF(LogPoseSearch, Warning, "OnGetComplete - PruneDuplicateValues is not deterministic");
 						}
 
 						if (TestSearchIndexA.ValuesVectorToPoseIndexes != TestSearchIndexB.ValuesVectorToPoseIndexes)
 						{
-							UE_LOG(LogPoseSearch, Warning, TEXT("OnGetComplete - PruneDuplicateValues ValuesVectorToPoseIndexes generation is not deterministic"));
+							UE_LOGF(LogPoseSearch, Warning, "OnGetComplete - PruneDuplicateValues ValuesVectorToPoseIndexes generation is not deterministic");
 						}
 					}
 				}
@@ -2131,14 +2102,14 @@ void FPoseSearchDatabaseAsyncCacheTask::OnGetComplete(UE::DerivedData::FCacheGet
 				// Building FSearchIndex
 				if (!PreprocessSearchIndexWeights(SearchIndex, MainDatabaseSchema, Deviation))
 				{
-					UE_LOG(LogPoseSearch, Error, TEXT("%s - %s BuildIndex Failed while preprocessing weights for schema '%s'"), *LexToString(FullIndexKey.Hash), *MainDatabaseName, *MainDatabaseSchema->GetName());
+					UE_LOGF(LogPoseSearch, Error, "%ls - %ls BuildIndex Failed while preprocessing weights for schema '%ls'", *LexToString(FullIndexKey.Hash), *GetDatabaseAndOuterName(MainDatabase), *MainDatabaseSchema->GetName());
 					ResetSearchIndex();
 					return;
 				}
 
 				if (Owner.IsCanceled())
 				{
-					UE_LOG(LogPoseSearch, Log, TEXT("%s - %s BuildIndex Cancelled"), *LexToString(FullIndexKey.Hash), *MainDatabaseName);
+					UE_LOGF(LogPoseSearch, Log, "%ls - %ls BuildIndex Cancelled", *LexToString(FullIndexKey.Hash), *GetDatabaseAndOuterName(MainDatabase));
 					ResetSearchIndex();
 					return;
 				}
@@ -2154,7 +2125,7 @@ void FPoseSearchDatabaseAsyncCacheTask::OnGetComplete(UE::DerivedData::FCacheGet
 					case Eigen::InvalidInput:	Reason = "Invalid Input"; 	 break;
 					default:					Reason = "Unknown Reasons";  break;
 					}
-					UE_LOG(LogPoseSearch, Error, TEXT("%s - %s BuildIndex Failed because of '%s' while calculating PCA data. Try with a different dataset or change the database 'Pose Search Mode'"), *LexToString(FullIndexKey.Hash), *MainDatabaseName, *Reason);
+					UE_LOGF(LogPoseSearch, Error, "%ls - %ls BuildIndex Failed because of '%ls' while calculating PCA data. Try with a different dataset or change the database 'Pose Search Mode'", *LexToString(FullIndexKey.Hash), *GetDatabaseAndOuterName(MainDatabase), *Reason);
 
 					ResetSearchIndex();
 					return;
@@ -2162,7 +2133,7 @@ void FPoseSearchDatabaseAsyncCacheTask::OnGetComplete(UE::DerivedData::FCacheGet
 
 				if (Owner.IsCanceled())
 				{
-					UE_LOG(LogPoseSearch, Log, TEXT("%s - %s BuildIndex Cancelled"), *LexToString(FullIndexKey.Hash), *MainDatabaseName);
+					UE_LOGF(LogPoseSearch, Log, "%ls - %ls BuildIndex Cancelled", *LexToString(FullIndexKey.Hash), *GetDatabaseAndOuterName(MainDatabase));
 					ResetSearchIndex();
 					return;
 				}
@@ -2180,12 +2151,12 @@ void FPoseSearchDatabaseAsyncCacheTask::OnGetComplete(UE::DerivedData::FCacheGet
 
 						if (TestSearchIndexA.PCAValues != TestSearchIndexB.PCAValues)
 						{
-							UE_LOG(LogPoseSearch, Warning, TEXT("OnGetComplete - PruneDuplicatePCAValues is not deterministic"));
+							UE_LOGF(LogPoseSearch, Warning, "OnGetComplete - PruneDuplicatePCAValues is not deterministic");
 						}
 
 						if (TestSearchIndexA.PCAValuesVectorToPoseIndexes != TestSearchIndexB.PCAValuesVectorToPoseIndexes)
 						{
-							UE_LOG(LogPoseSearch, Warning, TEXT("OnGetComplete - PruneDuplicatePCAValues PCAValuesVectorToPoseIndexes generation is not deterministic"));
+							UE_LOGF(LogPoseSearch, Warning, "OnGetComplete - PruneDuplicatePCAValues PCAValuesVectorToPoseIndexes generation is not deterministic");
 						}
 					}
 				}
@@ -2194,7 +2165,7 @@ void FPoseSearchDatabaseAsyncCacheTask::OnGetComplete(UE::DerivedData::FCacheGet
 				SearchIndex.PruneDuplicatePCAValues(MainDatabase->PCAValuesPruningSimilarityThreshold, MainDatabase->GetNumberOfPrincipalComponents());
 				if (Owner.IsCanceled())
 				{
-					UE_LOG(LogPoseSearch, Log, TEXT("%s - %s BuildIndex Cancelled"), *LexToString(FullIndexKey.Hash), *MainDatabaseName);
+					UE_LOGF(LogPoseSearch, Log, "%ls - %ls BuildIndex Cancelled", *LexToString(FullIndexKey.Hash), *GetDatabaseAndOuterName(MainDatabase));
 					ResetSearchIndex();
 					return;
 				}
@@ -2210,7 +2181,7 @@ void FPoseSearchDatabaseAsyncCacheTask::OnGetComplete(UE::DerivedData::FCacheGet
 				PreprocessSearchIndexKDTree(SearchIndex, MainDatabase);
 				if (Owner.IsCanceled())
 				{
-					UE_LOG(LogPoseSearch, Log, TEXT("%s - %s BuildIndex Cancelled"), *LexToString(FullIndexKey.Hash), *MainDatabaseName);
+					UE_LOGF(LogPoseSearch, Log, "%ls - %ls BuildIndex Cancelled", *LexToString(FullIndexKey.Hash), *GetDatabaseAndOuterName(MainDatabase));
 					ResetSearchIndex();
 					return;
 				}
@@ -2224,24 +2195,24 @@ void FPoseSearchDatabaseAsyncCacheTask::OnGetComplete(UE::DerivedData::FCacheGet
 				PreprocessSearchIndexVPTree(SearchIndex, MainDatabase, RandomSeed);
 				if (Owner.IsCanceled())
 				{
-					UE_LOG(LogPoseSearch, Log, TEXT("%s - %s BuildIndex Cancelled"), *LexToString(FullIndexKey.Hash), *MainDatabaseName);
+					UE_LOGF(LogPoseSearch, Log, "%ls - %ls BuildIndex Cancelled", *LexToString(FullIndexKey.Hash), *GetDatabaseAndOuterName(MainDatabase));
 					ResetSearchIndex();
 					return;
 				}
 
-				UE_LOG(LogPoseSearch, Log, TEXT("%s - %s BuildIndex Succeeded"), *LexToString(FullIndexKey.Hash), *MainDatabaseName);
+				UE_LOGF(LogPoseSearch, Log, "%ls - %ls BuildIndex Succeeded", *LexToString(FullIndexKey.Hash), *GetDatabaseAndOuterName(MainDatabase));
 
 #if ENABLE_ANIM_DEBUG
 				if (bCompareSearchIndex && SearchIndexCompare != SearchIndex)
 				{
 					TStringBuilder<0> Message;
 					CompareSearchIndex(SearchIndexCompare, SearchIndex, MainDatabaseSchema, Message);
-					UE_LOG(LogPoseSearch, Warning, TEXT("%s - %s BuildIndex mismatch with DDC Index\n%s"), *LexToString(FullIndexKey.Hash), *MainDatabaseName, *Message);
+					UE_LOGF(LogPoseSearch, Warning, "%ls - %ls BuildIndex mismatch with DDC Index\n%ls", *LexToString(FullIndexKey.Hash), *GetDatabaseAndOuterName(MainDatabase), *Message);
 				}
 
 				if (!ValidateSearchInfoAgainstDatabase(SearchIndex, Database.Get(), FullIndexKey))
 				{
-					UE_LOG(LogPoseSearch, Error, TEXT("%s - %s BuildIndex Failed because of corrupted SearchIndex!"), *LexToString(FullIndexKey.Hash), *GetDatabaseAndOuterName(Database.Get()));
+					UE_LOGF(LogPoseSearch, Error, "%ls - %ls BuildIndex Failed because of corrupted SearchIndex!", *LexToString(FullIndexKey.Hash), *GetDatabaseAndOuterName(Database.Get()));
 					ResetSearchIndex();
 					return;
 				}
@@ -2259,25 +2230,25 @@ void FPoseSearchDatabaseAsyncCacheTask::OnGetComplete(UE::DerivedData::FCacheGet
 				FCacheRecordBuilder Builder(FullIndexKey);
 				Builder.AddValue(Id, RawData);
 
-				GetCache().Put({ { { MainDatabase->GetPathName() }, Builder.Build() } }, Owner, [this, MainDatabaseSchema, MainDatabaseName, FullIndexKey](FCachePutResponse&& Response)
+				GetCache().Put({ { { MainDatabase->GetPathName() }, Builder.Build() } }, Owner, [this, MainDatabaseSchema, MainDatabase, FullIndexKey](FCachePutResponse&& Response)
 					{
 						switch (Response.Status)
 						{
 						case EStatus::Error:
-							UE_LOG(LogPoseSearch, Log, TEXT("%s - %s Failed to store DDC"), *LexToString(FullIndexKey.Hash), *MainDatabaseName);
+							UE_LOGF(LogPoseSearch, Log, "%ls - %ls Failed to store DDC", *LexToString(FullIndexKey.Hash), *GetDatabaseAndOuterName(MainDatabase));
 							break;
 						case EStatus::Canceled:
-							UE_LOG(LogPoseSearch, Log, TEXT("%s - %s Canceled to store DDC"), *LexToString(FullIndexKey.Hash), *MainDatabaseName);
+							UE_LOGF(LogPoseSearch, Log, "%ls - %ls Canceled to store DDC", *LexToString(FullIndexKey.Hash), *GetDatabaseAndOuterName(MainDatabase));
 							break;
 						case EStatus::Ok:
-							UE_LOG(LogPoseSearch, Log, TEXT("%s - %s BuildIndex stored to DDC"), *LexToString(FullIndexKey.Hash), *MainDatabaseName);
+							UE_LOGF(LogPoseSearch, Log, "%ls - %ls BuildIndex stored to DDC", *LexToString(FullIndexKey.Hash), *GetDatabaseAndOuterName(MainDatabase));
 
 #if ENABLE_ANIM_DEBUG
 							if (GVarMotionMatchTestDDCConsistency)
 							{
 								const FCacheKey CacheKey{ Bucket, DerivedDataKey };
 								const FCacheGetRequest CacheRequest = { { Database->GetPathName() }, CacheKey, ECachePolicy::Default };
-								GetCache().Get(MakeArrayView(&CacheRequest, 1), Owner, [this, MainDatabaseSchema, MainDatabaseName](FCacheGetResponse&& Response)
+								GetCache().Get(MakeArrayView(&CacheRequest, 1), Owner, [this, MainDatabaseSchema, MainDatabase](FCacheGetResponse&& Response)
 									{
 										const FCacheKey& FullIndexKey = Response.Record.GetKey();
 										check(FullIndexKey.Hash == DerivedDataKey);
@@ -2294,7 +2265,7 @@ void FPoseSearchDatabaseAsyncCacheTask::OnGetComplete(UE::DerivedData::FCacheGet
 											{
 												TStringBuilder<0> Message;
 												CompareSearchIndex(TestSearchIndex, SearchIndex, MainDatabaseSchema, Message);
-												UE_LOG(LogPoseSearch, Warning, TEXT("%s - %s DDC Index mismatch with BuildIndex\n%s"), *LexToString(FullIndexKey.Hash), *MainDatabaseName, *Message);
+												UE_LOGF(LogPoseSearch, Warning, "%ls - %ls DDC Index mismatch with BuildIndex\n%ls", *LexToString(FullIndexKey.Hash), *GetDatabaseAndOuterName(MainDatabase), *Message);
 											}
 										}
 									});
@@ -2568,7 +2539,7 @@ void FAsyncPoseSearchDatabasesManagement::Tick(float DeltaTime)
 			FFileHelper::SaveArrayToFile(RawBytes, *FileName);
 
 			GVarMotionMatchTestDeterminismPartialKeyHashesSaveFile = false;
-			UE_LOG(LogPoseSearch, Log, TEXT("FAsyncPoseSearchDatabasesManagement::Tick PartialKeyHashes serialized in '%s'"), *FileName);
+			UE_LOGF(LogPoseSearch, Log, "FAsyncPoseSearchDatabasesManagement::Tick PartialKeyHashes serialized in '%ls'", *FileName);
 		}
 
 		if (GVarMotionMatchTestDeterminismPartialKeyHashes)
@@ -2582,12 +2553,12 @@ void FAsyncPoseSearchDatabasesManagement::Tick(float DeltaTime)
 				Reader << SerializablePartialKeyHashesFromFile;
 				if (SerializablePartialKeyHashes.LogDifferencesAgains(SerializablePartialKeyHashesFromFile))
 				{
-					UE_LOG(LogPoseSearch, Log, TEXT("FAsyncPoseSearchDatabasesManagement::Tick PartialKeyHashes determinism validation succeeded"));
+					UE_LOGF(LogPoseSearch, Log, "FAsyncPoseSearchDatabasesManagement::Tick PartialKeyHashes determinism validation succeeded");
 				}
 			}
 			else
 			{
-				UE_LOG(LogPoseSearch, Error, TEXT("FAsyncPoseSearchDatabasesManagement::Tick Failed to load '%s' with serialized PartialKeyHashes"), *FileName);
+				UE_LOGF(LogPoseSearch, Error, "FAsyncPoseSearchDatabasesManagement::Tick Failed to load '%ls' with serialized PartialKeyHashes", *FileName);
 			}
 			GVarMotionMatchTestDeterminismPartialKeyHashes = false;
 		}
@@ -2696,13 +2667,13 @@ void FAsyncPoseSearchDatabasesManagement::Tick(float DeltaTime)
 
 						if (FMemory::Memcmp(&RootTransform, &TestRootTransform, sizeof(FTransform)) != 0)
 						{
-							UE_LOG(LogPoseSearch, Error, TEXT("FAnimationAssetSampler - ExtractRootTransform is not deterministic"));
+							UE_LOGF(LogPoseSearch, Error, "FAnimationAssetSampler - ExtractRootTransform is not deterministic");
 						}
 
 						const TConstArrayView<FTransform> TestBones = TestPose.GetBones();
 						if (Bones.Num() != TestBones.Num())
 						{
-							UE_LOG(LogPoseSearch, Error, TEXT("FAnimationAssetSampler - ExtractPose is not deterministic"));
+							UE_LOGF(LogPoseSearch, Error, "FAnimationAssetSampler - ExtractPose is not deterministic");
 						}
 						else
 						{
@@ -2710,7 +2681,7 @@ void FAsyncPoseSearchDatabasesManagement::Tick(float DeltaTime)
 							{
 								if (FMemory::Memcmp(&Bones[BoneIndex], &TestBones[BoneIndex], sizeof(FTransform)) != 0)
 								{
-									UE_LOG(LogPoseSearch, Error, TEXT("FAnimationAssetSampler - ExtractPose is not deterministic"));
+									UE_LOGF(LogPoseSearch, Error, "FAnimationAssetSampler - ExtractPose is not deterministic");
 								}
 							}
 						}
@@ -2739,11 +2710,11 @@ void FAsyncPoseSearchDatabasesManagement::Tick(float DeltaTime)
 				{
 					if (LoadedData.Num() <= 0 || LoadedData.Num() % sizeof(FTransform) != 0)
 					{
-						UE_LOG(LogPoseSearch, Error, TEXT("FAnimationAssetSampler - Loaded the wrong amount of data!"));
+						UE_LOGF(LogPoseSearch, Error, "FAnimationAssetSampler - Loaded the wrong amount of data!");
 					}
 					else if (TestSample.SerializedData.Num() != LoadedData.Num())
 					{
-						UE_LOG(LogPoseSearch, Error, TEXT("FAnimationAssetSampler - Loaded data mismatch expected amount of data!"));
+						UE_LOGF(LogPoseSearch, Error, "FAnimationAssetSampler - Loaded data mismatch expected amount of data!");
 					}
 					else if (FMemory::Memcmp(TestSample.SerializedData.GetData(), LoadedData.GetData(), TestSample.SerializedData.Num()) != 0)
 					{
@@ -2761,7 +2732,7 @@ void FAsyncPoseSearchDatabasesManagement::Tick(float DeltaTime)
 								// NoTe:	TransformIndex == 0 for the root
 								//			TransformIndex > 0 are the bones, where BoneIndex = TransformIndex - 1
 								const int32 BoneIndex = TransformIndex - 1;
-								UE_LOG(LogPoseSearch, Error, TEXT("FAnimationAssetSampler - ExtractPose is not deterministic for %s for Bone %d"), *FileName, BoneIndex);
+								UE_LOGF(LogPoseSearch, Error, "FAnimationAssetSampler - ExtractPose is not deterministic for %ls for Bone %d", *FileName, BoneIndex);
 							}
 						}
 					}
@@ -2771,7 +2742,7 @@ void FAsyncPoseSearchDatabasesManagement::Tick(float DeltaTime)
 					bool bSavedFile = FFileHelper::SaveArrayToFile(TestSample.SerializedData, *FileName);
 					if (!bSavedFile)
 					{
-						UE_LOG(LogPoseSearch, Error, TEXT("FAnimationAssetSampler - Failed to save comparison file!"));
+						UE_LOGF(LogPoseSearch, Error, "FAnimationAssetSampler - Failed to save comparison file!");
 					}
 				}
 			}
@@ -2788,7 +2759,7 @@ void FAsyncPoseSearchDatabasesManagement::Tick(float DeltaTime)
 				if (Tasks[TaskIndex]->GetState() == FPoseSearchDatabaseAsyncCacheTask::EState::Ended ||
 					Tasks[TaskIndex]->GetState() == FPoseSearchDatabaseAsyncCacheTask::EState::Failed)
 				{
-					UE_LOG(LogPoseSearch, Log, TEXT("%s - %s Removed because of InvalidateCache with WaitForTaskCompletion"), *LexToString(Tasks[TaskIndex]->GetDerivedDataKey()), *GetDatabaseAndOuterName(Tasks[TaskIndex]->GetDatabase()));
+					UE_LOGF(LogPoseSearch, Log, "%ls - %ls Removed because of InvalidateCache with WaitForTaskCompletion", *LexToString(Tasks[TaskIndex]->GetDerivedDataKey()), *GetDatabaseAndOuterName(Tasks[TaskIndex]->GetDatabase()));
 					Tasks.RemoveAtSwap(TaskIndex, EAllowShrinking::No);
 				}
 			}
@@ -2797,7 +2768,7 @@ void FAsyncPoseSearchDatabasesManagement::Tick(float DeltaTime)
 		{
 			for (TUniquePtr<FPoseSearchDatabaseAsyncCacheTask>& TaskPtr : Tasks)
 			{
-				UE_LOG(LogPoseSearch, Log, TEXT("%s - %s Cancelled because of InvalidateCache"), *LexToString(TaskPtr->GetDerivedDataKey()), *GetDatabaseAndOuterName(TaskPtr->GetDatabase()));
+				UE_LOGF(LogPoseSearch, Log, "%ls - %ls Cancelled because of InvalidateCache", *LexToString(TaskPtr->GetDerivedDataKey()), *GetDatabaseAndOuterName(TaskPtr->GetDatabase()));
 			}
 			Tasks.Reset();
 		}
@@ -2818,14 +2789,14 @@ void FAsyncPoseSearchDatabasesManagement::Tick(float DeltaTime)
 	{
 		if (!Tasks[TaskIndex]->IsValid())
 		{
-			UE_LOG(LogPoseSearch, Log, TEXT("%s - Removed because IsValid"), *LexToString(Tasks[TaskIndex]->GetDerivedDataKey()));
+			UE_LOGF(LogPoseSearch, Log, "%ls - Removed because IsValid", *LexToString(Tasks[TaskIndex]->GetDerivedDataKey()));
 			Tasks.RemoveAtSwap(TaskIndex, EAllowShrinking::No);
 		}
 		else if (Tasks[TaskIndex]->GetState() == FPoseSearchDatabaseAsyncCacheTask::EState::Cancelled)
 		{
 			const UPoseSearchDatabase* TaskDatabase = Tasks[TaskIndex]->GetDatabase();
 
-			UE_LOG(LogPoseSearch, Log, TEXT("%s - %s Removed because is Cancelled"), *LexToString(Tasks[TaskIndex]->GetDerivedDataKey()), *GetDatabaseAndOuterName(TaskDatabase));
+			UE_LOGF(LogPoseSearch, Log, "%ls - %ls Removed because is Cancelled", *LexToString(Tasks[TaskIndex]->GetDerivedDataKey()), *GetDatabaseAndOuterName(TaskDatabase));
 			Tasks.RemoveAtSwap(TaskIndex, EAllowShrinking::No);
 
 #if ENABLE_ANIM_DEBUG
@@ -2858,12 +2829,12 @@ void FAsyncPoseSearchDatabasesManagement::Tick(float DeltaTime)
 
 					if (!KeyBuilder.ValidateAgainst(TestKeyBuilder))
 					{
-						UE_LOG(LogPoseSearch, Error, TEXT("FKeyBuilder - key generation is not deterministic: %s / %s for asset %s"), *LexToString(IoHash), *LexToString(TestIoHash), *Database->GetName());
+						UE_LOGF(LogPoseSearch, Error, "FKeyBuilder - key generation is not deterministic: %ls / %ls for asset %ls", *LexToString(IoHash), *LexToString(TestIoHash), *Database->GetName());
 					}
 
 					if (IoHash != TestIoHash)
 					{
-						UE_LOG(LogPoseSearch, Error, TEXT("FKeyBuilder - key generation is not deterministic: %s / %s for asset %s"), *LexToString(IoHash), *LexToString(TestIoHash), *Database->GetName());
+						UE_LOGF(LogPoseSearch, Error, "FKeyBuilder - key generation is not deterministic: %ls / %ls for asset %ls", *LexToString(IoHash), *LexToString(TestIoHash), *Database->GetName());
 					}
 				}
 			}
@@ -2991,7 +2962,7 @@ void FAsyncPoseSearchDatabasesManagement::Clear()
 
 	for (TUniquePtr<FPoseSearchDatabaseAsyncCacheTask>& TaskPtr : This.Tasks)
 	{
-		UE_LOG(LogPoseSearch, Log, TEXT("%s - %s Cancelled because of Clear"), *LexToString(TaskPtr->GetDerivedDataKey()), *GetDatabaseAndOuterName(TaskPtr->GetDatabase()));
+		UE_LOGF(LogPoseSearch, Log, "%ls - %ls Cancelled because of Clear", *LexToString(TaskPtr->GetDerivedDataKey()), *GetDatabaseAndOuterName(TaskPtr->GetDatabase()));
 	}
 	This.Tasks.Reset();
 }
