PoseSearch - converting UE_LOG into UE_LOGF in PoseSearchChooserColumn

#Synced-CL 51688529

diff --git a/Engine/Plugins/Animation/PoseSearch/Source/Runtime/Private/Chooser/PoseSearchChooserColumn.cpp b/Engine/Plugins/Animation/PoseSearch/Source/Runtime/Private/Chooser/PoseSearchChooserColumn.cpp
--- a/Engine/Plugins/Animation/PoseSearch/Source/Runtime/Private/Chooser/PoseSearchChooserColumn.cpp
+++ b/Engine/Plugins/Animation/PoseSearch/Source/Runtime/Private/Chooser/PoseSearchChooserColumn.cpp
@@ -298,7 +298,7 @@ namespace UE::PoseSearch
 						{
 							if (!ensure(GetReferencedObject(Chooser, TableCell.RowIndex) == DatabaseAnimationAsset->AnimAsset))
 							{
-								UE_LOG(LogPoseSearch, Error, TEXT("FMapping::FindDatabaseAnimationAsset - Mismatch between chooser table '%s' with root '%s' and Database '%s' for row %d"), *Chooser->GetName(), *Chooser->GetRootChooser()->GetName(), *Database->GetName(), TableCell.RowIndex);
+								UE_LOGF(LogPoseSearch, Error, "FMapping::FindDatabaseAnimationAsset - Mismatch between chooser table '%ls' with root '%ls' and Database '%ls' for row %d", *Chooser->GetName(), *Chooser->GetRootChooser()->GetName(), *Database->GetName(), TableCell.RowIndex);
 							}
 						}
 					}
@@ -719,7 +719,7 @@ void FPoseSearchColumn::Filter(FChooserEvaluationContext& Context, const FChoose
 
 	if (!PoseHistory)
 	{
-		UE_LOG(LogPoseSearch, Error, TEXT("FPoseSearchColumn::Filter, missing IPoseHistory"));
+		UE_LOGF(LogPoseSearch, Error, "FPoseSearchColumn::Filter, missing IPoseHistory");
 	}
 	else
 	{
@@ -762,7 +762,7 @@ void FPoseSearchColumn::Filter(FChooserEvaluationContext& Context, const FChoose
 					else if (Role != Pair.Key->Schema->GetDefaultRole())
 					{
 						// @todo: improve logging
-						UE_LOG(LogPoseSearch, Error, TEXT("FPoseSearchColumn::Filter inconsistent Roles"));
+						UE_LOGF(LogPoseSearch, Error, "FPoseSearchColumn::Filter inconsistent Roles");
 					}
 #endif // !NO_LOGGING
 					DatabasesToSearch.Add(Pair.Key);
@@ -1123,7 +1123,7 @@ void FPoseSearchColumn::AddToDetails(FInstancedPropertyBag& PropertyBag, int32 C
 	const FPoseSearchColumnRow& RowValue = GetRowValue(RowIndex);
 	if (!ensure(GetReferencedObject(Chooser, RowIndex) == RowValue.Data.AnimAsset))
 	{
-		UE_LOG(LogPoseSearch, Error, TEXT("FPoseSearchColumn::AddToDetails - Mismatch between chooser table '%s' with root '%s' and Database '%s' for row %d"), *Chooser->GetName(), *Chooser->GetRootChooser()->GetName(), *GetNameSafe(RowValue.Database), RowIndex);
+		UE_LOGF(LogPoseSearch, Error, "FPoseSearchColumn::AddToDetails - Mismatch between chooser table '%ls' with root '%ls' and Database '%ls' for row %d", *Chooser->GetName(), *Chooser->GetRootChooser()->GetName(), *GetNameSafe(RowValue.Database), RowIndex);
 	}
 
 	FPoseSearchColumnRowReflection RowReflection;
@@ -1209,7 +1209,7 @@ void FPoseSearchColumn::ReplaceReferences(UObject* ReferenceToReplace, UObject*
 			UPoseSearchDatabase* DatabaseReplaceWith = Cast<UPoseSearchDatabase>(ReplaceWith);
 			if (ReplaceWith && !DatabaseReplaceWith)
 			{
-				UE_LOG(LogPoseSearch, Error, TEXT("FPoseSearchColumn::ReplaceReferences, mismatch between ReferenceToReplace and ReplaceWith"));
+				UE_LOGF(LogPoseSearch, Error, "FPoseSearchColumn::ReplaceReferences, mismatch between ReferenceToReplace and ReplaceWith");
 			}
 			else
 			{
@@ -1244,7 +1244,7 @@ void FPoseSearchColumn::ReplaceReferences(UObject* ReferenceToReplace, UObject*
 			UChooserTable* ChooserToReplaceWith = Cast<UChooserTable>(ReplaceWith);
 			if (ReplaceWith && !ChooserToReplaceWith)
 			{
-				UE_LOG(LogPoseSearch, Error, TEXT("FPoseSearchColumn::ReplaceReferences, mismatch between ReferenceToReplace and ReplaceWith"));
+				UE_LOGF(LogPoseSearch, Error, "FPoseSearchColumn::ReplaceReferences, mismatch between ReferenceToReplace and ReplaceWith");
 			}
 			else if (ChooserToReplace == Chooser)
 			{
