PoseSearch - UE_LOG -> UE_LOGF

#Synced-CL 51688529

diff --git a/Engine/Plugins/Animation/PoseSearch/Source/Editor/Private/PoseSearchDatabaseAssetTree.cpp b/Engine/Plugins/Animation/PoseSearch/Source/Editor/Private/PoseSearchDatabaseAssetTree.cpp
--- a/Engine/Plugins/Animation/PoseSearch/Source/Editor/Private/PoseSearchDatabaseAssetTree.cpp
+++ b/Engine/Plugins/Animation/PoseSearch/Source/Editor/Private/PoseSearchDatabaseAssetTree.cpp
@@ -792,7 +792,7 @@ namespace UE::PoseSearch
 								}
 								else
 								{
-									UE_LOG(LogPoseSearchEditor, Log, TEXT("Failed to copy %s. Asset(s) with BranchIn notifies do not have clipboard support."), *DatabaseAnimationAsset->GetName())
+									UE_LOGF(LogPoseSearchEditor, Log, "Failed to copy %ls. Asset(s) with BranchIn notifies do not have clipboard support.", *DatabaseAnimationAsset->GetName())
 								}
 							}
 						}
@@ -803,7 +803,7 @@ namespace UE::PoseSearch
 			}
 			else
 			{
-				UE_LOG(LogPoseSearchEditor, Warning, TEXT("Failed create clipboard object while attempting to copy data"));
+				UE_LOGF(LogPoseSearchEditor, Warning, "Failed create clipboard object while attempting to copy data");
 			}
 		}
 	}
@@ -830,7 +830,7 @@ namespace UE::PoseSearch
 		}
 		else
 		{
-			UE_LOG(LogPoseSearchEditor, Warning, TEXT("Failed to get valid clipboard data while attempting to paste data"));
+			UE_LOGF(LogPoseSearchEditor, Warning, "Failed to get valid clipboard data while attempting to paste data");
 		}
 	}
 
diff --git a/Engine/Plugins/Animation/PoseSearch/Source/Editor/Private/PoseSearchDatabaseViewModel.cpp b/Engine/Plugins/Animation/PoseSearch/Source/Editor/Private/PoseSearchDatabaseViewModel.cpp
--- a/Engine/Plugins/Animation/PoseSearch/Source/Editor/Private/PoseSearchDatabaseViewModel.cpp
+++ b/Engine/Plugins/Animation/PoseSearch/Source/Editor/Private/PoseSearchDatabaseViewModel.cpp
@@ -54,7 +54,7 @@ bool FDatabasePreviewActor::SpawnPreviewActor(UWorld* World, const UPoseSearchDa
 	USkeleton* Skeleton = PoseSearchDatabase->Schema->GetSkeleton(Role);
 	if (!Skeleton)
 	{
-		UE_LOG(LogPoseSearchEditor, Log, TEXT("Couldn't spawn preview Actor for asset %s because its Role '%s' is missing in Schema '%s'"), *GetNameSafe(DatabaseAnimationAsset->GetAnimationAsset()), *Role.ToString(), *PoseSearchDatabase->Schema->GetName());
+		UE_LOGF(LogPoseSearchEditor, Log, "Couldn't spawn preview Actor for asset %ls because its Role '%ls' is missing in Schema '%ls'", *GetNameSafe(DatabaseAnimationAsset->GetAnimationAsset()), *Role.ToString(), *PoseSearchDatabase->Schema->GetName());
 		return false;
 	}
 
@@ -693,7 +693,7 @@ bool FDatabaseViewModel::DeleteFromDatabase(int32 AnimationAssetIndex)
 				}
 				else
 				{
-					UE_LOG(LogPoseSearchEditor, Error, TEXT("found DatabaseAnimationAssetBase with valid BranchInId, but invalid AnimSequenceBase in %s"), *Database->GetName());
+					UE_LOGF(LogPoseSearchEditor, Error, "found DatabaseAnimationAssetBase with valid BranchInId, but invalid AnimSequenceBase in %ls", *Database->GetName());
 				}
 			}
 
diff --git a/Engine/Plugins/Animation/PoseSearch/Source/Runtime/Private/AnimNode_MotionMatching.cpp b/Engine/Plugins/Animation/PoseSearch/Source/Runtime/Private/AnimNode_MotionMatching.cpp
--- a/Engine/Plugins/Animation/PoseSearch/Source/Runtime/Private/AnimNode_MotionMatching.cpp
+++ b/Engine/Plugins/Animation/PoseSearch/Source/Runtime/Private/AnimNode_MotionMatching.cpp
@@ -148,7 +148,7 @@ void FAnimNode_MotionMatching::UpdateAssetPlayer(const FAnimationUpdateContext&
 
 	if (!PoseHistory)
 	{
-		UE_LOG(LogPoseSearch, Error, TEXT("FAnimNode_MotionMatching::UpdateAssetPlayer, missing IPoseHistory"));
+		UE_LOGF(LogPoseSearch, Error, "FAnimNode_MotionMatching::UpdateAssetPlayer, missing IPoseHistory");
 	}
 	else
 	{
@@ -318,7 +318,7 @@ void FAnimNode_MotionMatching::Evaluate_AnyThread(FPoseContext& Output)
 	UAnimInstance* AnimInstance = Cast<UAnimInstance>(Output.AnimInstanceProxy->GetAnimInstanceObject());
 	if (!AnimInstance)
 	{
-		UE_LOG(LogPoseSearch, Error, TEXT("FAnimNode_MotionMatching::Evaluate_AnyThread couldn't find the AnimInstance!?"));
+		UE_LOGF(LogPoseSearch, Error, "FAnimNode_MotionMatching::Evaluate_AnyThread couldn't find the AnimInstance!?");
 		return;
 	}
 
@@ -326,7 +326,7 @@ void FAnimNode_MotionMatching::Evaluate_AnyThread(FPoseContext& Output)
 
 	if (!RootMotionProvider)
 	{
-		UE_LOG(LogPoseSearch, Error, TEXT("FAnimNode_MotionMatching::Evaluate_AnyThread couldn't find the IAnimRootMotionProvider"));
+		UE_LOGF(LogPoseSearch, Error, "FAnimNode_MotionMatching::Evaluate_AnyThread couldn't find the IAnimRootMotionProvider");
 		return;
 	}
 
@@ -352,7 +352,7 @@ void FAnimNode_MotionMatching::Evaluate_AnyThread(FPoseContext& Output)
 	const UAnimationAsset* AnimationAsset = MultiAnimAsset->GetAnimationAsset(MotionMatchingState.SearchResult.Role);
 	if (!AnimationAsset)
 	{
-		UE_LOG(LogPoseSearch, Error, TEXT("FAnimNode_MotionMatching::Evaluate_AnyThread MultiAnimAsset %s for Role %s is invalid!"), *GetNameSafe(MultiAnimAsset), *MotionMatchingState.SearchResult.Role.ToString());
+		UE_LOGF(LogPoseSearch, Error, "FAnimNode_MotionMatching::Evaluate_AnyThread MultiAnimAsset %ls for Role %ls is invalid!", *GetNameSafe(MultiAnimAsset), *MotionMatchingState.SearchResult.Role.ToString());
 		return;
 	}
 
diff --git a/Engine/Plugins/Animation/PoseSearch/Source/Runtime/Private/AnimNode_MotionMatchingInteraction.cpp b/Engine/Plugins/Animation/PoseSearch/Source/Runtime/Private/AnimNode_MotionMatchingInteraction.cpp
--- a/Engine/Plugins/Animation/PoseSearch/Source/Runtime/Private/AnimNode_MotionMatchingInteraction.cpp
+++ b/Engine/Plugins/Animation/PoseSearch/Source/Runtime/Private/AnimNode_MotionMatchingInteraction.cpp
@@ -87,7 +87,7 @@ void FAnimNode_MotionMatchingInteraction::UpdateAssetPlayer(const FAnimationUpda
 	}
 	else
 	{
-		UE_LOG(LogPoseSearch, Error, TEXT("FAnimNode_MotionMatchingInteraction::Update_AnyThread couldn't find the FPoseHistoryProvider"));
+		UE_LOGF(LogPoseSearch, Error, "FAnimNode_MotionMatchingInteraction::Update_AnyThread couldn't find the FPoseHistoryProvider");
 	}
 
 	const bool bDidBlendToRequestAnInertialBlend = bBlendToExecuted && bUseInertialBlend;
@@ -123,7 +123,7 @@ void FAnimNode_MotionMatchingInteraction::Evaluate_AnyThread(FPoseContext& Outpu
 	UAnimInstance* AnimInstance = Cast<UAnimInstance>(Output.AnimInstanceProxy->GetAnimInstanceObject());
 	if (!AnimInstance)
 	{
-		UE_LOG(LogPoseSearch, Error, TEXT("FAnimNode_MotionMatchingInteraction::Evaluate_AnyThread couldn't find the AnimInstance!?"));
+		UE_LOGF(LogPoseSearch, Error, "FAnimNode_MotionMatchingInteraction::Evaluate_AnyThread couldn't find the AnimInstance!?");
 		return;
 	}
 
@@ -131,7 +131,7 @@ void FAnimNode_MotionMatchingInteraction::Evaluate_AnyThread(FPoseContext& Outpu
 
 	if (!RootMotionProvider)
 	{
-		UE_LOG(LogPoseSearch, Error, TEXT("FAnimNode_MotionMatchingInteraction::Evaluate_AnyThread couldn't find the IAnimRootMotionProvider"));
+		UE_LOGF(LogPoseSearch, Error, "FAnimNode_MotionMatchingInteraction::Evaluate_AnyThread couldn't find the IAnimRootMotionProvider");
 		return;
 	}
 
@@ -157,7 +157,7 @@ void FAnimNode_MotionMatchingInteraction::Evaluate_AnyThread(FPoseContext& Outpu
 	const UAnimationAsset* AnimationAsset = MultiAnimAsset->GetAnimationAsset(CurrentResult.Role);
 	if (!AnimationAsset)
 	{
-		UE_LOG(LogPoseSearch, Error, TEXT("FAnimNode_MotionMatchingInteraction::Evaluate_AnyThread MultiAnimAsset %s for Role %s is invalid!"), *GetNameSafe(MultiAnimAsset), *CurrentResult.Role.ToString());
+		UE_LOGF(LogPoseSearch, Error, "FAnimNode_MotionMatchingInteraction::Evaluate_AnyThread MultiAnimAsset %ls for Role %ls is invalid!", *GetNameSafe(MultiAnimAsset), *CurrentResult.Role.ToString());
 		return;
 	}
 
diff --git a/Engine/Plugins/Animation/PoseSearch/Source/Runtime/Private/MotionMatchingAnimNodeLibrary.cpp b/Engine/Plugins/Animation/PoseSearch/Source/Runtime/Private/MotionMatchingAnimNodeLibrary.cpp
--- a/Engine/Plugins/Animation/PoseSearch/Source/Runtime/Private/MotionMatchingAnimNodeLibrary.cpp
+++ b/Engine/Plugins/Animation/PoseSearch/Source/Runtime/Private/MotionMatchingAnimNodeLibrary.cpp
@@ -68,7 +68,7 @@ void UMotionMatchingAnimNodeLibrary::GetMotionMatchingSearchResult(const FMotion
 	{
 		Result = FPoseSearchBlueprintResult();
 		bIsResultValid = false;
-		UE_LOG(LogPoseSearch, Warning, TEXT("UMotionMatchingAnimNodeLibrary::GetMotionMatchingSearchResult called on an invalid context or with an invalid type"));
+		UE_LOGF(LogPoseSearch, Warning, "UMotionMatchingAnimNodeLibrary::GetMotionMatchingSearchResult called on an invalid context or with an invalid type");
 	}
 }
 
@@ -84,7 +84,7 @@ void UMotionMatchingAnimNodeLibrary::GetMotionMatchingBlendSettings(const FMotio
 	else
 	{
 		BlendSettings = FMotionMatchingBlueprintBlendSettings();
-		UE_LOG(LogPoseSearch, Warning, TEXT("UMotionMatchingAnimNodeLibrary::GetMotionMatchingBlendSettings called on an invalid context or with an invalid type"));
+		UE_LOGF(LogPoseSearch, Warning, "UMotionMatchingAnimNodeLibrary::GetMotionMatchingBlendSettings called on an invalid context or with an invalid type");
 	}
 }
 
@@ -99,7 +99,7 @@ void UMotionMatchingAnimNodeLibrary::OverrideMotionMatchingBlendSettings(const F
 	}
 	else
 	{
-		UE_LOG(LogPoseSearch, Warning, TEXT("UMotionMatchingAnimNodeLibrary::OverrideMotionMatchingBlendSettings called on an invalid context or with an invalid type"));
+		UE_LOGF(LogPoseSearch, Warning, "UMotionMatchingAnimNodeLibrary::OverrideMotionMatchingBlendSettings called on an invalid context or with an invalid type");
 	}
 }
 
@@ -111,7 +111,7 @@ void UMotionMatchingAnimNodeLibrary::SetDatabaseToSearch(const FMotionMatchingAn
 	}
 	else
 	{
-		UE_LOG(LogPoseSearch, Warning, TEXT("UMotionMatchingAnimNodeLibrary::SetDatabaseToSearch called on an invalid context or with an invalid type"));
+		UE_LOGF(LogPoseSearch, Warning, "UMotionMatchingAnimNodeLibrary::SetDatabaseToSearch called on an invalid context or with an invalid type");
 	}
 }
 
@@ -123,7 +123,7 @@ void UMotionMatchingAnimNodeLibrary::SetDatabasesToSearch(const FMotionMatchingA
 	}
 	else
 	{
-		UE_LOG(LogPoseSearch, Warning, TEXT("UMotionMatchingAnimNodeLibrary::SetDatabasesToSearch called on an invalid context or with an invalid type"));
+		UE_LOGF(LogPoseSearch, Warning, "UMotionMatchingAnimNodeLibrary::SetDatabasesToSearch called on an invalid context or with an invalid type");
 	}
 }
 
@@ -135,7 +135,7 @@ void UMotionMatchingAnimNodeLibrary::ResetDatabasesToSearch(const FMotionMatchin
 	}
 	else
 	{
-		UE_LOG(LogPoseSearch, Warning, TEXT("UMotionMatchingAnimNodeLibrary::ResetDatabasesToSearch called on an invalid context or with an invalid type"));
+		UE_LOGF(LogPoseSearch, Warning, "UMotionMatchingAnimNodeLibrary::ResetDatabasesToSearch called on an invalid context or with an invalid type");
 	}
 }
 
@@ -147,7 +147,7 @@ void UMotionMatchingAnimNodeLibrary::SetInterruptMode(const FMotionMatchingAnimN
 	}
 	else
 	{
-		UE_LOG(LogPoseSearch, Warning, TEXT("UMotionMatchingAnimNodeLibrary::SetInterruptMode called on an invalid context or with an invalid type"));
+		UE_LOGF(LogPoseSearch, Warning, "UMotionMatchingAnimNodeLibrary::SetInterruptMode called on an invalid context or with an invalid type");
 	}
 }
 
@@ -158,6 +158,6 @@ bool UMotionMatchingAnimNodeLibrary::AnyNewBlendToThisFrame(const FMotionMatchin
 		return MotionMatchingNodePtr->AnyNewBlendToThisFrame();
 	}
 	
-	UE_LOG(LogPoseSearch, Warning, TEXT("UMotionMatchingAnimNodeLibrary::AnyNewBlendToThisFrame called on an invalid context or with an invalid type"));
+	UE_LOGF(LogPoseSearch, Warning, "UMotionMatchingAnimNodeLibrary::AnyNewBlendToThisFrame called on an invalid context or with an invalid type");
 	return false;
 }
\ No newline at end of file
diff --git a/Engine/Plugins/Animation/PoseSearch/Source/Runtime/Private/MotionMatchingInteractionAnimNodeLibrary.cpp b/Engine/Plugins/Animation/PoseSearch/Source/Runtime/Private/MotionMatchingInteractionAnimNodeLibrary.cpp
--- a/Engine/Plugins/Animation/PoseSearch/Source/Runtime/Private/MotionMatchingInteractionAnimNodeLibrary.cpp
+++ b/Engine/Plugins/Animation/PoseSearch/Source/Runtime/Private/MotionMatchingInteractionAnimNodeLibrary.cpp
@@ -18,7 +18,7 @@ void UMotionMatchingInteractionAnimNodeLibrary::SetAvailabilities(const FMotionM
 	}
 	else
 	{
-		UE_LOG(LogPoseSearch, Warning, TEXT("UMotionMatchingInteractionAnimNodeLibrary::SetAvailabilities called on an invalid context or with an invalid type"));
+		UE_LOGF(LogPoseSearch, Warning, "UMotionMatchingInteractionAnimNodeLibrary::SetAvailabilities called on an invalid context or with an invalid type");
 	}
 }
 
@@ -29,6 +29,6 @@ bool UMotionMatchingInteractionAnimNodeLibrary::IsInteracting(const FMotionMatch
 		return MotionMatchingInteractionNodePtr->IsInteracting();
 	}
 
-	UE_LOG(LogPoseSearch, Warning, TEXT("UMotionMatchingInteractionAnimNodeLibrary::IsInteracting called on an invalid context or with an invalid type"));
+	UE_LOGF(LogPoseSearch, Warning, "UMotionMatchingInteractionAnimNodeLibrary::IsInteracting called on an invalid context or with an invalid type");
 	return false;
 }
diff --git a/Engine/Plugins/Animation/PoseSearch/Source/Runtime/Private/PoseSearchAssetIndexer.cpp b/Engine/Plugins/Animation/PoseSearch/Source/Runtime/Private/PoseSearchAssetIndexer.cpp
--- a/Engine/Plugins/Animation/PoseSearch/Source/Runtime/Private/PoseSearchAssetIndexer.cpp
+++ b/Engine/Plugins/Animation/PoseSearch/Source/Runtime/Private/PoseSearchAssetIndexer.cpp
@@ -336,9 +336,9 @@ FAssetIndexer::FCachedEntry& FAssetIndexer::GetEntry(float SampleTime)
 		{
 			if (!BoneContainers[RoleIndex].IsValid())
 			{
-				UE_LOG(LogPoseSearch,
+				UE_LOGF(LogPoseSearch,
 					Warning,
-					TEXT("Invalid BoneContainer encountered in FAssetIndexer::GetEntry. Asset: %s. Schema: %s. BoneContainerAsset: %s. NumBoneIndices: %d"),
+					"Invalid BoneContainer encountered in FAssetIndexer::GetEntry. Asset: %ls. Schema: %ls. BoneContainerAsset: %ls. NumBoneIndices: %d",
 					*AssetSamplers.GetAssetName(),
 					*GetNameSafe(&Schema),
 					*GetNameSafe(BoneContainers[RoleIndex].GetAsset()),
@@ -366,7 +366,7 @@ FAssetIndexer::FCachedEntry& FAssetIndexer::GetEntry(float SampleTime)
 				const TConstArrayView<FTransform> TestBones = TestPose.GetBones();
 				if (Bones.Num() != TestBones.Num())
 				{
-					UE_LOG(LogPoseSearch, Warning, TEXT("FAssetIndexer::GetEntry - ExtractPose is not deterministic"));
+					UE_LOGF(LogPoseSearch, Warning, "FAssetIndexer::GetEntry - ExtractPose is not deterministic");
 				}
 				else
 				{
@@ -374,7 +374,7 @@ FAssetIndexer::FCachedEntry& FAssetIndexer::GetEntry(float SampleTime)
 					{
 						if (FMemory::Memcmp(&Bones[BoneIndex], &TestBones[BoneIndex], sizeof(FTransform)) != 0)
 						{
-							UE_LOG(LogPoseSearch, Warning, TEXT("FAssetIndexer::GetEntry - ExtractPose is not deterministic"));
+							UE_LOGF(LogPoseSearch, Warning, "FAssetIndexer::GetEntry - ExtractPose is not deterministic");
 						}
 					}
 				}
@@ -478,7 +478,7 @@ bool FAssetIndexer::GetSampleRotation(FQuat& OutSampleRotation, float SampleTime
 					return true;
 				}
 
-				UE_LOG(LogPoseSearch, Error, TEXT("FAssetIndexer::GetSampleRotation: required UAnimNotifyState_PoseSearchSamplingAttribute in '%s' has an invalid Bone"), *AssetSamplers.GetAssetName());
+				UE_LOGF(LogPoseSearch, Error, "FAssetIndexer::GetSampleRotation: required UAnimNotifyState_PoseSearchSamplingAttribute in '%ls' has an invalid Bone", *AssetSamplers.GetAssetName());
 				return false;
 			}
 
@@ -487,7 +487,7 @@ bool FAssetIndexer::GetSampleRotation(FQuat& OutSampleRotation, float SampleTime
 			return true;
 		}
 
-		UE_LOG(LogPoseSearch, Error, TEXT("FAssetIndexer::GetSampleRotation: required UAnimNotifyState_PoseSearchSamplingAttribute not found in '%s'"), *AssetSamplers.GetAssetName());
+		UE_LOGF(LogPoseSearch, Error, "FAssetIndexer::GetSampleRotation: required UAnimNotifyState_PoseSearchSamplingAttribute not found in '%ls'", *AssetSamplers.GetAssetName());
 		OutSampleRotation = FQuat::Identity;
 		return false;
 	}
@@ -563,7 +563,7 @@ bool FAssetIndexer::GetSamplePositionInternal(FVector& OutSamplePosition, float
 					return true;
 				}
 
-				UE_LOG(LogPoseSearch, Error, TEXT("FAssetIndexer::GetSamplePositionInternal: required UAnimNotifyState_PoseSearchSamplingAttribute in '%s' has an invalid Bone"), *AssetSamplers.GetAssetName());
+				UE_LOGF(LogPoseSearch, Error, "FAssetIndexer::GetSamplePositionInternal: required UAnimNotifyState_PoseSearchSamplingAttribute in '%ls' has an invalid Bone", *AssetSamplers.GetAssetName());
 				return false;
 			}
 
@@ -572,7 +572,7 @@ bool FAssetIndexer::GetSamplePositionInternal(FVector& OutSamplePosition, float
 			return true;
 		}
 
-		UE_LOG(LogPoseSearch, Error, TEXT("FAssetIndexer::GetSamplePositionInternal: required UAnimNotifyState_PoseSearchSamplingAttribute not found in '%s'"), *AssetSamplers.GetAssetName());
+		UE_LOGF(LogPoseSearch, Error, "FAssetIndexer::GetSamplePositionInternal: required UAnimNotifyState_PoseSearchSamplingAttribute not found in '%ls'", *AssetSamplers.GetAssetName());
 		OutSamplePosition = FVector::ZeroVector;
 		return false;
 	}
@@ -627,7 +627,7 @@ bool FAssetIndexer::GetSampleVelocity(FVector& OutSampleVelocity, float SampleTi
 					return false;
 				}
 
-				UE_LOG(LogPoseSearch, Error, TEXT("FAssetIndexer::GetSampleVelocity: required UAnimNotifyState_PoseSearchSamplingAttribute in '%s' has an invalid Bone"), *AssetSamplers.GetAssetName());
+				UE_LOGF(LogPoseSearch, Error, "FAssetIndexer::GetSampleVelocity: required UAnimNotifyState_PoseSearchSamplingAttribute in '%ls' has an invalid Bone", *AssetSamplers.GetAssetName());
 				return false;
 			}
 
@@ -636,7 +636,7 @@ bool FAssetIndexer::GetSampleVelocity(FVector& OutSampleVelocity, float SampleTi
 			return true;
 		}
 
-		UE_LOG(LogPoseSearch, Error, TEXT("FAssetIndexer::GetSampleVelocity: required UAnimNotifyState_PoseSearchSamplingAttribute not found in '%s'"), *AssetSamplers.GetAssetName());
+		UE_LOGF(LogPoseSearch, Error, "FAssetIndexer::GetSampleVelocity: required UAnimNotifyState_PoseSearchSamplingAttribute not found in '%ls'", *AssetSamplers.GetAssetName());
 		OutSampleVelocity = FVector::ZeroVector;
 		return false;
 	}
@@ -719,7 +719,7 @@ void FAssetIndexer::CompareCachedEntries(const FAssetIndexer& Other) const
 {
 	if (CachedEntries.Num() != Other.CachedEntries.Num())
 	{
-		UE_LOG(LogPoseSearch, Warning, TEXT("CompareCachedEntries - FAssetIndexer::CachedEntries::Num is not deterministic"));
+		UE_LOGF(LogPoseSearch, Warning, "CompareCachedEntries - FAssetIndexer::CachedEntries::Num is not deterministic");
 	}
 	else
 	{
@@ -730,12 +730,12 @@ void FAssetIndexer::CompareCachedEntries(const FAssetIndexer& Other) const
 				const FCachedEntry* Entry = &Pair.Value;
 				if (Entry->SampleTime != OtherEntry->SampleTime)
 				{
-					UE_LOG(LogPoseSearch, Warning, TEXT("CompareCachedEntries - FAssetIndexer::CachedEntries::SampleTime is not deterministic (%f, %f)"), Entry->SampleTime, OtherEntry->SampleTime);
+					UE_LOGF(LogPoseSearch, Warning, "CompareCachedEntries - FAssetIndexer::CachedEntries::SampleTime is not deterministic (%f, %f)", Entry->SampleTime, OtherEntry->SampleTime);
 				}
 
 				if (Entry->RootTransform.Num() != OtherEntry->RootTransform.Num())
 				{
-					UE_LOG(LogPoseSearch, Warning, TEXT("CompareCachedEntries - FAssetIndexer::CachedEntries::RootTransform::Num is not deterministic"));
+					UE_LOGF(LogPoseSearch, Warning, "CompareCachedEntries - FAssetIndexer::CachedEntries::RootTransform::Num is not deterministic");
 				}
 				else
 				{
@@ -743,14 +743,14 @@ void FAssetIndexer::CompareCachedEntries(const FAssetIndexer& Other) const
 					{
 						if (FMemory::Memcmp(&Entry->RootTransform[RoleIndex], &OtherEntry->RootTransform[RoleIndex], sizeof(FTransform)) != 0)
 						{
-							UE_LOG(LogPoseSearch, Warning, TEXT("CompareCachedEntries - FAssetIndexer::CachedEntries::RootTransform[%d] is not deterministic"), RoleIndex);
+							UE_LOGF(LogPoseSearch, Warning, "CompareCachedEntries - FAssetIndexer::CachedEntries::RootTransform[%d] is not deterministic", RoleIndex);
 						}
 					}
 				}
 
 				if (Entry->ComponentSpacePose.Num() != OtherEntry->ComponentSpacePose.Num())
 				{
-					UE_LOG(LogPoseSearch, Warning, TEXT("CompareCachedEntries - FAssetIndexer::CachedEntries::ComponentSpacePose::Num is not deterministic"));
+					UE_LOGF(LogPoseSearch, Warning, "CompareCachedEntries - FAssetIndexer::CachedEntries::ComponentSpacePose::Num is not deterministic");
 				}
 				else
 				{
@@ -758,14 +758,14 @@ void FAssetIndexer::CompareCachedEntries(const FAssetIndexer& Other) const
 					{
 						if (Entry->ComponentSpacePose[RoleIndex].GetComponentSpaceFlags() != OtherEntry->ComponentSpacePose[RoleIndex].GetComponentSpaceFlags())
 						{
-							UE_LOG(LogPoseSearch, Warning, TEXT("CompareCachedEntries - FAssetIndexer::CachedEntries::ComponentSpacePose[%d]::ComponentSpaceFlags is not deterministic"), RoleIndex);
+							UE_LOGF(LogPoseSearch, Warning, "CompareCachedEntries - FAssetIndexer::CachedEntries::ComponentSpacePose[%d]::ComponentSpaceFlags is not deterministic", RoleIndex);
 						}
 
 						const TConstArrayView<FTransform> Bones = Entry->ComponentSpacePose[RoleIndex].GetPose().GetBones();
 						const TConstArrayView<FTransform> OtherBones = OtherEntry->ComponentSpacePose[RoleIndex].GetPose().GetBones();
 						if (Bones.Num() != OtherBones.Num())
 						{
-							UE_LOG(LogPoseSearch, Warning, TEXT("CompareCachedEntries - FAssetIndexer::CachedEntries::ComponentSpacePose[%d]::Bones is not deterministic"), RoleIndex);
+							UE_LOGF(LogPoseSearch, Warning, "CompareCachedEntries - FAssetIndexer::CachedEntries::ComponentSpacePose[%d]::Bones is not deterministic", RoleIndex);
 						}
 						else
 						{
@@ -773,7 +773,7 @@ void FAssetIndexer::CompareCachedEntries(const FAssetIndexer& Other) const
 							{
 								if (FMemory::Memcmp(&Bones[BoneIndex], &OtherBones[BoneIndex], sizeof(FTransform)) != 0)
 								{
-									UE_LOG(LogPoseSearch, Warning, TEXT("CompareCachedEntries - FAssetIndexer::CachedEntries::ComponentSpacePose[%d]::Bones[%d] is not deterministic"), RoleIndex, BoneIndex);
+									UE_LOGF(LogPoseSearch, Warning, "CompareCachedEntries - FAssetIndexer::CachedEntries::ComponentSpacePose[%d]::Bones[%d] is not deterministic", RoleIndex, BoneIndex);
 								}
 							}
 						}
@@ -782,7 +782,7 @@ void FAssetIndexer::CompareCachedEntries(const FAssetIndexer& Other) const
 			}
 			else
 			{
-				UE_LOG(LogPoseSearch, Warning, TEXT("CompareCachedEntries - FAssetIndexer::CachedEntries is not deterministic. Missing CachedEntry at time %f"), Pair.Key);
+				UE_LOGF(LogPoseSearch, Warning, "CompareCachedEntries - FAssetIndexer::CachedEntries is not deterministic. Missing CachedEntry at time %f", Pair.Key);
 			}
 		}
 	}
diff --git a/Engine/Plugins/Animation/PoseSearch/Source/Runtime/Private/PoseSearchAssetSampler.cpp b/Engine/Plugins/Animation/PoseSearch/Source/Runtime/Private/PoseSearchAssetSampler.cpp
--- a/Engine/Plugins/Animation/PoseSearch/Source/Runtime/Private/PoseSearchAssetSampler.cpp
+++ b/Engine/Plugins/Animation/PoseSearch/Source/Runtime/Private/PoseSearchAssetSampler.cpp
@@ -56,7 +56,7 @@ static FTransform ExtractRootTransformInternal(const UAnimMontage* AnimMontage,
 	// @todo: add support for SlotName / multiple SlotAnimTracks
 	if (AnimMontage->SlotAnimTracks.Num() != 1)
 	{
-		UE_LOG(LogPoseSearch, Error, TEXT("ExtractRootTransformInternal: so far we support only montages with one SlotAnimTracks. %s has %d"), *AnimMontage->GetName(), AnimMontage->SlotAnimTracks.Num());
+		UE_LOGF(LogPoseSearch, Error, "ExtractRootTransformInternal: so far we support only montages with one SlotAnimTracks. %ls has %d", *AnimMontage->GetName(), AnimMontage->SlotAnimTracks.Num());
 		return FTransform::Identity;
 	}
 
@@ -337,7 +337,7 @@ float FAnimationAssetSampler::ToNormalizedTime(float RealTime) const
 			return NormalizedTime;
 		}
 
-		UE_LOG(LogPoseSearch, Error, TEXT("FAnimationAssetSampler::ToNormalizedTime: requested RealTime %f is greater than CachedPlayLength %f for UBlendSpace %s!"), RealTime, CachedPlayLength, *AnimationAssetPtr->GetName());
+		UE_LOGF(LogPoseSearch, Error, "FAnimationAssetSampler::ToNormalizedTime: requested RealTime %f is greater than CachedPlayLength %f for UBlendSpace %ls!", RealTime, CachedPlayLength, *AnimationAssetPtr->GetName());
 		return FMath::Clamp(NormalizedTime, 0.f, 1.f);
 	}
 
@@ -379,7 +379,7 @@ FTransform FAnimationAssetSampler::GetTotalRootTransform() const
 		// @todo: add support for SlotName / multiple SlotAnimTracks
 		if (AnimMontage->SlotAnimTracks.Num() != 1)
 		{
-			UE_LOG(LogPoseSearch, Error, TEXT("FAssetSamplerBase::GetTotalRootTransform: so far we support only montages with one SlotAnimTracks. %s has %d"), *AnimMontage->GetName(), AnimMontage->SlotAnimTracks.Num());
+			UE_LOGF(LogPoseSearch, Error, "FAssetSamplerBase::GetTotalRootTransform: so far we support only montages with one SlotAnimTracks. %ls has %d", *AnimMontage->GetName(), AnimMontage->SlotAnimTracks.Num());
 			return FTransform::Identity;
 		}
 
@@ -412,7 +412,7 @@ void FAnimationAssetSampler::ExtractPose(const FAnimExtractContext& ExtractionCt
 		}
 		else
 		{
-			UE_LOG(LogPoseSearch, Error, TEXT("FAnimMontageSampler::ExtractPose: UBlendSpace %s couldn't return a valid pose for BlendParameters (%.2f, %.2f, %.2f)"), *BlendSpace->GetName(), BlendParameters.X, BlendParameters.Y, BlendParameters.Z);
+			UE_LOGF(LogPoseSearch, Error, "FAnimMontageSampler::ExtractPose: UBlendSpace %ls couldn't return a valid pose for BlendParameters (%.2f, %.2f, %.2f)", *BlendSpace->GetName(), BlendParameters.X, BlendParameters.Y, BlendParameters.Z);
 			OutAnimPoseData.GetPose().ResetToRefPose();
 		}
 	}
@@ -421,7 +421,7 @@ void FAnimationAssetSampler::ExtractPose(const FAnimExtractContext& ExtractionCt
 		// @todo: add support for SlotName / multiple SlotAnimTracks
 		if (AnimMontage->SlotAnimTracks.Num() != 1)
 		{
-			UE_LOG(LogPoseSearch, Error, TEXT("FAnimMontageSampler::ExtractPose: so far we support only montages with one SlotAnimTracks. %s has %d"), *AnimMontage->GetName(), AnimMontage->SlotAnimTracks.Num());
+			UE_LOGF(LogPoseSearch, Error, "FAnimMontageSampler::ExtractPose: so far we support only montages with one SlotAnimTracks. %ls has %d", *AnimMontage->GetName(), AnimMontage->SlotAnimTracks.Num());
 			OutAnimPoseData.GetPose().ResetToRefPose();
 		}
 		else
@@ -672,7 +672,7 @@ void FAnimationAssetSampler::ExtractAnimNotifyContext(float Time, FAnimNotifyCon
 		}
 		else
 		{
-			UE_LOG(LogPoseSearch, Error, TEXT("FAnimationAssetSampler::ExtractAnimNotifyContext: Unsupported BlendSpace NotifyTriggerMode for '%s'"), *BlendSpace->GetName());
+			UE_LOGF(LogPoseSearch, Error, "FAnimationAssetSampler::ExtractAnimNotifyContext: Unsupported BlendSpace NotifyTriggerMode for '%ls'", *BlendSpace->GetName());
 		}
 	}
 	else if (const UAnimSequenceBase* SequenceBase = Cast<UAnimSequenceBase>(AnimationAssetPtr.Get()))
@@ -742,7 +742,7 @@ TConstArrayView<FAnimNotifyEvent> FAnimationAssetSampler::GetAllAnimNotifyEvents
 		}
 		else
 		{
-			UE_LOG(LogPoseSearch, Error, TEXT("FAnimationAssetSampler::GetAllAnimNotifyEvents: Unsupported BlendSpace NotifyTriggerMode for '%s'"), *BlendSpace->GetName());
+			UE_LOGF(LogPoseSearch, Error, "FAnimationAssetSampler::GetAllAnimNotifyEvents: Unsupported BlendSpace NotifyTriggerMode for '%ls'", *BlendSpace->GetName());
 		}
 	}
 	else if (const UAnimSequenceBase* SequenceBase = Cast<UAnimSequenceBase>(AnimationAssetPtr.Get()))
diff --git a/Engine/Plugins/Animation/PoseSearch/Source/Runtime/Private/PoseSearchAssetSamplerLibrary.cpp b/Engine/Plugins/Animation/PoseSearch/Source/Runtime/Private/PoseSearchAssetSamplerLibrary.cpp
--- a/Engine/Plugins/Animation/PoseSearch/Source/Runtime/Private/PoseSearchAssetSamplerLibrary.cpp
+++ b/Engine/Plugins/Animation/PoseSearch/Source/Runtime/Private/PoseSearchAssetSamplerLibrary.cpp
@@ -15,13 +15,13 @@ FPoseSearchAssetSamplerPose UPoseSearchAssetSamplerLibrary::SamplePose(const FPo
 	using namespace UE::PoseSearch;
 	if (!Input.Animation)
 	{
-		UE_LOG(LogPoseSearch, Error, TEXT("UPoseSearchAssetSamplerLibrary::SamplePose invalid Input.Animation"));
+		UE_LOGF(LogPoseSearch, Error, "UPoseSearchAssetSamplerLibrary::SamplePose invalid Input.Animation");
 		return FPoseSearchAssetSamplerPose();
 	}
 	
 	if (Input.bMirrored && !Input.MirrorDataTable)
 	{
-		UE_LOG(LogPoseSearch, Error, TEXT("UPoseSearchAssetSamplerLibrary::SamplePose unable to mirror the pose from %s at time %f because of invalid MirrorDataTable"), *Input.Animation->GetName(), Input.AnimationTime);
+		UE_LOGF(LogPoseSearch, Error, "UPoseSearchAssetSamplerLibrary::SamplePose unable to mirror the pose from %ls at time %f because of invalid MirrorDataTable", *Input.Animation->GetName(), Input.AnimationTime);
 		return FPoseSearchAssetSamplerPose();
 	}
 
@@ -87,7 +87,7 @@ FTransform UPoseSearchAssetSamplerLibrary::GetTransformByName(UPARAM(ref) FPoseS
 {
 	if (!AssetSamplerPose.Pose.IsValid())
 	{
-		UE_LOG(LogPoseSearch, Error, TEXT("UPoseSearchAssetSamplerLibrary::GetTransformByName invalid AssetSamplerPose.Pose"));
+		UE_LOGF(LogPoseSearch, Error, "UPoseSearchAssetSamplerLibrary::GetTransformByName invalid AssetSamplerPose.Pose");
 		return FTransform::Identity;
 	}
 
@@ -99,14 +99,14 @@ FTransform UPoseSearchAssetSamplerLibrary::GetTransformByName(UPARAM(ref) FPoseS
 	BoneReference.Initialize(Skeleton);
 	if (!BoneReference.HasValidSetup())
 	{
-		UE_LOG(LogPoseSearch, Error, TEXT("UPoseSearchAssetSamplerLibrary::GetTransformByName invalid BoneName %s for Skeleton %s"), *BoneName.ToString(), *GetNameSafe(Skeleton));
+		UE_LOGF(LogPoseSearch, Error, "UPoseSearchAssetSamplerLibrary::GetTransformByName invalid BoneName %ls for Skeleton %ls", *BoneName.ToString(), *GetNameSafe(Skeleton));
 		return FTransform::Identity;
 	}
 
 	const FCompactPoseBoneIndex CompactPoseBoneIndex = BoneContainer.GetCompactPoseIndexFromSkeletonPoseIndex(FSkeletonPoseBoneIndex(BoneReference.BoneIndex));
 	if (!CompactPoseBoneIndex.IsValid())
 	{
-		UE_LOG(LogPoseSearch, Error, TEXT("UPoseSearchAssetSamplerLibrary::GetTransformByName invalid FCompactPoseBoneIndex for BoneName %s for Skeleton %s"), *BoneName.ToString(), *GetNameSafe(Skeleton));
+		UE_LOGF(LogPoseSearch, Error, "UPoseSearchAssetSamplerLibrary::GetTransformByName invalid FCompactPoseBoneIndex for BoneName %ls for Skeleton %ls", *BoneName.ToString(), *GetNameSafe(Skeleton));
 		return FTransform::Identity;
 	}
 
diff --git a/Engine/Plugins/Animation/PoseSearch/Source/Runtime/Private/PoseSearchConstraints.cpp b/Engine/Plugins/Animation/PoseSearch/Source/Runtime/Private/PoseSearchConstraints.cpp
--- a/Engine/Plugins/Animation/PoseSearch/Source/Runtime/Private/PoseSearchConstraints.cpp
+++ b/Engine/Plugins/Animation/PoseSearch/Source/Runtime/Private/PoseSearchConstraints.cpp
@@ -105,7 +105,7 @@ namespace UE::PoseSearch
 
 					if (!Animation)
 					{
-						UE_LOG(LogPoseSearch, Error, TEXT("FConstraints::ProcessAll aborting because of invalid animation asset for Role '%s' in MultiAnimAsset '%s'"), *Role.ToString(), *MultiAnimAsset->GetName());
+						UE_LOGF(LogPoseSearch, Error, "FConstraints::ProcessAll aborting because of invalid animation asset for Role '%ls' in MultiAnimAsset '%ls'", *Role.ToString(), *MultiAnimAsset->GetName());
 						return;
 					}
 
@@ -252,7 +252,7 @@ namespace UE::PoseSearch
 										{
 											const FVector Pos = RootBoneTransform.GetTranslation();
 											const FRotator Rot(RootBoneTransform.GetRotation());
-											UE_LOG(LogPoseSearch, Error, TEXT("FConstraints::ProcessAll unsupported non identity root bone in %s at time %f Pos(%f, %f, %f), Rot(%f, %f, %f)"), *Animation->GetName(), AssetTime, Pos.X, Pos.Y, Pos.Z, Rot.Pitch, Rot.Yaw, Rot.Roll);
+											UE_LOGF(LogPoseSearch, Error, "FConstraints::ProcessAll unsupported non identity root bone in %ls at time %f Pos(%f, %f, %f), Rot(%f, %f, %f)", *Animation->GetName(), AssetTime, Pos.X, Pos.Y, Pos.Z, Rot.Pitch, Rot.Yaw, Rot.Roll);
 										}
 										RootBoneTransform = Sampler.ExtractRootTransform(AssetTime);
 										ComponentSpacePoses[AnimContextIndex].InitPose(Poses[AnimContextIndex]);
@@ -305,12 +305,12 @@ namespace UE::PoseSearch
 #if !NO_LOGGING
 							if (!FromSocketRoleIndex)
 							{
-								UE_LOG(LogPoseSearch, Error, TEXT("FConstraints::ProcessAll, Failed to find Role '%s' used in a PoseSearchConstraint in asset '%s'"), *Constraint.FromSocketRole.ToString(), *MultiAnimAsset->GetName());
+								UE_LOGF(LogPoseSearch, Error, "FConstraints::ProcessAll, Failed to find Role '%ls' used in a PoseSearchConstraint in asset '%ls'", *Constraint.FromSocketRole.ToString(), *MultiAnimAsset->GetName());
 							}
 
 							if (!ToSocketRoleIndex)
 							{
-								UE_LOG(LogPoseSearch, Error, TEXT("FConstraints::ProcessAll, Failed to find Role '%s' used in a PoseSearchConstraint in asset '%s'"), *Constraint.ToSocketRole.ToString(), *MultiAnimAsset->GetName());
+								UE_LOGF(LogPoseSearch, Error, "FConstraints::ProcessAll, Failed to find Role '%ls' used in a PoseSearchConstraint in asset '%ls'", *Constraint.ToSocketRole.ToString(), *MultiAnimAsset->GetName());
 							}
 #endif // !NO_LOGGING
 							ConstraintIt.RemoveCurrent();
@@ -338,12 +338,12 @@ namespace UE::PoseSearch
 #if !NO_LOGGING
 								if (FromBoneIndex == INDEX_NONE)
 								{
-									UE_LOG(LogPoseSearch, Error, TEXT("FConstraints::ProcessAll, Failed to find Socket '%s' from SkeletalMeshComponent '%s' for Role '%s'"), *Constraint.FromSocketName.ToString(), *GetNameSafe(SkeletalMeshComponents[*FromSocketRoleIndex]), *Constraint.FromSocketRole.ToString());
+									UE_LOGF(LogPoseSearch, Error, "FConstraints::ProcessAll, Failed to find Socket '%ls' from SkeletalMeshComponent '%ls' for Role '%ls'", *Constraint.FromSocketName.ToString(), *GetNameSafe(SkeletalMeshComponents[*FromSocketRoleIndex]), *Constraint.FromSocketRole.ToString());
 								}
 
 								if (ToBoneIndex == INDEX_NONE)
 								{
-									UE_LOG(LogPoseSearch, Error, TEXT("FConstraints::ProcessAll, Failed to find Socket '%s' from SkeletalMeshComponent '%s' for Role '%s'"), *Constraint.ToSocketName.ToString(), *GetNameSafe(SkeletalMeshComponents[*ToSocketRoleIndex]), *Constraint.ToSocketRole.ToString());
+									UE_LOGF(LogPoseSearch, Error, "FConstraints::ProcessAll, Failed to find Socket '%ls' from SkeletalMeshComponent '%ls' for Role '%ls'", *Constraint.ToSocketName.ToString(), *GetNameSafe(SkeletalMeshComponents[*ToSocketRoleIndex]), *Constraint.ToSocketRole.ToString());
 								}
 #endif // !NO_LOGGING
 								ConstraintIt.RemoveCurrent();
diff --git a/Engine/Plugins/Animation/PoseSearch/Source/Runtime/Private/PoseSearchContext.cpp b/Engine/Plugins/Animation/PoseSearch/Source/Runtime/Private/PoseSearchContext.cpp
--- a/Engine/Plugins/Animation/PoseSearch/Source/Runtime/Private/PoseSearchContext.cpp
+++ b/Engine/Plugins/Animation/PoseSearch/Source/Runtime/Private/PoseSearchContext.cpp
@@ -960,7 +960,7 @@ float FSearchContext::GetSampleCurveValueInternal(float SampleTime, const FName&
 #if !NO_LOGGING
 	if (!bSuccess && !PoseHistory->IsEmpty())
 	{
-		UE_LOG(LogPoseSearch, Warning, TEXT("FSearchContext::GetSampleCurveValueInternal - Couldn't find Curve %s in the input IPoseHistory requested by %s. Consider adding it to the Pose History!"), *CurveName.ToString(), *Schema->GetName());
+		UE_LOGF(LogPoseSearch, Warning, "FSearchContext::GetSampleCurveValueInternal - Couldn't find Curve %ls in the input IPoseHistory requested by %ls. Consider adding it to the Pose History!", *CurveName.ToString(), *Schema->GetName());
 	}
 #endif // !NO_LOGGING
 
@@ -1155,12 +1155,12 @@ FTransform FSearchContext::GetWorldBoneTransformAtTime(float SampleTime, const F
 			{
 				if (!PoseHistory->IsEmpty())
 				{
-					UE_LOG(LogPoseSearch, Warning, TEXT("FSearchContext::GetWorldBoneTransformAtTime - Couldn't find BoneIndexType %d (%s) for Skeleton %s in the input IPoseHistory requested by %s. Consider adding it to the Pose History!"), BoneIndexType, *Skeleton->GetReferenceSkeleton().GetBoneName(BoneIndexType).ToString(), *Skeleton->GetName(), *Schema->GetName());
+					UE_LOGF(LogPoseSearch, Warning, "FSearchContext::GetWorldBoneTransformAtTime - Couldn't find BoneIndexType %d (%ls) for Skeleton %ls in the input IPoseHistory requested by %ls. Consider adding it to the Pose History!", BoneIndexType, *Skeleton->GetReferenceSkeleton().GetBoneName(BoneIndexType).ToString(), *Skeleton->GetName(), *Schema->GetName());
 				}
 			}
 			else
 			{
-				UE_LOG(LogPoseSearch, Warning, TEXT("FSearchContext::GetWorldBoneTransformAtTime - Schema '%s' Skeleton is not properly set"), *Schema->GetName());
+				UE_LOGF(LogPoseSearch, Warning, "FSearchContext::GetWorldBoneTransformAtTime - Schema '%ls' Skeleton is not properly set", *Schema->GetName());
 			}
 		}
 
diff --git a/Engine/Plugins/Animation/PoseSearch/Source/Runtime/Private/PoseSearchDatabase.cpp b/Engine/Plugins/Animation/PoseSearch/Source/Runtime/Private/PoseSearchDatabase.cpp
--- a/Engine/Plugins/Animation/PoseSearch/Source/Runtime/Private/PoseSearchDatabase.cpp
+++ b/Engine/Plugins/Animation/PoseSearch/Source/Runtime/Private/PoseSearchDatabase.cpp
@@ -369,7 +369,7 @@ FFloatInterval FPoseSearchDatabaseAnimationAssetBase::GetEffectiveSamplingRange(
 
 	if (Range.Min > Range.Max)
 	{
-		UE_LOG(LogPoseSearch, Warning, TEXT("Sampling range minimum (%f) is greated than max (%f). Setting min to be equal to max."), Range.Min, Range.Max)
+		UE_LOGF(LogPoseSearch, Warning, "Sampling range minimum (%f) is greated than max (%f). Setting min to be equal to max.", Range.Min, Range.Max)
 		
 		Range.Min = Range.Max;
 	}
@@ -1242,15 +1242,15 @@ void UPoseSearchDatabase::SynchronizeWithExternalDependencies(TConstArrayView<UA
 			UPackage* SequenceBasePackage = SequenceBase->GetPackage();
 			if (!SequenceBasePackage)
 			{
-				UE_LOG(LogPoseSearch, Warning, TEXT("AnimSequenceBase '%s' cannot be synchronized with PoseSearchDatabase '%s' because it's not in a package"), *SequenceBase->GetName(), *GetName());
+				UE_LOGF(LogPoseSearch, Warning, "AnimSequenceBase '%ls' cannot be synchronized with PoseSearchDatabase '%ls' because it's not in a package", *SequenceBase->GetName(), *GetName());
 			}
 			else if (SequenceBasePackage->HasAnyFlags(RF_Transient | RF_DuplicateTransient | RF_MirroredGarbage))
 			{
-				UE_LOG(LogPoseSearch, Warning, TEXT("AnimSequenceBase '%s' cannot be synchronized with PoseSearchDatabase '%s' because it's in a package transient or garbaged package"), *SequenceBase->GetName(), *GetName());
+				UE_LOGF(LogPoseSearch, Warning, "AnimSequenceBase '%ls' cannot be synchronized with PoseSearchDatabase '%ls' because it's in a package transient or garbaged package", *SequenceBase->GetName(), *GetName());
 			}
 			else if (SequenceBase->HasAnyFlags(RF_Transient | RF_DuplicateTransient | RF_MirroredGarbage))
 			{
-				UE_LOG(LogPoseSearch, Warning, TEXT("AnimSequenceBase '%s' cannot be synchronized with PoseSearchDatabase '%s' because it's transient or garbaged"), *SequenceBase->GetName(), *GetName());
+				UE_LOGF(LogPoseSearch, Warning, "AnimSequenceBase '%ls' cannot be synchronized with PoseSearchDatabase '%ls' because it's transient or garbaged", *SequenceBase->GetName(), *GetName());
 			}
 			else
 			{
@@ -1352,7 +1352,7 @@ void UPoseSearchDatabase::TestSynchronizeWithExternalDependencies()
 
 	if (DatabaseAnimationAssetsCopy != DatabaseAnimationAssets)
 	{
-		UE_LOG(LogPoseSearch, Error, TEXT("TestSynchronizeWithExternalDependencies failed"));
+		UE_LOGF(LogPoseSearch, Error, "TestSynchronizeWithExternalDependencies failed");
 		DatabaseAnimationAssets = DatabaseAnimationAssetsCopy;
 	}
 }
@@ -2013,7 +2013,7 @@ void UPoseSearchDatabase::SearchPCAKDTree(UE::PoseSearch::FSearchContext& Search
 			const double DeltaTimeRegular = EndTimeRegular - StartTimeRegular;
 			const double DeltaTimeMaxHeap = EndTimeMaxHeap - StartTimeMaxHeap;
 			const int32 Winner = DeltaTimeRegular == DeltaTimeMaxHeap ? 0 : DeltaTimeRegular < DeltaTimeMaxHeap ? -1 : 1;
-			UE_LOG(LogPoseSearch, Log, TEXT("Profiling: Regular(%f), MaxHeap(%f), Winner(%d)"), DeltaTimeRegular, DeltaTimeMaxHeap, Winner);
+			UE_LOGF(LogPoseSearch, Log, "Profiling: Regular(%f), MaxHeap(%f), Winner(%d)", DeltaTimeRegular, DeltaTimeMaxHeap, Winner);
 		}
 
 		// SortedResultsIndexes contains the sorted indexes by Distance of the Results. This is because FKNNMaxHeapResultSet doesn't return a sorted array like FKNNResultSet does
@@ -2038,7 +2038,7 @@ void UPoseSearchDatabase::SearchPCAKDTree(UE::PoseSearch::FSearchContext& Search
 				{
 					const float ExpectdDistanceSqr = ResultDistanceSqr[ResultIndex];
 					const float ActualDistanceSqr = Results[SortedResultsIndexes[ResultIndex]].Distance;
-					UE_LOG(LogPoseSearch, Warning, TEXT("Inconsistent Result at index %d (%d-%d) (%f-%f)"), ResultIndex, ExpectdIndex, ActualIndex, ExpectdDistanceSqr, ActualDistanceSqr);
+					UE_LOGF(LogPoseSearch, Warning, "Inconsistent Result at index %d (%d-%d) (%f-%f)", ResultIndex, ExpectdIndex, ActualIndex, ExpectdDistanceSqr, ActualDistanceSqr);
 				}
 			}
 		}
@@ -2068,16 +2068,16 @@ void UPoseSearchDatabase::SearchPCAKDTree(UE::PoseSearch::FSearchContext& Search
 				{
 					if (!FMath::IsNearlyEqual(PCAValueIndexCost[ResultIndex].Value, Results[ResultIndex].Distance, UE_KINDA_SMALL_NUMBER))
 					{
-						UE_LOG(LogPoseSearch, Error, TEXT("UPoseSearchDatabase::SearchPCAKDTree - KDTree search order is inconsistent with exaustive search in PCA space"));
+						UE_LOGF(LogPoseSearch, Error, "UPoseSearchDatabase::SearchPCAKDTree - KDTree search order is inconsistent with exaustive search in PCA space");
 					}
 					else
 					{
-						UE_LOG(LogPoseSearch, Log, TEXT("UPoseSearchDatabase::SearchPCAKDTree - found two points at the same distance from the query in different order between KDTree and exaustive search"));
+						UE_LOGF(LogPoseSearch, Log, "UPoseSearchDatabase::SearchPCAKDTree - found two points at the same distance from the query in different order between KDTree and exaustive search");
 					}
 				}
 				else if (!FMath::IsNearlyEqual(PCAValueIndexCost[ResultIndex].Value, Results[ResultIndex].Distance, UE_KINDA_SMALL_NUMBER))
 				{
-					UE_LOG(LogPoseSearch, Error, TEXT("UPoseSearchDatabase::SearchPCAKDTree - KDTree search cost is inconsistent with exaustive search in PCA space"));
+					UE_LOGF(LogPoseSearch, Error, "UPoseSearchDatabase::SearchPCAKDTree - KDTree search cost is inconsistent with exaustive search in PCA space");
 				}
 			}
 		}
@@ -2597,7 +2597,7 @@ void UPoseSearchDatabase::SearchEvent(UE::PoseSearch::FSearchContext& SearchCont
 	const FSearchIndex& SearchIndex = GetSearchIndex();
 	if (SearchIndex.IsValuesEmpty())
 	{
-		UE_LOG(LogPoseSearch, Error, TEXT("UPoseSearchDatabase::SearchEvent unimplemented for reconstructed poses"));
+		UE_LOGF(LogPoseSearch, Error, "UPoseSearchDatabase::SearchEvent unimplemented for reconstructed poses");
 	}
 	else
 	{
diff --git a/Engine/Plugins/Animation/PoseSearch/Source/Runtime/Private/PoseSearchDerivedDataKey.cpp b/Engine/Plugins/Animation/PoseSearch/Source/Runtime/Private/PoseSearchDerivedDataKey.cpp
--- a/Engine/Plugins/Animation/PoseSearch/Source/Runtime/Private/PoseSearchDerivedDataKey.cpp
+++ b/Engine/Plugins/Animation/PoseSearch/Source/Runtime/Private/PoseSearchDerivedDataKey.cpp
@@ -94,7 +94,7 @@ FKeyBuilder::FKeyBuilder(const UObject* Object, bool bUseDataVer, bool bPerformC
 #if ENABLE_KEY_PROFILING
 	const double EndTime(FPlatformTime::Seconds());
 	const double DeltaTimeInMilliseconds = (EndTime - StartTime) * 1000.0;
-	UE_LOG(LogPoseSearch, Log, TEXT("Profiling: FKeyBuilder::FKeyBuilder %f ms"), DeltaTimeInMilliseconds);
+	UE_LOGF(LogPoseSearch, Log, "Profiling: FKeyBuilder::FKeyBuilder %f ms", DeltaTimeInMilliseconds);
 #endif // ENABLE_KEY_PROFILING
 
 }
@@ -115,7 +115,7 @@ bool FKeyBuilder::ShouldSkipProperty(const FProperty* InProperty) const
 	if (FArchiveUObject::ShouldSkipProperty(InProperty))
 	{
 		#if UE_POSE_SEARCH_DERIVED_DATA_LOGGING
-		UE_LOG(LogPoseSearch, Log, TEXT("  x %s (ShouldSkipProperty)"), *InProperty->GetFullName());
+		UE_LOGF(LogPoseSearch, Log, "  x %ls (ShouldSkipProperty)", *InProperty->GetFullName());
 		#endif
 		return true;
 	}
@@ -123,7 +123,7 @@ bool FKeyBuilder::ShouldSkipProperty(const FProperty* InProperty) const
 	if (InProperty->HasAllPropertyFlags(CPF_Transient))
 	{
 		#if UE_POSE_SEARCH_DERIVED_DATA_LOGGING
-		UE_LOG(LogPoseSearch, Log, TEXT("  x %s (Transient)"), *InProperty->GetFullName());
+		UE_LOGF(LogPoseSearch, Log, "  x %ls (Transient)", *InProperty->GetFullName());
 		#endif
 		return true;
 	}
@@ -131,7 +131,7 @@ bool FKeyBuilder::ShouldSkipProperty(const FProperty* InProperty) const
 	if (InProperty->HasMetaData(ExcludeFromHashName))
 	{
 		#if UE_POSE_SEARCH_DERIVED_DATA_LOGGING
-		UE_LOG(LogPoseSearch, Log, TEXT("  x %s (ExcludeFromHash)"), *InProperty->GetFullName());
+		UE_LOGF(LogPoseSearch, Log, "  x %ls (ExcludeFromHash)", *InProperty->GetFullName());
 		#endif
 		return true;
 	}
@@ -139,7 +139,7 @@ bool FKeyBuilder::ShouldSkipProperty(const FProperty* InProperty) const
 	if (InProperty->HasMetaData(IgnoreForMemberInitializationTestName))
 	{
 		#if UE_POSE_SEARCH_DERIVED_DATA_LOGGING
-		UE_LOG(LogPoseSearch, Log, TEXT("  x %s (IgnoreForMemberInitializationTest)"), *InProperty->GetFullName());
+		UE_LOGF(LogPoseSearch, Log, "  x %ls (IgnoreForMemberInitializationTest)", *InProperty->GetFullName());
 		#endif
 		return true;
 	}
@@ -147,7 +147,7 @@ bool FKeyBuilder::ShouldSkipProperty(const FProperty* InProperty) const
 	check(!InProperty->HasMetaData(NeverInHashName));
 
 	#if UE_POSE_SEARCH_DERIVED_DATA_LOGGING
-	UE_LOG(LogPoseSearch, Log, TEXT("  - %s"), *InProperty->GetFullName());
+	UE_LOGF(LogPoseSearch, Log, "  - %ls", *InProperty->GetFullName());
 	#endif
 
 	return false;
@@ -158,7 +158,7 @@ void FKeyBuilder::Serialize(void* Data, int64 Length)
 	#if UE_POSE_SEARCH_DERIVED_DATA_LOGGING_VERBOSE
 	const uint8* HasherData = reinterpret_cast<uint8*>(Data);
 	FString RawBytesString = BytesToString(HasherData, Length);
-	UE_LOG(LogPoseSearch, Log, TEXT("  > %s"), *RawBytesString);
+	UE_LOGF(LogPoseSearch, Log, "  > %ls", *RawBytesString);
 	#endif
 
 	Hasher.Update(Data, Length);
@@ -239,7 +239,7 @@ FArchive& FKeyBuilder::TryAddDependency(UObject* Object)
 	if (bAlreadyProcessed)
 	{
 #if UE_POSE_SEARCH_DERIVED_DATA_LOGGING
-		UE_LOG(LogPoseSearch, Log, TEXT("AlreadyProcessed '%s' (%s)"), *Object->GetName(), *Object->GetClass()->GetName());
+		UE_LOGF(LogPoseSearch, Log, "AlreadyProcessed '%ls' (%ls)", *Object->GetName(), *Object->GetClass()->GetName());
 #endif
 		return *this;
 	}
@@ -260,35 +260,35 @@ FArchive& FKeyBuilder::operator<<(class UObject*& Object)
 	{
 		if (Object->HasAnyFlags(RF_Transient))
 		{
-			UE_LOG(LogPoseSearch, Log, TEXT("Transient '%s' (%s)"), *Object->GetName(), *Object->GetClass()->GetName());
+			UE_LOGF(LogPoseSearch, Log, "Transient '%ls' (%ls)", *Object->GetName(), *Object->GetClass()->GetName());
 		}
 		if (Object->HasAnyFlags(RF_MirroredGarbage))
 		{
-			UE_LOG(LogPoseSearch, Log, TEXT("MirroredGarbage '%s' (%s)"), *Object->GetName(), *Object->GetClass()->GetName());
+			UE_LOGF(LogPoseSearch, Log, "MirroredGarbage '%ls' (%ls)", *Object->GetName(), *Object->GetClass()->GetName());
 		}
 		if (Object->HasAnyFlags(RF_NewerVersionExists))
 		{
-			UE_LOG(LogPoseSearch, Log, TEXT("NewerVersionExists '%s' (%s)"), *Object->GetName(), *Object->GetClass()->GetName());
+			UE_LOGF(LogPoseSearch, Log, "NewerVersionExists '%ls' (%ls)", *Object->GetName(), *Object->GetClass()->GetName());
 		}
 		if (Object->HasAnyFlags(RF_BeginDestroyed))
 		{
-			UE_LOG(LogPoseSearch, Log, TEXT("BeginDestroyed '%s' (%s)"), *Object->GetName(), *Object->GetClass()->GetName());
+			UE_LOGF(LogPoseSearch, Log, "BeginDestroyed '%ls' (%ls)", *Object->GetName(), *Object->GetClass()->GetName());
 		}
 		if (Object->HasAnyFlags(RF_FinishDestroyed))
 		{
-			UE_LOG(LogPoseSearch, Log, TEXT("FinishDestroyed '%s' (%s)"), *Object->GetName(), *Object->GetClass()->GetName());
+			UE_LOGF(LogPoseSearch, Log, "FinishDestroyed '%ls' (%ls)", *Object->GetName(), *Object->GetClass()->GetName());
 		}
 		if (Object->HasAnyFlags(RF_DuplicateTransient))
 		{
-			UE_LOG(LogPoseSearch, Log, TEXT("DuplicateTransient '%s' (%s)"), *Object->GetName(), *Object->GetClass()->GetName());
+			UE_LOGF(LogPoseSearch, Log, "DuplicateTransient '%ls' (%ls)", *Object->GetName(), *Object->GetClass()->GetName());
 		}
 		if (Object->HasAnyFlags(RF_NonPIEDuplicateTransient))
 		{
-			UE_LOG(LogPoseSearch, Log, TEXT("DuplicateTransient '%s' (%s)"), *Object->GetName(), *Object->GetClass()->GetName());
+			UE_LOGF(LogPoseSearch, Log, "DuplicateTransient '%ls' (%ls)", *Object->GetName(), *Object->GetClass()->GetName());
 		}
 		if (IsExcludedType(Object))
 		{
-			UE_LOG(LogPoseSearch, Log, TEXT("Excluded '%s' (%s)"), *Object->GetName(), *Object->GetClass()->GetName());
+			UE_LOGF(LogPoseSearch, Log, "Excluded '%ls' (%ls)", *Object->GetName(), *Object->GetClass()->GetName());
 		}
 	}
 #endif // UE_POSE_SEARCH_DERIVED_DATA_LOGGING
@@ -324,7 +324,7 @@ void FKeyBuilder::SerializeObjectInternal(UObject* Object)
 		{
 			// for specific types we only add their names to the hash
 #if UE_POSE_SEARCH_DERIVED_DATA_LOGGING
-			UE_LOG(LogPoseSearch, Log, TEXT("AddingNameOnly '%s' (%s)"), *Object->GetName(), *Object->GetClass()->GetName());
+			UE_LOGF(LogPoseSearch, Log, "AddingNameOnly '%ls' (%ls)", *Object->GetName(), *Object->GetClass()->GetName());
 #endif
 			// non FStringBuilderBase version
 			//FString FullName = Object->GetFullName();
@@ -337,12 +337,12 @@ void FKeyBuilder::SerializeObjectInternal(UObject* Object)
 		else
 		{
 #if UE_POSE_SEARCH_DERIVED_DATA_LOGGING
-			UE_LOG(LogPoseSearch, Log, TEXT("Begin '%s' (%s)"), *Object->GetName(), *Object->GetClass()->GetName());
+			UE_LOGF(LogPoseSearch, Log, "Begin '%ls' (%ls)", *Object->GetName(), *Object->GetClass()->GetName());
 #endif
 			Object->Serialize(*this);
 
 #if UE_POSE_SEARCH_DERIVED_DATA_LOGGING
-			UE_LOG(LogPoseSearch, Log, TEXT("End '%s' (%s)"), *Object->GetName(), *Object->GetClass()->GetName());
+			UE_LOGF(LogPoseSearch, Log, "End '%ls' (%ls)", *Object->GetName(), *Object->GetClass()->GetName());
 #endif
 		}
 
@@ -472,7 +472,7 @@ FIoHash FKeyBuilder::Finalize()
 #if ENABLE_KEY_PROFILING
 	const double EndTime(FPlatformTime::Seconds());
 	const double DeltaTimeInMilliseconds = (EndTime - StartTime) * 1000.0;
-	UE_LOG(LogPoseSearch, Log, TEXT("Profiling: FKeyBuilder::Finalize %f ms"), DeltaTimeInMilliseconds);
+	UE_LOGF(LogPoseSearch, Log, "Profiling: FKeyBuilder::Finalize %f ms", DeltaTimeInMilliseconds);
 #endif // ENABLE_KEY_PROFILING
 
 	return IoHash;
diff --git a/Engine/Plugins/Animation/PoseSearch/Source/Runtime/Private/PoseSearchFeatureChannel_Curve.cpp b/Engine/Plugins/Animation/PoseSearch/Source/Runtime/Private/PoseSearchFeatureChannel_Curve.cpp
--- a/Engine/Plugins/Animation/PoseSearch/Source/Runtime/Private/PoseSearchFeatureChannel_Curve.cpp
+++ b/Engine/Plugins/Animation/PoseSearch/Source/Runtime/Private/PoseSearchFeatureChannel_Curve.cpp
@@ -45,7 +45,7 @@ void UPoseSearchFeatureChannel_Curve::BuildQuery(UE::PoseSearch::FSearchContext&
 		else
 		{
 			// @todo: support non UAnimInstance anim contexts for AnimNext
-			UE_LOG(LogPoseSearch, Warning, TEXT("UPoseSearchFeatureChannel_Curve::BuildQuery - unsupported null UAnimInstance: WIP support for AnimNext!"));
+			UE_LOGF(LogPoseSearch, Warning, "UPoseSearchFeatureChannel_Curve::BuildQuery - unsupported null UAnimInstance: WIP support for AnimNext!");
 		}
 		return;
 	}
diff --git a/Engine/Plugins/Animation/PoseSearch/Source/Runtime/Private/PoseSearchFeatureChannel_Heading.cpp b/Engine/Plugins/Animation/PoseSearch/Source/Runtime/Private/PoseSearchFeatureChannel_Heading.cpp
--- a/Engine/Plugins/Animation/PoseSearch/Source/Runtime/Private/PoseSearchFeatureChannel_Heading.cpp
+++ b/Engine/Plugins/Animation/PoseSearch/Source/Runtime/Private/PoseSearchFeatureChannel_Heading.cpp
@@ -28,7 +28,7 @@ void UPoseSearchFeatureChannel_Heading::FindOrAddToSchema(UPoseSearchSchema* Sch
 	const int8 SchemaBoneIdx = Schema->FindSchemaBoneIdx(BoneName, Role, bDefaultWithRootBone);
 	if (SchemaBoneIdx == InvalidSchemaBoneIdx)
 	{
-		UE_LOG(LogPoseSearch, Error, TEXT("UPoseSearchFeatureChannel_Heading::FindOrAddToSchema: couldn't find SchemaBoneIdx for Bone '%s' with Role '%s' in UPoseSearchSchema '%s'"), *BoneName.ToString(), *Role.ToString(), *GetNameSafe(Schema));
+		UE_LOGF(LogPoseSearch, Error, "UPoseSearchFeatureChannel_Heading::FindOrAddToSchema: couldn't find SchemaBoneIdx for Bone '%ls' with Role '%ls' in UPoseSearchSchema '%ls'", *BoneName.ToString(), *Role.ToString(), *GetNameSafe(Schema));
 		return;
 	}
 
@@ -131,7 +131,7 @@ void UPoseSearchFeatureChannel_Heading::BuildQuery(UE::PoseSearch::FSearchContex
 		else
 		{
 			// @todo: support non UAnimInstance anim contexts for AnimNext
-			UE_LOG(LogPoseSearch, Warning, TEXT("UPoseSearchFeatureChannel_Heading::BuildQuery - unsupported null UAnimInstance: WIP support for AnimNext!"));
+			UE_LOGF(LogPoseSearch, Warning, "UPoseSearchFeatureChannel_Heading::BuildQuery - unsupported null UAnimInstance: WIP support for AnimNext!");
 		}
 		return;
 	}
diff --git a/Engine/Plugins/Animation/PoseSearch/Source/Runtime/Private/PoseSearchFeatureChannel_Padding.cpp b/Engine/Plugins/Animation/PoseSearch/Source/Runtime/Private/PoseSearchFeatureChannel_Padding.cpp
--- a/Engine/Plugins/Animation/PoseSearch/Source/Runtime/Private/PoseSearchFeatureChannel_Padding.cpp
+++ b/Engine/Plugins/Animation/PoseSearch/Source/Runtime/Private/PoseSearchFeatureChannel_Padding.cpp
@@ -38,7 +38,7 @@ void UPoseSearchFeatureChannel_Padding::FillWeights(TArrayView<float> Weights) c
 
 bool UPoseSearchFeatureChannel_Padding::IndexAsset(UE::PoseSearch::FAssetIndexer& Indexer) const
 {
-	UE_LOG(LogPoseSearch, Log, TEXT("UPoseSearchFeatureChannel_Padding::IndexAsset: padding data with '%d' additional floats per pose in UPoseSearchSchema '%s'"), PaddingSize, *GetNameSafe(GetSchema()));
+	UE_LOGF(LogPoseSearch, Log, "UPoseSearchFeatureChannel_Padding::IndexAsset: padding data with '%d' additional floats per pose in UPoseSearchSchema '%ls'", PaddingSize, *GetNameSafe(GetSchema()));
 	return true;
 }
 
diff --git a/Engine/Plugins/Animation/PoseSearch/Source/Runtime/Private/PoseSearchFeatureChannel_PermutationTime.cpp b/Engine/Plugins/Animation/PoseSearch/Source/Runtime/Private/PoseSearchFeatureChannel_PermutationTime.cpp
--- a/Engine/Plugins/Animation/PoseSearch/Source/Runtime/Private/PoseSearchFeatureChannel_PermutationTime.cpp
+++ b/Engine/Plugins/Animation/PoseSearch/Source/Runtime/Private/PoseSearchFeatureChannel_PermutationTime.cpp
@@ -17,7 +17,7 @@ void UPoseSearchFeatureChannel_PermutationTime::FindOrAddToSchema(UPoseSearchSch
 #endif // WITH_EDITORONLY_DATA
 		Schema->AddTemporaryChannel(PermutationTime);
 
-		UE_LOG(LogPoseSearch, Warning, TEXT("required UPoseSearchFeatureChannel_PermutationTime has been added by the Schema '%s' with Weight of zero. Explicitly add one of this channel to the Schema to be able to tune its weight"), *Schema->GetName());
+		UE_LOGF(LogPoseSearch, Warning, "required UPoseSearchFeatureChannel_PermutationTime has been added by the Schema '%ls' with Weight of zero. Explicitly add one of this channel to the Schema to be able to tune its weight", *Schema->GetName());
 	}
 }
 
diff --git a/Engine/Plugins/Animation/PoseSearch/Source/Runtime/Private/PoseSearchFeatureChannel_Phase.cpp b/Engine/Plugins/Animation/PoseSearch/Source/Runtime/Private/PoseSearchFeatureChannel_Phase.cpp
--- a/Engine/Plugins/Animation/PoseSearch/Source/Runtime/Private/PoseSearchFeatureChannel_Phase.cpp
+++ b/Engine/Plugins/Animation/PoseSearch/Source/Runtime/Private/PoseSearchFeatureChannel_Phase.cpp
@@ -403,7 +403,7 @@ bool UPoseSearchFeatureChannel_Phase::IndexAsset(UE::PoseSearch::FAssetIndexer&
 	if (Indexer.GetNumIndexedPoses() <= 1)
 	{
 		// We must have at least 2 poses in order to index phase properly.
-		UE_LOG(LogPoseSearch, Error, TEXT("UPoseSearchFeatureChannel_Phase::IndexAsset - Attempt to index animation asset with single pose is not supported. Asset name: %s"), *Indexer.GetAssetName());
+		UE_LOGF(LogPoseSearch, Error, "UPoseSearchFeatureChannel_Phase::IndexAsset - Attempt to index animation asset with single pose is not supported. Asset name: %ls", *Indexer.GetAssetName());
 		return false;
 	}
 
diff --git a/Engine/Plugins/Animation/PoseSearch/Source/Runtime/Private/PoseSearchFeatureChannel_Pose.cpp b/Engine/Plugins/Animation/PoseSearch/Source/Runtime/Private/PoseSearchFeatureChannel_Pose.cpp
--- a/Engine/Plugins/Animation/PoseSearch/Source/Runtime/Private/PoseSearchFeatureChannel_Pose.cpp
+++ b/Engine/Plugins/Animation/PoseSearch/Source/Runtime/Private/PoseSearchFeatureChannel_Pose.cpp
@@ -122,7 +122,7 @@ bool UPoseSearchFeatureChannel_Pose::Finalize(UPoseSearchSchema* Schema)
 
 			if (PermutationTimeType != EPermutationTimeType::UseSampleTime)
 			{
-				UE_LOG(LogPoseSearch, Error, TEXT("UPoseSearchFeatureChannel_Pose::Finalize - unsupported PermutationTimeType %d for EPoseSearchBoneFlags::Phase in Schema '%s'"), (int32)PermutationTimeType, *Schema->GetName());
+				UE_LOGF(LogPoseSearch, Error, "UPoseSearchFeatureChannel_Pose::Finalize - unsupported PermutationTimeType %d for EPoseSearchBoneFlags::Phase in Schema '%ls'", (int32)PermutationTimeType, *Schema->GetName());
 				return false;
 			}
 			SubChannels.Add(Phase);
diff --git a/Engine/Plugins/Animation/PoseSearch/Source/Runtime/Private/PoseSearchFeatureChannel_Position.cpp b/Engine/Plugins/Animation/PoseSearch/Source/Runtime/Private/PoseSearchFeatureChannel_Position.cpp
--- a/Engine/Plugins/Animation/PoseSearch/Source/Runtime/Private/PoseSearchFeatureChannel_Position.cpp
+++ b/Engine/Plugins/Animation/PoseSearch/Source/Runtime/Private/PoseSearchFeatureChannel_Position.cpp
@@ -27,7 +27,7 @@ void UPoseSearchFeatureChannel_Position::FindOrAddToSchema(UPoseSearchSchema* Sc
 	const int8 SchemaBoneIdx = Schema->FindSchemaBoneIdx(BoneName, Role, bDefaultWithRootBone);
 	if (SchemaBoneIdx == InvalidSchemaBoneIdx)
 	{
-		UE_LOG(LogPoseSearch, Error, TEXT("UPoseSearchFeatureChannel_Position::FindOrAddToSchema: couldn't find SchemaBoneIdx for Bone '%s' with Role '%s' in UPoseSearchSchema '%s'"), *BoneName.ToString(), *Role.ToString(), *GetNameSafe(Schema));
+		UE_LOGF(LogPoseSearch, Error, "UPoseSearchFeatureChannel_Position::FindOrAddToSchema: couldn't find SchemaBoneIdx for Bone '%ls' with Role '%ls' in UPoseSearchSchema '%ls'", *BoneName.ToString(), *Role.ToString(), *GetNameSafe(Schema));
 		return;
 	}
 
@@ -110,7 +110,7 @@ void UPoseSearchFeatureChannel_Position::BuildQuery(UE::PoseSearch::FSearchConte
 		else
 		{
 			// @todo: support non UAnimInstance anim contexts for AnimNext
-			UE_LOG(LogPoseSearch, Warning, TEXT("UPoseSearchFeatureChannel_Position::BuildQuery - unsupported null UAnimInstance: WIP support for AnimNext!"));
+			UE_LOGF(LogPoseSearch, Warning, "UPoseSearchFeatureChannel_Position::BuildQuery - unsupported null UAnimInstance: WIP support for AnimNext!");
 		}
 		return;
 	}
diff --git a/Engine/Plugins/Animation/PoseSearch/Source/Runtime/Private/PoseSearchFeatureChannel_TimeToEvent.cpp b/Engine/Plugins/Animation/PoseSearch/Source/Runtime/Private/PoseSearchFeatureChannel_TimeToEvent.cpp
--- a/Engine/Plugins/Animation/PoseSearch/Source/Runtime/Private/PoseSearchFeatureChannel_TimeToEvent.cpp
+++ b/Engine/Plugins/Animation/PoseSearch/Source/Runtime/Private/PoseSearchFeatureChannel_TimeToEvent.cpp
@@ -28,13 +28,13 @@ void UPoseSearchFeatureChannel_TimeToEvent::BuildQuery(UE::PoseSearch::FSearchCo
 
 	if (!bUseBlueprintQueryOverride)
 	{
-		UE_LOG(LogPoseSearch, Warning, TEXT("UPoseSearchFeatureChannel_TimeToEvent::BuildQuery - UPoseSearchFeatureChannel_TimeToEvent is designed to work only as BluePrint overridden class. The query TimeToEvent value will be defaulted to 0"));
+		UE_LOGF(LogPoseSearch, Warning, "UPoseSearchFeatureChannel_TimeToEvent::BuildQuery - UPoseSearchFeatureChannel_TimeToEvent is designed to work only as BluePrint overridden class. The query TimeToEvent value will be defaulted to 0");
 		return;
 	}
 
 	if (SearchContext.GetContexts().IsEmpty())
 	{
-		UE_LOG(LogPoseSearch, Error, TEXT("UPoseSearchFeatureChannel_Position::BuildQuery - no provided anim contexts!"));
+		UE_LOGF(LogPoseSearch, Error, "UPoseSearchFeatureChannel_Position::BuildQuery - no provided anim contexts!");
 		return;
 	}
 	
@@ -45,7 +45,7 @@ void UPoseSearchFeatureChannel_TimeToEvent::BuildQuery(UE::PoseSearch::FSearchCo
 	}
 	else
 	{
-		UE_LOG(LogPoseSearch, Warning, TEXT("UPoseSearchFeatureChannel_TimeToEvent::BuildQuery - UPoseSearchFeatureChannel_TimeToEvent is designed to work only as BluePrint overridden class. The query TimeToEvent value will be defaulted to 0"));
+		UE_LOGF(LogPoseSearch, Warning, "UPoseSearchFeatureChannel_TimeToEvent::BuildQuery - UPoseSearchFeatureChannel_TimeToEvent is designed to work only as BluePrint overridden class. The query TimeToEvent value will be defaulted to 0");
 	}
 }
 
diff --git a/Engine/Plugins/Animation/PoseSearch/Source/Runtime/Private/PoseSearchFeatureChannel_Velocity.cpp b/Engine/Plugins/Animation/PoseSearch/Source/Runtime/Private/PoseSearchFeatureChannel_Velocity.cpp
--- a/Engine/Plugins/Animation/PoseSearch/Source/Runtime/Private/PoseSearchFeatureChannel_Velocity.cpp
+++ b/Engine/Plugins/Animation/PoseSearch/Source/Runtime/Private/PoseSearchFeatureChannel_Velocity.cpp
@@ -63,7 +63,7 @@ void UPoseSearchFeatureChannel_Velocity::BuildQuery(UE::PoseSearch::FSearchConte
 		else
 		{
 			// @todo: support non UAnimInstance anim contexts for AnimNext
-			UE_LOG(LogPoseSearch, Warning, TEXT("UPoseSearchFeatureChannel_Velocity::BuildQuery - unsupported null UAnimInstance: WIP support for AnimNext!"));
+			UE_LOGF(LogPoseSearch, Warning, "UPoseSearchFeatureChannel_Velocity::BuildQuery - unsupported null UAnimInstance: WIP support for AnimNext!");
 		}
 		return;
 	}
diff --git a/Engine/Plugins/Animation/PoseSearch/Source/Runtime/Private/PoseSearchIndex.cpp b/Engine/Plugins/Animation/PoseSearch/Source/Runtime/Private/PoseSearchIndex.cpp
--- a/Engine/Plugins/Animation/PoseSearch/Source/Runtime/Private/PoseSearchIndex.cpp
+++ b/Engine/Plugins/Animation/PoseSearch/Source/Runtime/Private/PoseSearchIndex.cpp
@@ -371,14 +371,14 @@ void FEventData::ValidateEventData() const
 			return A.Key.ToString() < B.Key.ToString();
 		}))
 	{
-		UE_LOG(LogPoseSearch, Error, TEXT("FEventData::ValidateEventData FGameplayTag are not properly sorted!"));
+		UE_LOGF(LogPoseSearch, Error, "FEventData::ValidateEventData FGameplayTag are not properly sorted!");
 	}
 
 	for (const FEventData::FTagToPoseIndexes& TagToPoseIndexes : Data)
 	{
 		if (!Algo::IsSorted(TagToPoseIndexes.Value))
 		{
-			UE_LOG(LogPoseSearch, Error, TEXT("FEventData::ValidateEventData FPoseIndexes are not properly sorted!"));
+			UE_LOGF(LogPoseSearch, Error, "FEventData::ValidateEventData FPoseIndexes are not properly sorted!");
 			break;
 		}
 	}
diff --git a/Engine/Plugins/Animation/PoseSearch/Source/Runtime/Private/PoseSearchInteractionAsset.cpp b/Engine/Plugins/Animation/PoseSearch/Source/Runtime/Private/PoseSearchInteractionAsset.cpp
--- a/Engine/Plugins/Animation/PoseSearch/Source/Runtime/Private/PoseSearchInteractionAsset.cpp
+++ b/Engine/Plugins/Animation/PoseSearch/Source/Runtime/Private/PoseSearchInteractionAsset.cpp
@@ -199,7 +199,7 @@ bool UPoseSearchInteractionAsset::IsLooping() const
 			}
 			else
 			{
-				UE_LOG(LogPoseSearch, Warning, TEXT("UPoseSearchInteractionAsset::IsLooping non fully supported UAnimationAsset derived type '%s' used as item in '%s'"), *AnimationAsset->GetClass()->GetName(), *GetName());
+				UE_LOGF(LogPoseSearch, Warning, "UPoseSearchInteractionAsset::IsLooping non fully supported UAnimationAsset derived type '%ls' used as item in '%ls'", *AnimationAsset->GetClass()->GetName(), *GetName());
 				return false;
 			}
 
@@ -241,7 +241,7 @@ bool UPoseSearchInteractionAsset::HasRootMotion() const
 			}
 			else
 			{
-				UE_LOG(LogPoseSearch, Warning, TEXT("UPoseSearchInteractionAsset::HasRootMotion non fully supported UAnimationAsset derived type '%s' used as item in '%s'"), *AnimationAsset->GetClass()->GetName(), *GetName());
+				UE_LOGF(LogPoseSearch, Warning, "UPoseSearchInteractionAsset::HasRootMotion non fully supported UAnimationAsset derived type '%ls' used as item in '%ls'", *AnimationAsset->GetClass()->GetName(), *GetName());
 			}
 			bHasAtLeastOneValidItem = true;
 		}
@@ -409,7 +409,7 @@ void UPoseSearchInteractionAsset::CalculateWarpTransforms(float Time, float Time
 				{
 					const FVector Pos = RootBoneTransform.GetLocation();
 					const FRotator Rot(RootBoneTransform.GetRotation());
-					UE_LOG(LogPoseSearch, Error, TEXT("UPoseSearchInteractionAsset::CalculateWarpTransforms unsupported non identity root bone in %s at time %f Pos(%f, %f, %f), Rot(%f, %f, %f)"), *Items[ItemIndex].Animation->GetName(), Time, Pos.X, Pos.Y, Pos.Z, Rot.Pitch, Rot.Yaw, Rot.Roll);
+					UE_LOGF(LogPoseSearch, Error, "UPoseSearchInteractionAsset::CalculateWarpTransforms unsupported non identity root bone in %ls at time %f Pos(%f, %f, %f), Rot(%f, %f, %f)", *Items[ItemIndex].Animation->GetName(), Time, Pos.X, Pos.Y, Pos.Z, Rot.Pitch, Rot.Yaw, Rot.Roll);
 				}
 			}
 #endif // ENABLE_ANIM_DEBUG && WITH_EDITOR
@@ -471,7 +471,7 @@ void UPoseSearchInteractionAsset::CalculateWarpTransforms(float Time, float Time
 					}
 					else
 					{
-						UE_LOG(LogPoseSearch, Error, TEXT("UPoseSearchInteractionAsset::CalculateWarpTransforms actors are too close to each other to be able to calculate warp transforms ['%s' at time %f]"), *GetName(), Time);
+						UE_LOGF(LogPoseSearch, Error, "UPoseSearchInteractionAsset::CalculateWarpTransforms actors are too close to each other to be able to calculate warp transforms ['%ls' at time %f]", *GetName(), Time);
 					}
 				}
 			}
diff --git a/Engine/Plugins/Animation/PoseSearch/Source/Runtime/Private/PoseSearchInteractionIsland.cpp b/Engine/Plugins/Animation/PoseSearch/Source/Runtime/Private/PoseSearchInteractionIsland.cpp
--- a/Engine/Plugins/Animation/PoseSearch/Source/Runtime/Private/PoseSearchInteractionIsland.cpp
+++ b/Engine/Plugins/Animation/PoseSearch/Source/Runtime/Private/PoseSearchInteractionIsland.cpp
@@ -36,7 +36,7 @@ static void ShowPrerequistes(FTickFunction& NestedTick, int32 Indent, TSet<FTick
 	{
 		if (ensure(Indent > 0))
 		{
-			UE_LOG(LogPoseSearch, Log, TEXT("%s ERROR! Cycle detected! Use 'a.PoseSearchInteraction.DiagnoseTickDependencies' CVar to gather more infos to diagnose the problem"), FCString::Spc((Indent - 1) * 2));
+			UE_LOGF(LogPoseSearch, Log, "%ls ERROR! Cycle detected! Use 'a.PoseSearchInteraction.DiagnoseTickDependencies' CVar to gather more infos to diagnose the problem", FCString::Spc((Indent - 1) * 2));
 		}
 	}
 	else
@@ -45,7 +45,7 @@ static void ShowPrerequistes(FTickFunction& NestedTick, int32 Indent, TSet<FTick
 		{
 			if (Prereq.PrerequisiteTickFunction)
 			{
-				UE_LOG(LogPoseSearch, Log, TEXT("%s prereq %s"), FCString::Spc(Indent * 2), *Prereq.PrerequisiteTickFunction->DiagnosticMessage());
+				UE_LOGF(LogPoseSearch, Log, "%ls prereq %ls", FCString::Spc(Indent * 2), *Prereq.PrerequisiteTickFunction->DiagnosticMessage());
 				ShowPrerequistes(*Prereq.PrerequisiteTickFunction, Indent + 1, VisitedTickFunctions);
 			}
 		}
@@ -57,7 +57,7 @@ static void LogTickFunction(FTickFunction& Tick, ENamedThreads::Type CurrentThre
 {
 	// scoping brackets to save some heap for the recursion
 	{
-		UE_LOG(LogPoseSearch, Log, TEXT("%stick %s [%1d, %1d] %6llu %2d %s"), FCString::Spc(Indent * 2), Tick.bHighPriority ? TEXT("*") : TEXT(" "), (int32)Tick.GetActualTickGroup(), (int32)Tick.GetActualEndTickGroup(), (uint64)GFrameCounter, (int32)CurrentThread, *Tick.DiagnosticMessage());
+		UE_LOGF(LogPoseSearch, Log, "%lstick %ls [%1d, %1d] %6llu %2d %ls", FCString::Spc(Indent * 2), Tick.bHighPriority ? TEXT("*") : TEXT(" "), (int32)Tick.GetActualTickGroup(), (int32)Tick.GetActualEndTickGroup(), (uint64)GFrameCounter, (int32)CurrentThread, *Tick.DiagnosticMessage());
 		if (bLogPrerequisites)
 		{
 			TSet<FTickFunction*> VisitedTickFunctions;
@@ -300,7 +300,7 @@ static void AddPrerequisite(FTickFunction& TickFunction, UObject* TargetObject,
 #if ENABLE_ANIM_DEBUG
 		if (TickFunction.bCanEverTick && !TickFunction.GetPrerequisites().Contains(FTickPrerequisite(TargetObject, TargetTickFunction)))
 		{
-			UE_LOG(LogPoseSearch, Error, TEXT("UE::PoseSearch::AddPrerequisite, Failed to add prerequisite from [%s] to [%s, %s]!"), *TickFunction.DiagnosticMessage(), *TargetObject->GetName(), *TargetTickFunction.DiagnosticMessage());
+			UE_LOGF(LogPoseSearch, Error, "UE::PoseSearch::AddPrerequisite, Failed to add prerequisite from [%ls] to [%ls, %ls]!", *TickFunction.DiagnosticMessage(), *TargetObject->GetName(), *TargetTickFunction.DiagnosticMessage());
 		}
 
 		if (GVarPoseSearchInteractionDiagnoseTickDependencies)
@@ -952,19 +952,19 @@ void FInteractionIsland::FPreTickFunction::ExecuteTick(float DeltaTime, enum ELe
 		{
 			// @todo: need to figure out why when creating a new island FPreTickFunction gets called twice (it's not a real issue rather than a performance hit)
 			//		  use GVarPoseSearchInteractionCacheIslands = false to debug the issue (it destroys the islands every frame)
-			UE_LOG(LogPoseSearch, Log, TEXT("FInteractionIsland::FPreTickFunction::ExecuteTick, called twice before UPoseSearchInteractionSubsystem::Tick!"));
+			UE_LOGF(LogPoseSearch, Log, "FInteractionIsland::FPreTickFunction::ExecuteTick, called twice before UPoseSearchInteractionSubsystem::Tick!");
 		}
 		else
 		{
 			if (Island->bPostTickFunctionExecuted)
 			{
-				UE_LOG(LogPoseSearch, Error, TEXT("FInteractionIsland::FPreTickFunction::ExecuteTick, FPostTickFunction::ExecuteTick alreay run?!"));
+				UE_LOGF(LogPoseSearch, Error, "FInteractionIsland::FPreTickFunction::ExecuteTick, FPostTickFunction::ExecuteTick alreay run?!");
 				Island->LogTickDependencies();
 			}
 
 			if (Island->bClosingTickFunctionExecuted)
 			{
-				UE_LOG(LogPoseSearch, Error, TEXT("FInteractionIsland::FPreTickFunction::ExecuteTick, FClosingTickFunction::ExecuteTick alreay run?!"));
+				UE_LOGF(LogPoseSearch, Error, "FInteractionIsland::FPreTickFunction::ExecuteTick, FClosingTickFunction::ExecuteTick alreay run?!");
 				Island->LogTickDependencies();
 			}
 
@@ -989,19 +989,19 @@ void FInteractionIsland::FPostTickFunction::ExecuteTick(float DeltaTime, enum EL
 		{
 			// @todo: need to figure out why when creating a new island FPostTickFunction gets called twice (it's not a real issue rather than a performance hit)
 			//		  use GVarPoseSearchInteractionCacheIslands = false to debug the issue (it destroys the islands every frame)
-			UE_LOG(LogPoseSearch, Log, TEXT("FInteractionIsland::FPostTickFunction::ExecuteTick, called twice before UPoseSearchInteractionSubsystem::Tick!"));
+			UE_LOGF(LogPoseSearch, Log, "FInteractionIsland::FPostTickFunction::ExecuteTick, called twice before UPoseSearchInteractionSubsystem::Tick!");
 		}
 		else
 		{
 			if (!Island->bPreTickFunctionExecuted)
 			{
-				UE_LOG(LogPoseSearch, Error, TEXT("FInteractionIsland::FPostTickFunction::ExecuteTick, FPreTickFunction::ExecuteTick didn't run!"));
+				UE_LOGF(LogPoseSearch, Error, "FInteractionIsland::FPostTickFunction::ExecuteTick, FPreTickFunction::ExecuteTick didn't run!");
 				Island->LogTickDependencies();
 			}
 
 			if (Island->bClosingTickFunctionExecuted)
 			{
-				UE_LOG(LogPoseSearch, Error, TEXT("FInteractionIsland::FPostTickFunction::ExecuteTick, FClosingTickFunction::ExecuteTick alreay run?!"));
+				UE_LOGF(LogPoseSearch, Error, "FInteractionIsland::FPostTickFunction::ExecuteTick, FClosingTickFunction::ExecuteTick alreay run?!");
 				Island->LogTickDependencies();
 			}
 
@@ -1027,19 +1027,19 @@ void FInteractionIsland::FClosingTickFunction::ExecuteTick(float DeltaTime, enum
 		{
 			// @todo: need to figure out why when creating a new island FClosingTickFunction gets called twice (it's not a real issue rather than a performance hit)
 			//		  use GVarPoseSearchInteractionCacheIslands = false to debug the issue (it destroys the islands every frame)
-			UE_LOG(LogPoseSearch, Log, TEXT("FInteractionIsland::FClosingTickFunction::ExecuteTick, called twice before UPoseSearchInteractionSubsystem::Tick!"));
+			UE_LOGF(LogPoseSearch, Log, "FInteractionIsland::FClosingTickFunction::ExecuteTick, called twice before UPoseSearchInteractionSubsystem::Tick!");
 		}
 		else
 		{
 			if (!Island->bPreTickFunctionExecuted)
 			{
-				UE_LOG(LogPoseSearch, Error, TEXT("FInteractionIsland::FClosingTickFunction::ExecuteTick, FPreTickFunction::ExecuteTick didn't run!"));
+				UE_LOGF(LogPoseSearch, Error, "FInteractionIsland::FClosingTickFunction::ExecuteTick, FPreTickFunction::ExecuteTick didn't run!");
 				Island->LogTickDependencies();
 			}
 
 			if (!Island->bPostTickFunctionExecuted)
 			{
-				UE_LOG(LogPoseSearch, Error, TEXT("FInteractionIsland::FClosingTickFunction::ExecuteTick, FPostTickFunction::ExecuteTick didn't run!"));
+				UE_LOGF(LogPoseSearch, Error, "FInteractionIsland::FClosingTickFunction::ExecuteTick, FPostTickFunction::ExecuteTick didn't run!");
 				Island->LogTickDependencies();
 			}
 
@@ -1183,7 +1183,7 @@ void FInteractionIsland::AddTickDependencies(UActorComponent* TickActorComponent
 		}
 		else
 		{
-			UE_LOG(LogPoseSearch, Error, TEXT("FInteractionIsland::AddTickDependencies, error while retrieving the tick function for %s"), *TickActorComponent->GetName());
+			UE_LOGF(LogPoseSearch, Error, "FInteractionIsland::AddTickDependencies, error while retrieving the tick function for %ls", *TickActorComponent->GetName());
 		}
 	}
 	else
@@ -1284,7 +1284,7 @@ void FInteractionIsland::AddTickDependencies(UActorComponent* TickActorComponent
 		if (!bIsPreTickFunctionValid || !bIsPostTickFunctionValid || !bIsClosingTickFunctionValid)
 		{
 			// if this validation triggers here, FInteractionIsland is not respectint the already present dependencies, creating cycles
-			UE_LOG(LogPoseSearch, Error, TEXT("============== FInteractionIsland::AddTickDependencies ValidateTickDependencies failed! Analyze the log and tune the FPoseSearchInteractionAvailability::TickPriority =============="));
+			UE_LOGF(LogPoseSearch, Error, "============== FInteractionIsland::AddTickDependencies ValidateTickDependencies failed! Analyze the log and tune the FPoseSearchInteractionAvailability::TickPriority ==============");
 			LogTickDependencies();
 		}
 	}
@@ -1307,19 +1307,19 @@ void FInteractionIsland::ValidateTickDependenciesLeaks() const
 			{
 				if (TickPrerequisite.PrerequisiteTickFunction == &PreTickFunction)
 				{
-					UE_LOG(LogPoseSearch, Error, TEXT("============== FInteractionIsland::ValidateTickDependenciesLeaks failed! PreTickFunction leaked in %s =============="), *TickActor->GetName());
+					UE_LOGF(LogPoseSearch, Error, "============== FInteractionIsland::ValidateTickDependenciesLeaks failed! PreTickFunction leaked in %ls ==============", *TickActor->GetName());
 					bAnyError = true;
 				}
 
 				if (TickPrerequisite.PrerequisiteTickFunction == &PostTickFunction)
 				{
-					UE_LOG(LogPoseSearch, Error, TEXT("============== FInteractionIsland::ValidateTickDependenciesLeaks failed! PostTickFunction leaked in %s =============="), *TickActor->GetName());
+					UE_LOGF(LogPoseSearch, Error, "============== FInteractionIsland::ValidateTickDependenciesLeaks failed! PostTickFunction leaked in %ls ==============", *TickActor->GetName());
 					bAnyError = true;
 				}
 
 				if (TickPrerequisite.PrerequisiteTickFunction == &ClosingTickFunction)
 				{
-					UE_LOG(LogPoseSearch, Error, TEXT("============== FInteractionIsland::ValidateTickDependenciesLeaks failed! ClosingTickFunction leaked in %s =============="), *TickActor->GetName());
+					UE_LOGF(LogPoseSearch, Error, "============== FInteractionIsland::ValidateTickDependenciesLeaks failed! ClosingTickFunction leaked in %ls ==============", *TickActor->GetName());
 					bAnyError = true;
 				}
 			}
@@ -1329,19 +1329,19 @@ void FInteractionIsland::ValidateTickDependenciesLeaks() const
 			{
 				if (TickPrerequisite.PrerequisiteTickFunction == &PreTickFunction)
 				{
-					UE_LOG(LogPoseSearch, Error, TEXT("============== FInteractionIsland::ValidateTickDependenciesLeaks failed! PreTickFunction leaked in %s =============="), *TickActorComponent->GetName());
+					UE_LOGF(LogPoseSearch, Error, "============== FInteractionIsland::ValidateTickDependenciesLeaks failed! PreTickFunction leaked in %ls ==============", *TickActorComponent->GetName());
 					bAnyError = true;
 				}
 
 				if (TickPrerequisite.PrerequisiteTickFunction == &PostTickFunction)
 				{
-					UE_LOG(LogPoseSearch, Error, TEXT("============== FInteractionIsland::ValidateTickDependenciesLeaks failed! PostTickFunction leaked in %s =============="), *TickActorComponent->GetName());
+					UE_LOGF(LogPoseSearch, Error, "============== FInteractionIsland::ValidateTickDependenciesLeaks failed! PostTickFunction leaked in %ls ==============", *TickActorComponent->GetName());
 					bAnyError = true;
 				}
 
 				if (TickPrerequisite.PrerequisiteTickFunction == &ClosingTickFunction)
 				{
-					UE_LOG(LogPoseSearch, Error, TEXT("============== FInteractionIsland::ValidateTickDependenciesLeaks failed! ClosingTickFunction leaked in %s =============="), *TickActorComponent->GetName());
+					UE_LOGF(LogPoseSearch, Error, "============== FInteractionIsland::ValidateTickDependenciesLeaks failed! ClosingTickFunction leaked in %ls ==============", *TickActorComponent->GetName());
 					bAnyError = true;
 				}
 			}
@@ -1369,15 +1369,15 @@ void FInteractionIsland::RemoveTickDependencies(bool bValidateTickDependencies)
 		{
 			if (bPreTickFunctionExecuted)
 			{
-				UE_LOG(LogPoseSearch, Error, TEXT("FInteractionIsland::RemoveTickDependencies, unexpected FPreTickFunction::ExecuteTick run!"));
+				UE_LOGF(LogPoseSearch, Error, "FInteractionIsland::RemoveTickDependencies, unexpected FPreTickFunction::ExecuteTick run!");
 			}
 			else if (bPostTickFunctionExecuted)
 			{
-				UE_LOG(LogPoseSearch, Error, TEXT("FInteractionIsland::RemoveTickDependencies, unexpected FPostTickFunction::ExecuteTick run!"));
+				UE_LOGF(LogPoseSearch, Error, "FInteractionIsland::RemoveTickDependencies, unexpected FPostTickFunction::ExecuteTick run!");
 			}
 			else
 			{
-				UE_LOG(LogPoseSearch, Error, TEXT("FInteractionIsland::RemoveTickDependencies, unexpected FClosingTickFunction::ExecuteTick run!"));
+				UE_LOGF(LogPoseSearch, Error, "FInteractionIsland::RemoveTickDependencies, unexpected FClosingTickFunction::ExecuteTick run!");
 			}
 				
 			LogTickDependencies();
@@ -1391,15 +1391,15 @@ void FInteractionIsland::RemoveTickDependencies(bool bValidateTickDependencies)
 		{
 			if (!bPreTickFunctionExecuted)
 			{
-				UE_LOG(LogPoseSearch, Error, TEXT("FInteractionIsland::RemoveTickDependencies, expected FPreTickFunction::ExecuteTick didn't run!"));
+				UE_LOGF(LogPoseSearch, Error, "FInteractionIsland::RemoveTickDependencies, expected FPreTickFunction::ExecuteTick didn't run!");
 			}
 			else if (!bPostTickFunctionExecuted)
 			{
-				UE_LOG(LogPoseSearch, Error, TEXT("FInteractionIsland::RemoveTickDependencies, expected FPostTickFunction::ExecuteTick didn't run!"));
+				UE_LOGF(LogPoseSearch, Error, "FInteractionIsland::RemoveTickDependencies, expected FPostTickFunction::ExecuteTick didn't run!");
 			}
 			else
 			{
-				UE_LOG(LogPoseSearch, Error, TEXT("FInteractionIsland::RemoveTickDependencies, expected FClosingTickFunction::ExecuteTick didn't run!"));
+				UE_LOGF(LogPoseSearch, Error, "FInteractionIsland::RemoveTickDependencies, expected FClosingTickFunction::ExecuteTick didn't run!");
 			}
 			
 			LogTickDependencies();
@@ -1495,19 +1495,19 @@ void FInteractionIsland::RemoveTickDependencies(bool bValidateTickDependencies)
 
 		if (!PreTickFunction.GetPrerequisites().IsEmpty())
 		{
-			UE_LOG(LogPoseSearch, Warning, TEXT("FInteractionIsland::RemoveTickDependencies, unable to cleanly remove all the PreTickFunction prerequisites! (This could have happened if an actor got garbage collected or a BP got recompiled while PIE was paused)"));
+			UE_LOGF(LogPoseSearch, Warning, "FInteractionIsland::RemoveTickDependencies, unable to cleanly remove all the PreTickFunction prerequisites! (This could have happened if an actor got garbage collected or a BP got recompiled while PIE was paused)");
 			PreTickFunction.GetPrerequisites().Reset();
 		}
 
 		if (!PostTickFunction.GetPrerequisites().IsEmpty())
 		{
-			UE_LOG(LogPoseSearch, Warning, TEXT("FInteractionIsland::RemoveTickDependencies, unable to cleanly remove all the PostTickFunction prerequisites! (This could have happened if an actor got garbage collected or a BP got recompiled while PIE was paused)"));
+			UE_LOGF(LogPoseSearch, Warning, "FInteractionIsland::RemoveTickDependencies, unable to cleanly remove all the PostTickFunction prerequisites! (This could have happened if an actor got garbage collected or a BP got recompiled while PIE was paused)");
 			PostTickFunction.GetPrerequisites().Reset();
 		}
 
 		if (!ClosingTickFunction.GetPrerequisites().IsEmpty())
 		{
-			UE_LOG(LogPoseSearch, Warning, TEXT("FInteractionIsland::RemoveTickDependencies, unable to cleanly remove all the ClosingTickFunction prerequisites! (This could have happened if an actor got garbage collected or a BP got recompiled while PIE was paused)"));
+			UE_LOGF(LogPoseSearch, Warning, "FInteractionIsland::RemoveTickDependencies, unable to cleanly remove all the ClosingTickFunction prerequisites! (This could have happened if an actor got garbage collected or a BP got recompiled while PIE was paused)");
 			ClosingTickFunction.GetPrerequisites().Reset();
 		}
 
@@ -1534,15 +1534,15 @@ void FInteractionIsland::InjectToActor(const UObject* AnimContext, bool bAddTick
 		{
 			if (bPreTickFunctionExecuted)
 			{
-				UE_LOG(LogPoseSearch, Error, TEXT("FInteractionIsland::InjectToActor, unexpected FPreTickFunction::ExecuteTick run!"));
+				UE_LOGF(LogPoseSearch, Error, "FInteractionIsland::InjectToActor, unexpected FPreTickFunction::ExecuteTick run!");
 			}
 			else if (bPostTickFunctionExecuted)
 			{
-				UE_LOG(LogPoseSearch, Error, TEXT("FInteractionIsland::InjectToActor, unexpected FPostTickFunction::ExecuteTick run!"));
+				UE_LOGF(LogPoseSearch, Error, "FInteractionIsland::InjectToActor, unexpected FPostTickFunction::ExecuteTick run!");
 			}
 			else
 			{
-				UE_LOG(LogPoseSearch, Error, TEXT("FInteractionIsland::InjectToActor, unexpected FClosingTickFunction::ExecuteTick run!"));
+				UE_LOGF(LogPoseSearch, Error, "FInteractionIsland::InjectToActor, unexpected FClosingTickFunction::ExecuteTick run!");
 			}
 				
 			LogTickDependencies();
@@ -1602,7 +1602,7 @@ void FInteractionIsland::Uninitialize(bool bValidateTickDependencies)
 		if (!bIsPreTickFunctionValid || !bIsPostTickFunctionValid || !bIsClosingTickFunctionValid)
 		{
 			// if this validation triggers here, some additional tick dependency outside FInteractionIsland has been injected witout respecting the already present dependencies, creating cycles
-			UE_LOG(LogPoseSearch, Error, TEXT("============== FInteractionIsland::Uninitialize ValidateTickDependencies failed! =============="));
+			UE_LOGF(LogPoseSearch, Error, "============== FInteractionIsland::Uninitialize ValidateTickDependencies failed! ==============");
 			LogTickDependencies();
 		}
 	}
@@ -1660,7 +1660,7 @@ void FInteractionIsland::LogTickDependencies(const TConstArrayView<UActorCompone
 	{
 		if (TickActorComponent)
 		{
-			UE_LOG(LogPoseSearch, Log, TEXT("============== %s (Island %d) =============="), *TickActorComponent->GetOwner()->GetName(), InteractionIslandIndex);
+			UE_LOGF(LogPoseSearch, Log, "============== %ls (Island %d) ==============", *TickActorComponent->GetOwner()->GetName(), InteractionIslandIndex);
 			if (IInteractionIslandDependency* IslandDependency = FindCustomDependency(TickActorComponent))
 			{
 				// AnimNextComponent case
@@ -1670,7 +1670,7 @@ void FInteractionIsland::LogTickDependencies(const TConstArrayView<UActorCompone
 				}
 				else
 				{
-					UE_LOG(LogPoseSearch, Error, TEXT("FInteractionIsland::LogTickDependencies, error while retrieving the tick function for to %s"), *TickActorComponent->GetName());
+					UE_LOGF(LogPoseSearch, Error, "FInteractionIsland::LogTickDependencies, error while retrieving the tick function for to %ls", *TickActorComponent->GetName());
 				}
 			}
 			else
@@ -1681,7 +1681,7 @@ void FInteractionIsland::LogTickDependencies(const TConstArrayView<UActorCompone
 		}
 		else
 		{
-			UE_LOG(LogPoseSearch, Log, TEXT("============== !!!Missing Actor!!! (Island %d) =============="), InteractionIslandIndex);
+			UE_LOGF(LogPoseSearch, Log, "============== !!!Missing Actor!!! (Island %d) ==============", InteractionIslandIndex);
 		}
 	}
 }
diff --git a/Engine/Plugins/Animation/PoseSearch/Source/Runtime/Private/PoseSearchInteractionLibrary.cpp b/Engine/Plugins/Animation/PoseSearch/Source/Runtime/Private/PoseSearchInteractionLibrary.cpp
--- a/Engine/Plugins/Animation/PoseSearch/Source/Runtime/Private/PoseSearchInteractionLibrary.cpp
+++ b/Engine/Plugins/Animation/PoseSearch/Source/Runtime/Private/PoseSearchInteractionLibrary.cpp
@@ -101,7 +101,7 @@ void UPoseSearchInteractionLibrary::CalculateFullAlignedTransforms(const FPoseSe
 	}
 	else
 	{
-		UE_LOG(LogPoseSearch, Error, TEXT("UPoseSearchInteractionLibrary::CalculateFullAlignedTransforms failed!"));
+		UE_LOGF(LogPoseSearch, Error, "UPoseSearchInteractionLibrary::CalculateFullAlignedTransforms failed!");
 		OutFullAlignedTransforms.Reset();
 	}
 }
@@ -117,7 +117,7 @@ void UPoseSearchInteractionLibrary::CalculateFullAlignedTransform(const FPoseSea
 	}
 	else
 	{
-		UE_LOG(LogPoseSearch, Error, TEXT("UPoseSearchInteractionLibrary::CalculateFullAlignedTransform failed!"));
+		UE_LOGF(LogPoseSearch, Error, "UPoseSearchInteractionLibrary::CalculateFullAlignedTransform failed!");
 		OutFullAlignedTransform.SetIdentity();
 	}
 }
@@ -128,7 +128,7 @@ void UPoseSearchInteractionLibrary::CalculateNoAlignedTransform(const FPoseSearc
 
 	if (!Result.ActorRootTransforms.IsValidIndex(Result.RoleIndex))
 	{
-		UE_LOG(LogPoseSearch, Error, TEXT("UPoseSearchInteractionLibrary::CalculateNoAlignedTransform failed!"));
+		UE_LOGF(LogPoseSearch, Error, "UPoseSearchInteractionLibrary::CalculateNoAlignedTransform failed!");
 		OutNoAlignedTransform.SetIdentity();
 		return;
 	}
@@ -136,7 +136,7 @@ void UPoseSearchInteractionLibrary::CalculateNoAlignedTransform(const FPoseSearc
 	const UMultiAnimAsset* MultiAnimAsset = Cast<UMultiAnimAsset>(Result.SelectedAnim);
 	if (!MultiAnimAsset)
 	{
-		UE_LOG(LogPoseSearch, Error, TEXT("UPoseSearchInteractionLibrary::CalculateNoAlignedTransform failed!"));
+		UE_LOGF(LogPoseSearch, Error, "UPoseSearchInteractionLibrary::CalculateNoAlignedTransform failed!");
 		OutNoAlignedTransform.SetIdentity();
 		return;
 	}
@@ -144,7 +144,7 @@ void UPoseSearchInteractionLibrary::CalculateNoAlignedTransform(const FPoseSearc
 	const UAnimationAsset* AnimationAsset = MultiAnimAsset->GetAnimationAsset(Result.Role);
 	if (!AnimationAsset)
 	{
-		UE_LOG(LogPoseSearch, Error, TEXT("UPoseSearchInteractionLibrary::CalculateNoAlignedTransform failed!"));
+		UE_LOGF(LogPoseSearch, Error, "UPoseSearchInteractionLibrary::CalculateNoAlignedTransform failed!");
 		OutNoAlignedTransform.SetIdentity();
 		return;
 	}
diff --git a/Engine/Plugins/Animation/PoseSearch/Source/Runtime/Private/PoseSearchInteractionSubsystem.cpp b/Engine/Plugins/Animation/PoseSearch/Source/Runtime/Private/PoseSearchInteractionSubsystem.cpp
--- a/Engine/Plugins/Animation/PoseSearch/Source/Runtime/Private/PoseSearchInteractionSubsystem.cpp
+++ b/Engine/Plugins/Animation/PoseSearch/Source/Runtime/Private/PoseSearchInteractionSubsystem.cpp
@@ -169,7 +169,7 @@ const UE::PoseSearch::IPoseHistory* FInteractionAvailability::GetPoseHistory(con
 		}
 	}
 
-	UE_LOG(LogPoseSearch, Error, TEXT("FInteractionAvailability::GetPoseHistory couldn't find a pose history for AnimContext %s with name %s"), *GetNameSafe(AnimContext), *PoseHistoryName.ToString());
+	UE_LOGF(LogPoseSearch, Error, "FInteractionAvailability::GetPoseHistory couldn't find a pose history for AnimContext %ls with name %ls", *GetNameSafe(AnimContext), *PoseHistoryName.ToString());
 	return nullptr;
 }
 
@@ -1029,7 +1029,7 @@ void UPoseSearchInteractionSubsystem::GenerateSearchContexts(float DeltaSeconds,
 #if !NO_LOGGING
 					if (RoledSkeletonIndex == RoledSkeletons.Num())
 					{
-						UE_LOG(LogPoseSearch, Warning, TEXT("UPoseSearchInteractionSubsystem::GenerateSearchContexts unsupported Role %s for Database %s"), *Role.ToString(), *Database.GetName());
+						UE_LOGF(LogPoseSearch, Warning, "UPoseSearchInteractionSubsystem::GenerateSearchContexts unsupported Role %ls for Database %ls", *Role.ToString(), *Database.GetName());
 					}
 #endif // !NO_LOGGING
 				}
@@ -1173,7 +1173,7 @@ void UPoseSearchInteractionSubsystem::GenerateSearchContexts(float DeltaSeconds,
 					}
 					else
 					{
-						UE_LOG(LogPoseSearch, Log, TEXT("UPoseSearchInteractionSubsystem::GenerateSearchContexts null Availability.Database (with invalid Availability.Tag)"));
+						UE_LOGF(LogPoseSearch, Log, "UPoseSearchInteractionSubsystem::GenerateSearchContexts null Availability.Database (with invalid Availability.Tag)");
 					}
 				}
 			}
@@ -1362,7 +1362,7 @@ bool UPoseSearchInteractionSubsystem::ShouldRegenerateAllIslands()
 
 	if (AnimContextsAvailabilitiesIndexValue > AnimContextsAvailabilitiesNum)
 	{
-		UE_LOG(LogPoseSearch, Warning, TEXT("UPoseSearchInteractionSubsystem::ShouldRegenerateAllIslands not enough space to add more availabilities locklessly. It'll be adjusted automatically, but some availability requests has been lost this frame [capacity %d / requests %d]. To void this warning set the initial budget for 'Edit/ProjectSettings/Plugins/Pose Search/Availabilities Buffer Size' to %d. "), AnimContextsAvailabilitiesNum, AnimContextsAvailabilitiesIndexValue, AnimContextsAvailabilitiesIndexValue);
+		UE_LOGF(LogPoseSearch, Warning, "UPoseSearchInteractionSubsystem::ShouldRegenerateAllIslands not enough space to add more availabilities locklessly. It'll be adjusted automatically, but some availability requests has been lost this frame [capacity %d / requests %d]. To void this warning set the initial budget for 'Edit/ProjectSettings/Plugins/Pose Search/Availabilities Buffer Size' to %d. ", AnimContextsAvailabilitiesNum, AnimContextsAvailabilitiesIndexValue, AnimContextsAvailabilitiesIndexValue);
 		AnimContextsAvailabilitiesBuffer.SetNum(AnimContextsAvailabilitiesIndexValue);
 	}
 	else
@@ -1644,7 +1644,7 @@ void UPoseSearchInteractionSubsystem::DebugLogTickDependencies() const
 
 	if (GVarPoseSearchInteractionLoglandsTickDependencies)
 	{
-		UE_LOG(LogPoseSearch, Log, TEXT("=================================================================="));
+		UE_LOGF(LogPoseSearch, Log, "==================================================================");
 		for (const FInteractionIsland* Island : Islands)
 		{
 			if (Island->IsInitialized())
diff --git a/Engine/Plugins/Animation/PoseSearch/Source/Runtime/Private/PoseSearchInteractionValidator.cpp b/Engine/Plugins/Animation/PoseSearch/Source/Runtime/Private/PoseSearchInteractionValidator.cpp
--- a/Engine/Plugins/Animation/PoseSearch/Source/Runtime/Private/PoseSearchInteractionValidator.cpp
+++ b/Engine/Plugins/Animation/PoseSearch/Source/Runtime/Private/PoseSearchInteractionValidator.cpp
@@ -32,7 +32,7 @@ FInteractionValidator::FInteractionValidator(const UObject* AnimContext)
 	const AActor* MainActor = ValidatingIsland->GetMainActor();
 	if (!MainActor)
 	{
-		UE_LOG(LogPoseSearch, Error, TEXT("FInteractionValidator invalid MainActor! How did he die? Did you rebuild the animation blue print while PIE was running?"));
+		UE_LOGF(LogPoseSearch, Error, "FInteractionValidator invalid MainActor! How did he die? Did you rebuild the animation blue print while PIE was running?");
 		ValidatingIsland->LogTickDependencies();
 		return;
 	}
@@ -42,7 +42,7 @@ FInteractionValidator::FInteractionValidator(const UObject* AnimContext)
 		const int32 PreviousCounter = ValidatingIsland->InteractionIslandThreadSafeCounter.Set(-1);
 		if (PreviousCounter != 0)
 		{
-			UE_LOG(LogPoseSearch, Error, TEXT("Non thread safe call! Why is there any other actor running while we schedule the MainActor (%s)? %d"), *MainActor->GetName(), PreviousCounter);
+			UE_LOGF(LogPoseSearch, Error, "Non thread safe call! Why is there any other actor running while we schedule the MainActor (%ls)? %d", *MainActor->GetName(), PreviousCounter);
 			ValidatingIsland->LogTickDependencies();
 		}
 	}
@@ -52,7 +52,7 @@ FInteractionValidator::FInteractionValidator(const UObject* AnimContext)
 		const int32 PreviousCounter = ValidatingIsland->InteractionIslandThreadSafeCounter.Add(1);
 		if (PreviousCounter < 0)
 		{
-			UE_LOG(LogPoseSearch, Error, TEXT("Non thread safe call! The MainActor (%s) is running! Nobody else (%s) in the same island should be running! %d"), *MainActor->GetName(), *Actor->GetName(),  PreviousCounter);
+			UE_LOGF(LogPoseSearch, Error, "Non thread safe call! The MainActor (%ls) is running! Nobody else (%ls) in the same island should be running! %d", *MainActor->GetName(), *Actor->GetName(),  PreviousCounter);
 			ValidatingIsland->LogTickDependencies();
 		}
 	}
@@ -65,7 +65,7 @@ FInteractionValidator::FInteractionValidator(FInteractionIsland* InValidatingIsl
 	const int32 PreviousCounter = ValidatingIsland->TickFunctionsThreadSafeCounter.Set(-1);
 	if (PreviousCounter != 0)
 	{
-		UE_LOG(LogPoseSearch, Error, TEXT("Non thread safe call! TickFunctions running concurrently? %d"), PreviousCounter);
+		UE_LOGF(LogPoseSearch, Error, "Non thread safe call! TickFunctions running concurrently? %d", PreviousCounter);
 		ValidatingIsland->LogTickDependencies();
 	}
 }
@@ -78,7 +78,7 @@ FInteractionValidator::~FInteractionValidator()
 		const int32 PreviousCounter = ValidatingIsland->TickFunctionsThreadSafeCounter.Set(0);
 		if (PreviousCounter != -1)
 		{
-			UE_LOG(LogPoseSearch, Error, TEXT("Non thread safe call! TickFunctions running concurrently? %d"), PreviousCounter);
+			UE_LOGF(LogPoseSearch, Error, "Non thread safe call! TickFunctions running concurrently? %d", PreviousCounter);
 		}
 		return;
 	}
@@ -97,7 +97,7 @@ FInteractionValidator::~FInteractionValidator()
 	FInteractionIsland* Island = InteractionSubsystem->FindIsland(ValidatingAnimContext);
 	if (Island != ValidatingIsland)
 	{
-		UE_LOG(LogPoseSearch, Error, TEXT("FInteractionValidator why did the InteractionIsland changed?"));
+		UE_LOGF(LogPoseSearch, Error, "FInteractionValidator why did the InteractionIsland changed?");
 		return;
 	}
 
@@ -107,7 +107,7 @@ FInteractionValidator::~FInteractionValidator()
 	const AActor* MainActor = Island->GetMainActor();
 	if (!MainActor)
 	{
-		UE_LOG(LogPoseSearch, Error, TEXT("FInteractionValidator invalid MainActor! How did he died? Did you rebuild the animation blue print while PIE was running?"));
+		UE_LOGF(LogPoseSearch, Error, "FInteractionValidator invalid MainActor! How did he died? Did you rebuild the animation blue print while PIE was running?");
 		return;
 	}
 	
@@ -116,7 +116,7 @@ FInteractionValidator::~FInteractionValidator()
 		const int32 PreviousCounter = Island->InteractionIslandThreadSafeCounter.Set(0);
 		if (PreviousCounter >= 0)
 		{
-			UE_LOG(LogPoseSearch, Error, TEXT("Non thread safe call! Why was there some other actor running while we ended the scheduling the MainActor (%s)? %d"), *MainActor->GetName(), PreviousCounter);
+			UE_LOGF(LogPoseSearch, Error, "Non thread safe call! Why was there some other actor running while we ended the scheduling the MainActor (%ls)? %d", *MainActor->GetName(), PreviousCounter);
 		}
 	}
 	else
@@ -125,7 +125,7 @@ FInteractionValidator::~FInteractionValidator()
 		const int32 PreviousCounter = Island->InteractionIslandThreadSafeCounter.Add(-1);
 		if (PreviousCounter <= 0)
 		{
-			UE_LOG(LogPoseSearch, Error, TEXT("Non thread safe call! The MainActor (%s) is running! Nobody else (%s) in the same island should be running! %d"), *MainActor->GetName(), *Actor->GetName(), PreviousCounter);
+			UE_LOGF(LogPoseSearch, Error, "Non thread safe call! The MainActor (%ls) is running! Nobody else (%ls) in the same island should be running! %d", *MainActor->GetName(), *Actor->GetName(), PreviousCounter);
 		}
 	}
 }
diff --git a/Engine/Plugins/Animation/PoseSearch/Source/Runtime/Private/PoseSearchLibrary.cpp b/Engine/Plugins/Animation/PoseSearch/Source/Runtime/Private/PoseSearchLibrary.cpp
--- a/Engine/Plugins/Animation/PoseSearch/Source/Runtime/Private/PoseSearchLibrary.cpp
+++ b/Engine/Plugins/Animation/PoseSearch/Source/Runtime/Private/PoseSearchLibrary.cpp
@@ -209,7 +209,7 @@ namespace UE::PoseSearch
 			{
 				if (bContainsIsMandatory)
 				{
-					UE_LOG(LogPoseSearch, Error, TEXT("improperly setup UAnimSequenceBase. Database %s doesn't contain UAnimSequenceBase %s"), *Database->GetName(), *AssetToSearch->GetName());
+					UE_LOGF(LogPoseSearch, Error, "improperly setup UAnimSequenceBase. Database %ls doesn't contain UAnimSequenceBase %ls", *Database->GetName(), *AssetToSearch->GetName());
 				}
 			}
 			else
@@ -241,7 +241,7 @@ namespace UE::PoseSearch
 					{
 						if (!PoseSearchBranchIn->Database)
 						{
-							UE_LOG(LogPoseSearch, Error, TEXT("improperly setup UAnimNotifyState_PoseSearchBranchIn with null Database in %s"), *SequenceBase->GetName());
+							UE_LOGF(LogPoseSearch, Error, "improperly setup UAnimNotifyState_PoseSearchBranchIn with null Database in %ls", *SequenceBase->GetName());
 							continue;
 						}
 
@@ -419,7 +419,7 @@ namespace UE::PoseSearch
 					{
 						if (Role != Schema->GetDefaultRole())
 						{
-							UE_LOG(LogPoseSearch, Error, TEXT("GetCommonDefaultRole - inconsistent Role between provided Databases!"));
+							UE_LOGF(LogPoseSearch, Error, "GetCommonDefaultRole - inconsistent Role between provided Databases!");
 							break;
 						}
 					}
@@ -456,7 +456,7 @@ namespace UE::PoseSearch
 			}
 			else if (!ensure(PlayRate.Min <= PlayRate.Max && PlayRate.Min > UE_KINDA_SMALL_NUMBER))
 			{
-				UE_LOG(LogPoseSearch, Error, TEXT("Couldn't update the WantedPlayRate in CalculateWantedPlayRate, because of invalid PlayRate interval (%f, %f)"), PlayRate.Min, PlayRate.Max);
+				UE_LOGF(LogPoseSearch, Error, "Couldn't update the WantedPlayRate in CalculateWantedPlayRate, because of invalid PlayRate interval (%f, %f)", PlayRate.Min, PlayRate.Max);
 				WantedPlayRate = 1.f;
 			}
 			else if (!FMath::IsNearlyEqual(PlayRate.Min, PlayRate.Max, UE_KINDA_SMALL_NUMBER))
@@ -481,8 +481,8 @@ namespace UE::PoseSearch
 							const TConstArrayView<float> ResultData = SearchIndex.GetReconstructedPoseValues(SearchResult.PoseIdx, ReconstructedPoseValuesBuffer);
 							if (ResultData.IsEmpty())
 							{
-								UE_LOG(LogPoseSearch, Warning,
-									TEXT("Couldn't update the WantedPlayRate in CalculateWantedPlayRate, because couldn't reconstruct the pose value in GetReconstructedPoseValues for pose %d"),
+								UE_LOGF(LogPoseSearch, Warning,
+									"Couldn't update the WantedPlayRate in CalculateWantedPlayRate, because couldn't reconstruct the pose value in GetReconstructedPoseValues for pose %d",
 									SearchResult.PoseIdx);
 								WantedPlayRate = FMath::Clamp(1.f, PlayRate.Min, PlayRate.Max);
 							}
@@ -501,8 +501,8 @@ namespace UE::PoseSearch
 					}
 					else
 					{
-						UE_LOG(LogPoseSearch, Warning,
-							TEXT("Couldn't update the WantedPlayRate in CalculateWantedPlayRate, because Schema '%s' couldn't find a UPoseSearchFeatureChannel_Trajectory channel"),
+						UE_LOGF(LogPoseSearch, Warning,
+							"Couldn't update the WantedPlayRate in CalculateWantedPlayRate, because Schema '%ls' couldn't find a UPoseSearchFeatureChannel_Trajectory channel",
 							*GetNameSafe(SearchResult.Database->Schema));
 					}
 				}
@@ -920,11 +920,11 @@ void UPoseSearchLibrary::UpdateMotionMatchingState(
 
 			if (bIsIndexing)
 			{
-				UE_LOG(LogPoseSearch, Log, TEXT("%s"), *String);
+				UE_LOGF(LogPoseSearch, Log, "%ls", *String);
 			}
 			else
 			{
-				UE_LOG(LogPoseSearch, Warning, TEXT("%s"), *String);
+				UE_LOGF(LogPoseSearch, Warning, "%ls", *String);
 			}
 		}
 #endif // !NO_LOGGING
@@ -1035,13 +1035,13 @@ void UPoseSearchLibrary::MotionMatch(
 
 	if (!AnimInstance)
 	{
-		UE_LOG(LogPoseSearch, Error, TEXT("UPoseSearchLibrary::MotionMatch - null AnimInstances"));
+		UE_LOGF(LogPoseSearch, Error, "UPoseSearchLibrary::MotionMatch - null AnimInstances");
 		return;
 	}
 
 	if (!AnimInstance->CurrentSkeleton)
 	{
-		UE_LOG(LogPoseSearch, Error, TEXT("UPoseSearchLibrary::MotionMatch - null AnimInstances->CurrentSkeleton"));
+		UE_LOGF(LogPoseSearch, Error, "UPoseSearchLibrary::MotionMatch - null AnimInstances->CurrentSkeleton");
 		return;
 	}
 
@@ -1069,13 +1069,13 @@ void UPoseSearchLibrary::MotionMatch(
 		}
 		else
 		{
-			UE_LOG(LogPoseSearch, Error, TEXT("UPoseSearchLibrary::MotionMatch - Couldn't find a valid pose history in the pose history node '%s'"), *PoseHistoryName.ToString());
+			UE_LOGF(LogPoseSearch, Error, "UPoseSearchLibrary::MotionMatch - Couldn't find a valid pose history in the pose history node '%ls'", *PoseHistoryName.ToString());
 			return;
 		}
 	}
 	else
 	{
-		UE_LOG(LogPoseSearch, Error, TEXT("UPoseSearchLibrary::MotionMatch - Couldn't find pose history node with name '%s'"), *PoseHistoryName.ToString());
+		UE_LOGF(LogPoseSearch, Error, "UPoseSearchLibrary::MotionMatch - Couldn't find pose history node with name '%ls'", *PoseHistoryName.ToString());
 		return;
 	}
 
@@ -1196,14 +1196,14 @@ void UPoseSearchLibrary::MotionMatch(
 		float FutureAnimationTime = Future.AnimationTime;
 		if (FutureAnimationTime < FiniteDelta)
 		{
-			UE_LOG(LogPoseSearch, Warning, TEXT("UPoseSearchLibrary::MotionMatch - provided Future.AnimationTime (%f) is too small to be able to calculate velocities. Clamping it to minimum value of %f"), FutureAnimationTime, FiniteDelta);
+			UE_LOGF(LogPoseSearch, Warning, "UPoseSearchLibrary::MotionMatch - provided Future.AnimationTime (%f) is too small to be able to calculate velocities. Clamping it to minimum value of %f", FutureAnimationTime, FiniteDelta);
 			FutureAnimationTime = FiniteDelta;
 		}
 
 		const float MinFutureIntervalTime = FiniteDelta + UE_KINDA_SMALL_NUMBER;
 		if (FutureIntervalTime < MinFutureIntervalTime)
 		{
-			UE_LOG(LogPoseSearch, Warning, TEXT("UPoseSearchLibrary::MotionMatch - provided TimeToFutureAnimationStart (%f) is too small. Clamping it to minimum value of %f"), FutureIntervalTime, MinFutureIntervalTime);
+			UE_LOGF(LogPoseSearch, Warning, "UPoseSearchLibrary::MotionMatch - provided TimeToFutureAnimationStart (%f) is too small. Clamping it to minimum value of %f", FutureIntervalTime, MinFutureIntervalTime);
 			FutureIntervalTime = MinFutureIntervalTime;
 		}
 
@@ -1390,7 +1390,7 @@ void UPoseSearchLibrary::MotionMatch(
 #if WITH_EDITOR && !NO_LOGGING
 	if (SearchContext.IsAsyncBuildIndexInProgress())
 	{
-		UE_LOG(LogPoseSearch, Warning, TEXT("UPoseSearchLibrary::MotionMatch - some searches have been skipped, since databases AsyncBuildIndex are in still in progress or failed"));
+		UE_LOGF(LogPoseSearch, Warning, "UPoseSearchLibrary::MotionMatch - some searches have been skipped, since databases AsyncBuildIndex are in still in progress or failed");
 	}
 #endif // WITH_EDITOR && !NO_LOGGING
 }
diff --git a/Engine/Plugins/Animation/PoseSearch/Source/Runtime/Private/PoseSearchMirrorDataCache.cpp b/Engine/Plugins/Animation/PoseSearch/Source/Runtime/Private/PoseSearchMirrorDataCache.cpp
--- a/Engine/Plugins/Animation/PoseSearch/Source/Runtime/Private/PoseSearchMirrorDataCache.cpp
+++ b/Engine/Plugins/Animation/PoseSearch/Source/Runtime/Private/PoseSearchMirrorDataCache.cpp
@@ -41,7 +41,7 @@ void FMirrorDataCache::Init(const UMirrorDataTable* InMirrorDataTable)
 		}
 		else
 		{
-			UE_LOG(LogPoseSearch, Error, TEXT("FMirrorDataCache::Init - UMirrorDataTable '%s' Skeleton is not set!"), *InMirrorDataTable->GetName());
+			UE_LOGF(LogPoseSearch, Error, "FMirrorDataCache::Init - UMirrorDataTable '%ls' Skeleton is not set!", *InMirrorDataTable->GetName());
 			Reset();
 		}
 	}
@@ -63,7 +63,7 @@ void FMirrorDataCache::Init(const UMirrorDataTable* InMirrorDataTable, const FBo
 		}
 		else
 		{
-			UE_LOG(LogPoseSearch, Error, TEXT("FMirrorDataCache::Init - UMirrorDataTable '%s' Skeleton is not set!"), *InMirrorDataTable->GetName());
+			UE_LOGF(LogPoseSearch, Error, "FMirrorDataCache::Init - UMirrorDataTable '%ls' Skeleton is not set!", *InMirrorDataTable->GetName());
 			Reset();
 		}
 	}
diff --git a/Engine/Plugins/Animation/PoseSearch/Source/Runtime/Private/PoseSearchModule.cpp b/Engine/Plugins/Animation/PoseSearch/Source/Runtime/Private/PoseSearchModule.cpp
--- a/Engine/Plugins/Animation/PoseSearch/Source/Runtime/Private/PoseSearchModule.cpp
+++ b/Engine/Plugins/Animation/PoseSearch/Source/Runtime/Private/PoseSearchModule.cpp
@@ -75,7 +75,7 @@ public:
 		UE::Anim::IPoseSearchProvider::FSearchResult ProviderResult;
 		if (!PoseHistory)
 		{
-			UE_LOG(LogPoseSearch, Error, TEXT("FPoseSearchModule::Search, missing IPoseHistory"));
+			UE_LOGF(LogPoseSearch, Error, "FPoseSearchModule::Search, missing IPoseHistory");
 		}
 		else
 		{
diff --git a/Engine/Plugins/Animation/PoseSearch/Source/Runtime/Private/PoseSearchSchema.cpp b/Engine/Plugins/Animation/PoseSearch/Source/Runtime/Private/PoseSearchSchema.cpp
--- a/Engine/Plugins/Animation/PoseSearch/Source/Runtime/Private/PoseSearchSchema.cpp
+++ b/Engine/Plugins/Animation/PoseSearch/Source/Runtime/Private/PoseSearchSchema.cpp
@@ -154,7 +154,7 @@ int8 UPoseSearchSchema::FindSchemaBoneIdx(const FName& BoneName, const UE::PoseS
 	const FPoseSearchRoledSkeleton* RoledSkeleton = GetRoledSkeleton(Role);
 	if (!RoledSkeleton)
 	{
-		UE_LOG(LogPoseSearch, Error, TEXT("UPoseSearchSchema::FindSchemaBoneIdx: couldn't find data for the requested Role '%s' in UPoseSearchSchema '%s'"), *Role.ToString(), *GetNameSafe(this));
+		UE_LOGF(LogPoseSearch, Error, "UPoseSearchSchema::FindSchemaBoneIdx: couldn't find data for the requested Role '%ls' in UPoseSearchSchema '%ls'", *Role.ToString(), *GetNameSafe(this));
 		return InvalidSchemaBoneIdx;
 	}
 
@@ -193,7 +193,7 @@ int8 UPoseSearchSchema::AddBoneReference(const FBoneReference& BoneReference, co
 	FPoseSearchRoledSkeleton* RoledSkeleton = GetRoledSkeleton(Role);
 	if (!RoledSkeleton)
 	{
-		UE_LOG(LogPoseSearch, Error, TEXT("UPoseSearchSchema::AddBoneReference: couldn't find data for the requested Role '%s' in UPoseSearchSchema '%s'"), *Role.ToString(), *GetNameSafe(this));
+		UE_LOGF(LogPoseSearch, Error, "UPoseSearchSchema::AddBoneReference: couldn't find data for the requested Role '%ls' in UPoseSearchSchema '%ls'", *Role.ToString(), *GetNameSafe(this));
 		return InvalidSchemaBoneIdx;
 	}
 
@@ -201,7 +201,7 @@ int8 UPoseSearchSchema::AddBoneReference(const FBoneReference& BoneReference, co
 	const USkeleton* Skeleton = RoledSkeleton->Skeleton;
 	if (!Skeleton)
 	{
-		UE_LOG(LogPoseSearch, Error, TEXT("UPoseSearchSchema::AddBoneReference: couldn't find Skeleton with Role '%s' in UPoseSearchSchema '%s'"), *Role.ToString(), *GetNameSafe(this));
+		UE_LOGF(LogPoseSearch, Error, "UPoseSearchSchema::AddBoneReference: couldn't find Skeleton with Role '%ls' in UPoseSearchSchema '%ls'", *Role.ToString(), *GetNameSafe(this));
 		return InvalidSchemaBoneIdx;
 	}
 
@@ -212,7 +212,7 @@ int8 UPoseSearchSchema::AddBoneReference(const FBoneReference& BoneReference, co
 		TempBoneReference.Initialize(Skeleton);
 		if (!TempBoneReference.HasValidSetup())
 		{
-			UE_LOG(LogPoseSearch, Error, TEXT("UPoseSearchSchema::AddBoneReference: couldn't initialize FBoneReference '%s' with Skeleton '%s' with Role '%s' in UPoseSearchSchema '%s'"),
+			UE_LOGF(LogPoseSearch, Error, "UPoseSearchSchema::AddBoneReference: couldn't initialize FBoneReference '%ls' with Skeleton '%ls' with Role '%ls' in UPoseSearchSchema '%ls'",
 				*TempBoneReference.BoneName.ToString(), *GetNameSafe(Skeleton), *Role.ToString(), *GetNameSafe(this));
 			return InvalidSchemaBoneIdx;
 		}
@@ -240,14 +240,14 @@ int8 UPoseSearchSchema::AddCurveReference(const FName& CurveReference, const UE:
 	FPoseSearchRoledSkeleton* RoledSkeleton = GetRoledSkeleton(Role);
 	if (!RoledSkeleton)
 	{
-		UE_LOG(LogPoseSearch, Error, TEXT("UPoseSearchSchema::AddCurveReference: couldn't find data for the requested Role '%s' in UPoseSearchSchema '%s'"), *Role.ToString(), *GetNameSafe(this));
+		UE_LOGF(LogPoseSearch, Error, "UPoseSearchSchema::AddCurveReference: couldn't find data for the requested Role '%ls' in UPoseSearchSchema '%ls'", *Role.ToString(), *GetNameSafe(this));
 		return InvalidSchemaCurveIdx;
 	}
 
 	const USkeleton* Skeleton = RoledSkeleton->Skeleton;
 	if (!Skeleton)
 	{
-		UE_LOG(LogPoseSearch, Error, TEXT("UPoseSearchSchema::AddCurveReference: couldn't find Skeleton with Role '%s' in UPoseSearchSchema '%s'"), *Role.ToString(), *GetNameSafe(this));
+		UE_LOGF(LogPoseSearch, Error, "UPoseSearchSchema::AddCurveReference: couldn't find Skeleton with Role '%ls' in UPoseSearchSchema '%ls'", *Role.ToString(), *GetNameSafe(this));
 		return InvalidSchemaCurveIdx;
 	}
 
@@ -282,7 +282,7 @@ void UPoseSearchSchema::Finalize()
 		{
 			if (Skeletons[ComparisonIndex].Role == RoledSkeleton.Role)
 			{
-				UE_LOG(LogPoseSearch, Error, TEXT("UPoseSearchSchema::Finalize: couldn't Finalize '%s' because of duplicate Role '%s' in Skeletons"), *GetNameSafe(this), *RoledSkeleton.Role.ToString());
+				UE_LOGF(LogPoseSearch, Error, "UPoseSearchSchema::Finalize: couldn't Finalize '%ls' because of duplicate Role '%ls' in Skeletons", *GetNameSafe(this), *RoledSkeleton.Role.ToString());
 
 				ResetFinalize();
 				return;
@@ -292,7 +292,7 @@ void UPoseSearchSchema::Finalize()
 		const int8 SchemaBoneIdx = AddBoneReference(FBoneReference(), RoledSkeleton.Role);
 		if (SchemaBoneIdx != RootSchemaBoneIdx)
 		{
-			UE_LOG(LogPoseSearch, Error, TEXT("UPoseSearchSchema::Finalize: couldn't Finalize '%s' because couldn't initialize root bone properly"), *GetNameSafe(this));
+			UE_LOGF(LogPoseSearch, Error, "UPoseSearchSchema::Finalize: couldn't Finalize '%ls' because couldn't initialize root bone properly", *GetNameSafe(this));
 
 			ResetFinalize();
 			return;
@@ -309,9 +309,9 @@ void UPoseSearchSchema::Finalize()
 				#if WITH_EDITOR
 				TLabelBuilder LabelBuilder;
 				FString Label = ChannelPtr->GetLabel(LabelBuilder).ToString();
-				UE_LOG(LogPoseSearch, Error, TEXT("UPoseSearchSchema::Finalize: couldn't Finalize '%s' because of Channel '%s'"), *GetNameSafe(this), *Label);
+				UE_LOGF(LogPoseSearch, Error, "UPoseSearchSchema::Finalize: couldn't Finalize '%ls' because of Channel '%ls'", *GetNameSafe(this), *Label);
 				#else // WITH_EDITOR
-				UE_LOG(LogPoseSearch, Error, TEXT("UPoseSearchSchema::Finalize: couldn't Finalize '%s' because of Channel '%s'"), *GetNameSafe(this), *GetNameSafe(ChannelPtr));
+				UE_LOGF(LogPoseSearch, Error, "UPoseSearchSchema::Finalize: couldn't Finalize '%ls' because of Channel '%ls'", *GetNameSafe(this), *GetNameSafe(ChannelPtr));
 				#endif // WITH_EDITOR
 
 				ResetFinalize();
diff --git a/Engine/Plugins/Animation/PoseSearch/Source/Runtime/Private/PoseSearchTrajectoryLibrary.cpp b/Engine/Plugins/Animation/PoseSearch/Source/Runtime/Private/PoseSearchTrajectoryLibrary.cpp
--- a/Engine/Plugins/Animation/PoseSearch/Source/Runtime/Private/PoseSearchTrajectoryLibrary.cpp
+++ b/Engine/Plugins/Animation/PoseSearch/Source/Runtime/Private/PoseSearchTrajectoryLibrary.cpp
@@ -1159,7 +1159,7 @@ void UPoseSearchTrajectoryLibrary::GetTransformTrajectoryVelocity(const FTransfo
 {
 	if (FMath::IsNearlyEqual(Time1, Time2))
 	{
-		UE_LOG(LogPoseSearch, Warning, TEXT("UPoseSearchTrajectoryLibrary::GetTrajectoryVelocity - Time1 is same as Time2. Invalid time horizon."));
+		UE_LOGF(LogPoseSearch, Warning, "UPoseSearchTrajectoryLibrary::GetTrajectoryVelocity - Time1 is same as Time2. Invalid time horizon.");
 		OutVelocity = FVector::ZeroVector;
 		return;
 	}
@@ -1175,7 +1175,7 @@ void UPoseSearchTrajectoryLibrary::GetTransformTrajectoryAngularVelocity(const F
 {
 	if (FMath::IsNearlyEqual(Time1, Time2))
 	{
-		UE_LOG(LogPoseSearch, Warning, TEXT("UPoseSearchTrajectoryLibrary::GetTrajectoryAngularVelocity - Time1 is same as Time2. Invalid time horizon."));
+		UE_LOGF(LogPoseSearch, Warning, "UPoseSearchTrajectoryLibrary::GetTrajectoryAngularVelocity - Time1 is same as Time2. Invalid time horizon.");
 		OutAngularVelocity = FVector::ZeroVector;
 		return;
 	}
diff --git a/Engine/Plugins/Animation/PoseSearch/Source/Runtime/Private/PoseSearchTrajectoryTypes.cpp b/Engine/Plugins/Animation/PoseSearch/Source/Runtime/Private/PoseSearchTrajectoryTypes.cpp
--- a/Engine/Plugins/Animation/PoseSearch/Source/Runtime/Private/PoseSearchTrajectoryTypes.cpp
+++ b/Engine/Plugins/Animation/PoseSearch/Source/Runtime/Private/PoseSearchTrajectoryTypes.cpp
@@ -221,7 +221,7 @@ FArchive& operator<<(FArchive& Ar, FPoseSearchQueryTrajectorySample& TrajectoryS
 FPoseSearchQueryTrajectory::FPoseSearchQueryTrajectory(const FTransformTrajectory& InTrajectory)
 {
 #if WITH_EDITOR
-	UE_LOG(LogPoseSearch, Warning, TEXT("Converting FTransformTrajectory to FPoseSearchQueryTrajectory (deprecated type) implicitly. Consider updating code to avoid this conversion."))
+	UE_LOGF(LogPoseSearch, Warning, "Converting FTransformTrajectory to FPoseSearchQueryTrajectory (deprecated type) implicitly. Consider updating code to avoid this conversion.")
 #endif
 	
 	Samples.SetNum(InTrajectory.Samples.Num(), EAllowShrinking::No);
diff --git a/Engine/Plugins/Animation/PoseSearch/Source/Runtime/Public/PoseSearch/PoseSearchAssetSamplerLibrary.h b/Engine/Plugins/Animation/PoseSearch/Source/Runtime/Public/PoseSearch/PoseSearchAssetSamplerLibrary.h
--- a/Engine/Plugins/Animation/PoseSearch/Source/Runtime/Public/PoseSearch/PoseSearchAssetSamplerLibrary.h
+++ b/Engine/Plugins/Animation/PoseSearch/Source/Runtime/Public/PoseSearch/PoseSearchAssetSamplerLibrary.h
@@ -115,7 +115,7 @@ public:
 	{
 		if (!ComponentSpacePose.GetPose().IsValid())
 		{
-			UE_LOG(LogPoseSearch, Error, TEXT("UPoseSearchAssetSamplerLibrary::GetTransform invalid Pose"));
+			UE_LOGF(LogPoseSearch, Error, "UPoseSearchAssetSamplerLibrary::GetTransform invalid Pose");
 			return FTransform::Identity;
 		}
 
@@ -123,14 +123,14 @@ public:
 		{
 			if (Space != EPoseSearchAssetSamplerSpace::World)
 			{
-				UE_LOG(LogPoseSearch, Error, TEXT("UPoseSearchAssetSamplerLibrary::GetTransform invalid Space %s to get the RootTransform. Expected space: %s"), *UEnum::GetDisplayValueAsText(Space).ToString(), *UEnum::GetDisplayValueAsText(EPoseSearchAssetSamplerSpace::World).ToString());
+				UE_LOGF(LogPoseSearch, Error, "UPoseSearchAssetSamplerLibrary::GetTransform invalid Space %ls to get the RootTransform. Expected space: %ls", *UEnum::GetDisplayValueAsText(Space).ToString(), *UEnum::GetDisplayValueAsText(EPoseSearchAssetSamplerSpace::World).ToString());
 			}
 			return RootTransform;
 		}
 	
 		if (!ComponentSpacePose.GetPose().IsValidIndex(CompactPoseBoneIndex))
 		{
-			UE_LOG(LogPoseSearch, Error, TEXT("UPoseSearchAssetSamplerLibrary::GetTransform invalid CompactPoseBoneIndex %d"), CompactPoseBoneIndex.GetInt());
+			UE_LOGF(LogPoseSearch, Error, "UPoseSearchAssetSamplerLibrary::GetTransform invalid CompactPoseBoneIndex %d", CompactPoseBoneIndex.GetInt());
 			return FTransform::Identity;
 		}
 	
diff --git a/Engine/Plugins/Animation/PoseSearch/Source/Runtime/Public/PoseSearch/VPTree.h b/Engine/Plugins/Animation/PoseSearch/Source/Runtime/Public/PoseSearch/VPTree.h
--- a/Engine/Plugins/Animation/PoseSearch/Source/Runtime/Public/PoseSearch/VPTree.h
+++ b/Engine/Plugins/Animation/PoseSearch/Source/Runtime/Public/PoseSearch/VPTree.h
@@ -82,7 +82,7 @@ struct FVPTreeResultSet
 #if !NO_LOGGING
 		if (InNumNeighbors > HeapDefaultMaxSize)
 		{
-			UE_LOG(LogPoseSearch, Warning, TEXT("FVPTreeResultSet - preallocated Heap data size 'HeapDefaultMaxSize' of %d is less than requested %d num neighbors. Performances will be negatively impacted"), HeapDefaultMaxSize, InNumNeighbors);
+			UE_LOGF(LogPoseSearch, Warning, "FVPTreeResultSet - preallocated Heap data size 'HeapDefaultMaxSize' of %d is less than requested %d num neighbors. Performances will be negatively impacted", HeapDefaultMaxSize, InNumNeighbors);
 		}
 #endif // !NO_LOGGING
     }
@@ -120,7 +120,7 @@ struct FVPTreeResultSet
 #if VALIDATE_FINDNEIGHBORS
 			if (BestIndex != INDEX_NONE && Removed.Index == BestIndex)
 			{
-				UE_LOG(LogPoseSearch, Error, TEXT("FVPTreeResultSet::AddPoint - we removed the best node from the ResultSet..."));
+				UE_LOGF(LogPoseSearch, Error, "FVPTreeResultSet::AddPoint - we removed the best node from the ResultSet...");
 			}
 #endif // VALIDATE_FINDNEIGHBORS
 
@@ -274,7 +274,7 @@ struct FVPTree
 
 					if (!bOtherNodeIndexContainsBestNodeIndex)
 					{
-						UE_LOG(LogPoseSearch, Error, TEXT("FVPTree::FindNeighbors - we couldn't find the best node..."));
+						UE_LOGF(LogPoseSearch, Error, "FVPTree::FindNeighbors - we couldn't find the best node...");
 					}
 				}
 			}
@@ -302,7 +302,7 @@ struct FVPTree
 #if VALIDATE_FINDNEIGHBORS
 			else if (NodeIndex == BestNodeIndex)
 			{
-				UE_LOG(LogPoseSearch, Error, TEXT("FVPTree::FindNeighbors - we had the best node in our hands and let it slipped away..."));
+				UE_LOGF(LogPoseSearch, Error, "FVPTree::FindNeighbors - we had the best node in our hands and let it slipped away...");
 			}
 
 			bool bIsBestNodeStillReachable = false;
@@ -313,7 +313,7 @@ struct FVPTree
 #if !NO_LOGGING
 				if (NodesToSearch.Num() > NodesToSearchDefaultMaxSize)
 				{
-					UE_LOG(LogPoseSearch, Warning, TEXT("FVPTree::FindNeighbors - requested more than preallocated NodesToSearch data size 'NodesToSearchDefaultMaxSize' (%d / %d)"), NodesToSearch.Num(), NodesToSearchDefaultMaxSize);
+					UE_LOGF(LogPoseSearch, Warning, "FVPTree::FindNeighbors - requested more than preallocated NodesToSearch data size 'NodesToSearchDefaultMaxSize' (%d / %d)", NodesToSearch.Num(), NodesToSearchDefaultMaxSize);
 				}
 #endif // !NO_LOGGING
 
@@ -334,7 +334,7 @@ struct FVPTree
 #if !NO_LOGGING
 				if (NodesToSearch.Num() > NodesToSearchDefaultMaxSize)
 				{
-					UE_LOG(LogPoseSearch, Warning, TEXT("FVPTree::FindNeighbors - requested more than preallocated NodesToSearch data size 'NodesToSearchDefaultMaxSize' (%d / %d)"), NodesToSearch.Num(), NodesToSearchDefaultMaxSize);
+					UE_LOGF(LogPoseSearch, Warning, "FVPTree::FindNeighbors - requested more than preallocated NodesToSearch data size 'NodesToSearchDefaultMaxSize' (%d / %d)", NodesToSearch.Num(), NodesToSearchDefaultMaxSize);
 				}
 #endif // !NO_LOGGING
 
@@ -352,7 +352,7 @@ struct FVPTree
 #if VALIDATE_FINDNEIGHBORS
 			if (bNodeIndexContainsBestNodeIndex != bIsBestNodeStillReachable)
 			{
-				UE_LOG(LogPoseSearch, Error, TEXT("FVPTree::FindNeighbors - we lost strack of the best node..."));
+				UE_LOGF(LogPoseSearch, Error, "FVPTree::FindNeighbors - we lost strack of the best node...");
 			}
 #endif // VALIDATE_FINDNEIGHBORS
 		}
