PoseSearch - fix for uninitialized bIsResultValid in GetMotionMatchingBlendSettings and OverrideMotionMatchingBlendSettings

#Synced-CL 51727087

diff --git a/Engine/Plugins/Animation/PoseSearch/Source/Runtime/Private/MotionMatchingAnimNodeLibrary.cpp b/Engine/Plugins/Animation/PoseSearch/Source/Runtime/Private/MotionMatchingAnimNodeLibrary.cpp
--- a/Engine/Plugins/Animation/PoseSearch/Source/Runtime/Private/MotionMatchingAnimNodeLibrary.cpp
+++ b/Engine/Plugins/Animation/PoseSearch/Source/Runtime/Private/MotionMatchingAnimNodeLibrary.cpp
@@ -80,15 +80,17 @@ void UMotionMatchingAnimNodeLibrary::GetMotionMatchingBlendSettings(const FMotio
 		BlendSettings.BlendProfile = MotionMatchingNodePtr->BlendProfile;
 		BlendSettings.BlendTime = MotionMatchingNodePtr->BlendTime;
 		BlendSettings.bUseInertialBlend = MotionMatchingNodePtr->bUseInertialBlend;
+		bIsResultValid = true;
 	}
 	else
 	{
 		BlendSettings = FMotionMatchingBlueprintBlendSettings();
+		bIsResultValid = false;
 		UE_LOGF(LogPoseSearch, Warning, "UMotionMatchingAnimNodeLibrary::GetMotionMatchingBlendSettings called on an invalid context or with an invalid type");
 	}
 }
 
-void UMotionMatchingAnimNodeLibrary::OverrideMotionMatchingBlendSettings(const FMotionMatchingAnimNodeReference& MotionMatchingNode, const FMotionMatchingBlueprintBlendSettings& BlendSettings, bool& bIsResultValid)
+void UMotionMatchingAnimNodeLibrary::OverrideMotionMatchingBlendSettings(const FMotionMatchingAnimNodeReference& MotionMatchingNode, const FMotionMatchingBlueprintBlendSettings& BlendSettings)
 {
 	if (FAnimNode_MotionMatching* MotionMatchingNodePtr = MotionMatchingNode.GetAnimNodePtr<FAnimNode_MotionMatching>())
 	{
diff --git a/Engine/Plugins/Animation/PoseSearch/Source/Runtime/Public/PoseSearch/MotionMatchingAnimNodeLibrary.h b/Engine/Plugins/Animation/PoseSearch/Source/Runtime/Public/PoseSearch/MotionMatchingAnimNodeLibrary.h
--- a/Engine/Plugins/Animation/PoseSearch/Source/Runtime/Public/PoseSearch/MotionMatchingAnimNodeLibrary.h
+++ b/Engine/Plugins/Animation/PoseSearch/Source/Runtime/Public/PoseSearch/MotionMatchingAnimNodeLibrary.h
@@ -83,7 +83,7 @@ public:
 
 	// Override current blend settings for motion matching. Note that any pinned parameters will stomp this override on the next update.
 	UFUNCTION(BlueprintCallable, Category = "Animation|MotionMatching", meta = (BlueprintThreadSafe, DisplayName = "Override Motion Matching Blend Settings"))
-	static UE_API void OverrideMotionMatchingBlendSettings(const FMotionMatchingAnimNodeReference& MotionMatchingNode, const FMotionMatchingBlueprintBlendSettings& BlendSettings, bool& bIsResultValid);
+	static UE_API void OverrideMotionMatchingBlendSettings(const FMotionMatchingAnimNodeReference& MotionMatchingNode, const FMotionMatchingBlueprintBlendSettings& BlendSettings);
 
 	/**
 	 * Set the database to search on the motion matching node. This overrides the Database property on the motion matching node.
