PoseSearch - FPoseSearchColumnRowReflectionCustomization to be able to multi edit pose search column databases from the details panel

#Synced-CL 51727087

diff --git a/Engine/Plugins/Animation/PoseSearch/Source/Editor/Private/PoseSearchCustomization.cpp b/Engine/Plugins/Animation/PoseSearch/Source/Editor/Private/PoseSearchCustomization.cpp
--- a/Engine/Plugins/Animation/PoseSearch/Source/Editor/Private/PoseSearchCustomization.cpp
+++ b/Engine/Plugins/Animation/PoseSearch/Source/Editor/Private/PoseSearchCustomization.cpp
@@ -4,6 +4,7 @@
 #include "Animation/AnimSequence.h"
 #include "Animation/BlendSpace.h"
 #include "AssetRegistry/AssetData.h"
+#include "Chooser.h"
 #include "DetailCategoryBuilder.h"
 #include "DetailLayoutBuilder.h"
 #include "DetailWidgetRow.h"
@@ -14,6 +15,7 @@
 #include "IDetailsView.h"
 #include "IPropertyUtilities.h"
 #include "ObjectEditorUtils.h"
+#include "PoseSearch/Chooser/PoseSearchChooserColumn.h"
 #include "PoseSearch/MultiAnimAsset.h"
 #include "PoseSearch/PoseSearchDatabase.h"
 #include "PropertyCustomizationHelpers.h"
@@ -24,13 +26,127 @@
 #include "Widgets/Images/SImage.h"
 #include "Widgets/Input/SButton.h"
 #include "Widgets/Input/SCheckBox.h"
+#include "Widgets/Input/SComboButton.h"
 #include "Widgets/Layout/SBox.h"
 #include "Widgets/Layout/SWrapBox.h"
 #include "Widgets/SBoxPanel.h"
 #include "Widgets/Text/STextBlock.h"
 
+#include "Framework/MultiBox/MultiBoxBuilder.h"
+
 #define LOCTEXT_NAMESPACE "PoseSearchCustomization"
 
+//////////////////////////////////////////////////////////////////////////
+// FPoseSearchColumnRowReflectionCustomization
+
+void FPoseSearchColumnRowReflectionCustomization::CustomizeHeader(TSharedRef<IPropertyHandle> InStructPropertyHandle, FDetailWidgetRow& HeaderRow, IPropertyTypeCustomizationUtils& StructCustomizationUtils)
+{
+	HeaderRow.NameContent()
+	[
+		InStructPropertyHandle->CreatePropertyNameWidget()
+	];
+}
+
+void FPoseSearchColumnRowReflectionCustomization::CustomizeChildren(TSharedRef<IPropertyHandle> InStructPropertyHandle, IDetailChildrenBuilder& StructBuilder, IPropertyTypeCustomizationUtils& StructCustomizationUtils)
+{
+	TSharedPtr<IPropertyHandle> DataHandle = InStructPropertyHandle->GetChildHandle(GET_MEMBER_NAME_CHECKED(FPoseSearchColumnRow, Data));
+	TSharedPtr<IPropertyHandle> DatabaseHandle = InStructPropertyHandle->GetChildHandle(GET_MEMBER_NAME_CHECKED(FPoseSearchColumnRow, Database));
+	TSharedPtr<IPropertyHandle> RootChooserHandle = InStructPropertyHandle->GetChildHandle(GET_MEMBER_NAME_CHECKED(FPoseSearchColumnRowReflection, RootChooser));
+
+	if (DataHandle.IsValid())
+	{
+		StructBuilder.AddProperty(DataHandle.ToSharedRef());
+	}
+
+	// Custom Database row: combobox restricted to UPoseSearchDatabase instances
+	// found in RootChooser->NestedObjects, matching the column-widget behaviour.
+	if (DatabaseHandle.IsValid())
+	{
+		StructBuilder.AddProperty(DatabaseHandle.ToSharedRef()).CustomWidget()
+			.NameContent()
+			[
+				DatabaseHandle->CreatePropertyNameWidget()
+			]
+			.ValueContent()
+			[
+				SNew(SComboButton)
+				.ContentPadding(FMargin(2.f, 0.f))
+				.ButtonContent()
+				[
+					SNew(STextBlock)
+					.Font(IDetailLayoutBuilder::GetDetailFont())
+					.Text_Lambda([DatabaseHandle]()
+					{
+						UObject* DatabaseObj = nullptr;
+						const FPropertyAccess::Result Result = DatabaseHandle->GetValue(DatabaseObj);
+						if (Result == FPropertyAccess::MultipleValues)
+						{
+							return LOCTEXT("MultipleValuesText", "Multiple Values");
+						}
+
+						if (Result == FPropertyAccess::Success)
+						{
+							if (const UPoseSearchDatabase* Database = Cast<UPoseSearchDatabase>(DatabaseObj))
+							{
+								return FText::FromName(Database->GetFName());
+							}
+						}
+
+						return LOCTEXT("None", "None");
+					})
+				]
+				.OnGetMenuContent_Lambda([DatabaseHandle, RootChooserHandle]()
+				{
+					FMenuBuilder MenuBuilder(true, nullptr);
+					MenuBuilder.BeginSection(NAME_None, LOCTEXT("DatabaseSelection", "Database"));
+
+					// "None" entry to clear the selection
+					MenuBuilder.AddMenuEntry(
+						LOCTEXT("None", "None"),
+						LOCTEXT("NoneDatabase_Tooltip", "Remove the database assignment for this row"),
+						FSlateIcon(),
+						FUIAction(FExecuteAction::CreateLambda([DatabaseHandle]()
+						{
+							const FScopedTransaction Transaction(LOCTEXT("ClearDatabase", "Clear Database"));
+							DatabaseHandle->SetValue((UObject*)nullptr);
+						}))
+					);
+
+					// One entry per UPoseSearchDatabase in RootChooser->NestedObjects
+					if (RootChooserHandle.IsValid())
+					{
+						UObject* ChooserObj = nullptr;
+						if (RootChooserHandle->GetValue(ChooserObj) == FPropertyAccess::Success)
+						{
+							if (const UChooserTable* RootChooser = Cast<UChooserTable>(ChooserObj))
+							{
+								for (UObject* NestedObject : RootChooser->NestedObjects)
+								{
+									if (UPoseSearchDatabase* Database = Cast<UPoseSearchDatabase>(NestedObject))
+									{
+										MenuBuilder.AddMenuEntry(
+											FText::FromName(Database->GetFName()),
+											FText::GetEmpty(),
+											FSlateIcon(),
+											FUIAction(FExecuteAction::CreateLambda([DatabaseHandle, Database]()
+											{
+												const FScopedTransaction Transaction(LOCTEXT("SetDatabase", "Set Database"));
+												DatabaseHandle->SetValue(Database);
+											}))
+										);
+									}
+								}
+							}
+						}
+					}
+
+					MenuBuilder.EndSection();
+					return MenuBuilder.MakeWidget();
+				})
+			];
+	}
+}
+
 void FPoseSearchDatabaseAnimationAssetCustomization::CustomizeHeader(TSharedRef<IPropertyHandle> InStructPropertyHandle, FDetailWidgetRow& HeaderRow, IPropertyTypeCustomizationUtils& StructCustomizationUtils)
 {
 	FText SequenceNameText;
diff --git a/Engine/Plugins/Animation/PoseSearch/Source/Editor/Private/PoseSearchEditor.cpp b/Engine/Plugins/Animation/PoseSearch/Source/Editor/Private/PoseSearchEditor.cpp
--- a/Engine/Plugins/Animation/PoseSearch/Source/Editor/Private/PoseSearchEditor.cpp
+++ b/Engine/Plugins/Animation/PoseSearch/Source/Editor/Private/PoseSearchEditor.cpp
@@ -138,7 +138,8 @@ void FEditorModule::RegisterCustomPropertyTypeLayout(FName PropertyTypeName, FOn
 
 void FEditorModule::RegisterPropertyTypeCustomizations()
 {
-	RegisterCustomPropertyTypeLayout("PoseSearchDatabaseAnimationAsset", FOnGetPropertyTypeCustomizationInstance::CreateStatic(&FPoseSearchDatabaseAnimationAssetCustomization::MakeInstance));
+	RegisterCustomPropertyTypeLayout("PoseSearchDatabaseAnimationAsset",  FOnGetPropertyTypeCustomizationInstance::CreateStatic(&FPoseSearchDatabaseAnimationAssetCustomization::MakeInstance));
+	RegisterCustomPropertyTypeLayout("PoseSearchColumnRowReflection",     FOnGetPropertyTypeCustomizationInstance::CreateStatic(&FPoseSearchColumnRowReflectionCustomization::MakeInstance));
 }
 
 void FEditorModule::UnregisterCustomizations()
diff --git a/Engine/Plugins/Animation/PoseSearch/Source/Editor/Public/PoseSearchCustomization.h b/Engine/Plugins/Animation/PoseSearch/Source/Editor/Public/PoseSearchCustomization.h
--- a/Engine/Plugins/Animation/PoseSearch/Source/Editor/Public/PoseSearchCustomization.h
+++ b/Engine/Plugins/Animation/PoseSearch/Source/Editor/Public/PoseSearchCustomization.h
@@ -13,6 +13,21 @@
 class IPropertyHandle;
 class IDetailLayoutBuilder;
 
+// Customises FPoseSearchColumnRowReflection so the Database picker only offers
+// UPoseSearchDatabase instances embedded in the owning chooser (NestedObjects),
+// rather than the full project asset browser.
+class FPoseSearchColumnRowReflectionCustomization : public IPropertyTypeCustomization
+{
+public:
+	static TSharedRef<IPropertyTypeCustomization> MakeInstance()
+	{
+		return MakeShareable(new FPoseSearchColumnRowReflectionCustomization);
+	}
+
+	virtual void CustomizeHeader(TSharedRef<IPropertyHandle> InStructPropertyHandle, FDetailWidgetRow& HeaderRow, IPropertyTypeCustomizationUtils& StructCustomizationUtils) override;
+	virtual void CustomizeChildren(TSharedRef<IPropertyHandle> InStructPropertyHandle, IDetailChildrenBuilder& StructBuilder, IPropertyTypeCustomizationUtils& StructCustomizationUtils) override;
+};
+
 class FPoseSearchDatabaseAnimationAssetCustomization : public IPropertyTypeCustomization
 {
 public:
diff --git a/Engine/Plugins/Animation/PoseSearch/Source/Runtime/Private/Chooser/PoseSearchChooserColumn.cpp b/Engine/Plugins/Animation/PoseSearch/Source/Runtime/Private/Chooser/PoseSearchChooserColumn.cpp
--- a/Engine/Plugins/Animation/PoseSearch/Source/Runtime/Private/Chooser/PoseSearchChooserColumn.cpp
+++ b/Engine/Plugins/Animation/PoseSearch/Source/Runtime/Private/Chooser/PoseSearchChooserColumn.cpp
@@ -1127,15 +1127,8 @@ void FPoseSearchColumn::AddToDetails(FInstancedPropertyBag& PropertyBag, int32 C
 	}
 
 	FPoseSearchColumnRowReflection RowReflection;
-	RowReflection.bDisableReselection = RowValue.Data.bDisableReselection;
-	RowReflection.MirrorOption = RowValue.Data.MirrorOption;
-	RowReflection.bUseSingleSample = RowValue.Data.bUseSingleSample;
-	RowReflection.bUseGridForSampling = RowValue.Data.bUseGridForSampling;
-	RowReflection.NumberOfHorizontalSamples = RowValue.Data.NumberOfHorizontalSamples;
-	RowReflection.NumberOfVerticalSamples = RowValue.Data.NumberOfVerticalSamples;
-	RowReflection.BlendParamX = RowValue.Data.BlendParamX;
-	RowReflection.BlendParamY = RowValue.Data.BlendParamY;
-	RowReflection.Database = RowValue.Database;
+	static_cast<FPoseSearchColumnRow&>(RowReflection) = RowValue;
+	RowReflection.RootChooser = Chooser->GetRootChooser();
 
 	ensure(!RowReflection.Database || RowReflection.Database->GetOuter() == Chooser->GetRootChooser());
 
@@ -1156,34 +1149,18 @@ void FPoseSearchColumn::SetFromDetails(FInstancedPropertyBag& PropertyBag, int32
 	if (FStructView* StructView = Result.TryGetValue())
 	{
 		const FPoseSearchColumnRowReflection& RowReflection = StructView->Get<FPoseSearchColumnRowReflection>();
+		const FPoseSearchColumnRow& RowReflectionBase = static_cast<const FPoseSearchColumnRow&>(RowReflection);
 		FPoseSearchColumnRow& RowValue = EditRowValue(RowIndex);
 
-		ensure(!RowReflection.Database || RowReflection.Database->GetOuter() == Chooser->GetRootChooser());
+		ensure(!RowReflectionBase.Database || RowReflectionBase.Database->GetOuter() == Chooser->GetRootChooser());
 
 		UObject* AnimAsset = GetReferencedObject(GetChooser(), RowIndex);
-		if (RowValue.Data.AnimAsset != AnimAsset ||
-			RowValue.Data.bDisableReselection != RowReflection.bDisableReselection ||
-			RowValue.Data.MirrorOption != RowReflection.MirrorOption ||
-			RowValue.Data.bUseSingleSample != RowReflection.bUseSingleSample ||
-			RowValue.Data.bUseGridForSampling != RowReflection.bUseGridForSampling ||
-			RowValue.Data.NumberOfHorizontalSamples != RowReflection.NumberOfHorizontalSamples ||
-			RowValue.Data.NumberOfVerticalSamples != RowReflection.NumberOfVerticalSamples ||
-			RowValue.Data.BlendParamX != RowReflection.BlendParamX ||
-			RowValue.Data.BlendParamY != RowReflection.BlendParamY ||
-			RowValue.Database != RowReflection.Database)
+		if (RowValue != RowReflectionBase)
 		{
 			Chooser->Modify();
 
+			RowValue = RowReflectionBase;
 			RowValue.Data.AnimAsset = GetReferencedObject(GetChooser(), RowIndex);
-			RowValue.Data.bDisableReselection = RowReflection.bDisableReselection;
-			RowValue.Data.MirrorOption = RowReflection.MirrorOption;
-			RowValue.Data.bUseSingleSample = RowReflection.bUseSingleSample;
-			RowValue.Data.bUseGridForSampling = RowReflection.bUseGridForSampling;
-			RowValue.Data.NumberOfHorizontalSamples = RowReflection.NumberOfHorizontalSamples;
-			RowValue.Data.NumberOfVerticalSamples = RowReflection.NumberOfVerticalSamples;
-			RowValue.Data.BlendParamX = RowReflection.BlendParamX;
-			RowValue.Data.BlendParamY = RowReflection.BlendParamY;
-			RowValue.Database = RowReflection.Database;
 
 			FMapping(Chooser->GetRootChooser()).ApplyMappingToAllColumns();
 			Chooser->PostEditChange();
diff --git a/Engine/Plugins/Animation/PoseSearch/Source/Runtime/Public/PoseSearch/Chooser/PoseSearchChooserColumn.h b/Engine/Plugins/Animation/PoseSearch/Source/Runtime/Public/PoseSearch/Chooser/PoseSearchChooserColumn.h
--- a/Engine/Plugins/Animation/PoseSearch/Source/Runtime/Public/PoseSearch/Chooser/PoseSearchChooserColumn.h
+++ b/Engine/Plugins/Animation/PoseSearch/Source/Runtime/Public/PoseSearch/Chooser/PoseSearchChooserColumn.h
@@ -94,51 +94,26 @@ struct FPoseSearchColumnRow
 
 	UPROPERTY(EditAnywhere, Category = "Data")
 	TObjectPtr<UPoseSearchDatabase> Database;
+
+	friend bool operator==(const FPoseSearchColumnRow& A, const FPoseSearchColumnRow& B)
+	{
+		return A.Data == B.Data && A.Database == B.Database;
+	}
+};
+template<> struct TStructOpsTypeTraits<FPoseSearchColumnRow> : public TStructOpsTypeTraitsBase2<FPoseSearchColumnRow>
+{
+	enum { WithIdenticalViaEquality = true };
 };
 
 USTRUCT(Experimental)
-struct FPoseSearchColumnRowReflection
+struct FPoseSearchColumnRowReflection : public FPoseSearchColumnRow
 {
 	GENERATED_BODY()
 
 #if WITH_EDITORONLY_DATA
-	// if bDisableReselection is true, poses from the same asset cannot be reselected. Useful to avoid jumping on frames on the same looping animations
-	UPROPERTY(EditAnywhere, Category = "Data")
-	bool bDisableReselection = false;
-
-	// This allows users to set if this animation is original only (no mirrored data), original and mirrored, or only the mirrored version of this animation.
-	// It requires the mirror table to be set up in the database Schema.
-	UPROPERTY(EditAnywhere, Category = "Data")
-	EPoseSearchMirrorOption MirrorOption = EPoseSearchMirrorOption::UnmirroredOnly;
 
-	// If true this BlendSpace will output a single segment in the database.
 	UPROPERTY(EditAnywhere, Category = "Data")
-	bool bUseSingleSample = false;
-
-	// When turned on, this will use the set grid samples of the blend space asset for sampling. This will override the Number of Horizontal/Vertical Samples.
-	UPROPERTY(EditAnywhere, Category = "Data", meta = (EditCondition = "!bUseSingleSample", EditConditionHides))
-	bool bUseGridForSampling = false;
-
-	// Sets the number of horizontal samples in the blend space to pull the animation data coverage from. The larger the samples the more the data, but also the more memory and performance it takes.
-	UPROPERTY(EditAnywhere, Category = "Data", meta = (EditCondition = "!bUseSingleSample && !bUseGridForSampling", EditConditionHides, ClampMin = "1", UIMin = "1", UIMax = "25"))
-	int32 NumberOfHorizontalSamples = 9;
-	
-	// Sets the number of vertical samples in the blend space to pull the animation data coverage from.The larger the samples the more the data, but also the more memory and performance it takes.
-	UPROPERTY(EditAnywhere, Category = "Data", meta = (EditCondition = "!bUseSingleSample && !bUseGridForSampling", EditConditionHides, ClampMin = "1", UIMin = "1", UIMax = "25"))
-	int32 NumberOfVerticalSamples = 2;
-
-	// BlendParams used to sample this BlendSpace
-	UPROPERTY(EditAnywhere, Category = "Data", meta = (EditCondition = "bUseSingleSample", EditConditionHides))
-	float BlendParamX = 0.f;
-
-	// BlendParams used to sample this BlendSpace
-	UPROPERTY(EditAnywhere, Category = "Data", meta = (EditCondition = "bUseSingleSample", EditConditionHides))
-	float BlendParamY = 0.f;
-
-	// @todo: make it EditAnywhere, but populate the asset list filter to look only for the 
-	//		  databases owned by the FPoseSearchColumn root chooser table owner
-	UPROPERTY(VisibleAnywhere, Category = "Data")
-	TObjectPtr<UPoseSearchDatabase> Database;
+	TObjectPtr<UChooserTable> RootChooser;
 
 #endif // WITH_EDITORONLY_DATA
 };
