PoseSearch - Strip out disabled Pose Search Database animations during cook

#Synced-CL 51737212

diff --git a/Engine/Plugins/Animation/PoseSearch/Source/Runtime/Private/PoseSearchDatabase.cpp b/Engine/Plugins/Animation/PoseSearch/Source/Runtime/Private/PoseSearchDatabase.cpp
--- a/Engine/Plugins/Animation/PoseSearch/Source/Runtime/Private/PoseSearchDatabase.cpp
+++ b/Engine/Plugins/Animation/PoseSearch/Source/Runtime/Private/PoseSearchDatabase.cpp
@@ -11,6 +11,7 @@
 #include "PoseSearch/MultiAnimAsset.h"
 #include "PoseSearch/PoseSearchAnimNotifies.h"
 #include "PoseSearch/PoseSearchContext.h"
+#include "PoseSearch/PoseSearchCustomVersion.h"
 #include "PoseSearch/PoseSearchDefines.h"
 #include "PoseSearch/PoseSearchDerivedData.h"
 #include "PoseSearch/PoseSearchFeatureChannel_Group.h"
@@ -1384,11 +1385,39 @@ void UPoseSearchDatabase::PostSaveRoot(FObjectPostSaveRootContext ObjectSaveCont
 
 void UPoseSearchDatabase::Serialize(FArchive& Ar)
 {
-	Super::Serialize(Ar);
+	Ar.UsingCustomVersion(FPoseSearchCustomVersion::GUID);
+
+	const bool bSavingCooked = Ar.IsSaving() && Ar.IsCooking();
+#if WITH_EDITORONLY_DATA
+	// When cooking, strip AnimAsset from disabled entries so they aren't loaded at runtime
+	if (bSavingCooked)
+	{
+		TArray<TObjectPtr<UObject>> SavedAnimAssets;
+		SavedAnimAssets.Reserve(DatabaseAnimationAssets.Num());
+		for (FPoseSearchDatabaseAnimationAsset& AnimationAsset : DatabaseAnimationAssets)
+		{
+			SavedAnimAssets.Add(AnimationAsset.AnimAsset);
+			if (!AnimationAsset.bEnabled)
+			{
+				AnimationAsset.AnimAsset = nullptr;
+			}
+		}
+
+		Super::Serialize(Ar);
+
+		for (int32 i = 0; i < DatabaseAnimationAssets.Num(); ++i)
+		{
+			DatabaseAnimationAssets[i].AnimAsset = SavedAnimAssets[i];
+		}
+	}
+	else
+#endif // WITH_EDITORONLY_DATA
+	{
+		Super::Serialize(Ar);
+	}
 
 	if (!IsTemplate())
 	{
-		const bool bSavingCooked = Ar.IsSaving() && Ar.IsCooking();
 		const bool bLoadingCooked = Ar.IsLoadingFromCookedPackage();
 		if (bSavingCooked || bLoadingCooked)
 		{
diff --git a/Engine/Plugins/Animation/PoseSearch/Source/Runtime/Public/PoseSearch/PoseSearchCustomVersion.h b/Engine/Plugins/Animation/PoseSearch/Source/Runtime/Public/PoseSearch/PoseSearchCustomVersion.h
--- a/Engine/Plugins/Animation/PoseSearch/Source/Runtime/Public/PoseSearch/PoseSearchCustomVersion.h
+++ b/Engine/Plugins/Animation/PoseSearch/Source/Runtime/Public/PoseSearch/PoseSearchCustomVersion.h
@@ -21,6 +21,9 @@ struct FPoseSearchCustomVersion
 		// Added InterruptMode to posesearch debugger. This version matches the RouteId_MotionMatchingState3 event.
 		AddedInterruptModeToDebugger,
 
+		// Strip out disabled Pose Search Database animations during cook
+		StripOutDisabledDatabaseAnimationsDuringCook,
+
 		// -----<new versions can be added above this line>-------------------------------------------------
 		VersionPlusOne,
 		LatestVersion = VersionPlusOne - 1
