Pose Search: Databases created for Chooser Columns now default to BruteForce
#Synced-CL 51770067

diff --git a/Engine/Plugins/Animation/PoseSearch/Source/Editor/Private/Chooser/PoseSearchColumnEditor.cpp b/Engine/Plugins/Animation/PoseSearch/Source/Editor/Private/Chooser/PoseSearchColumnEditor.cpp
--- a/Engine/Plugins/Animation/PoseSearch/Source/Editor/Private/Chooser/PoseSearchColumnEditor.cpp
+++ b/Engine/Plugins/Animation/PoseSearch/Source/Editor/Private/Chooser/PoseSearchColumnEditor.cpp
@@ -145,6 +145,8 @@ TSharedRef<SWidget> CreatePoseSearchColumnWidget(UChooserTable* Chooser, FChoose
 
 												UPoseSearchDatabase* NewDatabase = NewObject<UPoseSearchDatabase>(RootChooser, UPoseSearchDatabase::StaticClass(), NewDatabaseName, RF_Transactional);
 												NewDatabase->Schema = TargetSchema;
+												// KDTree frequently fails in Chooser columns so use Brute Force instead.
+												NewDatabase->PoseSearchMode = EPoseSearchMode::BruteForce;
 
 												const FScopedTransaction Transaction(LOCTEXT("Assign New Database", "Assign New Database"));
 												RootChooser->Modify();
