PoseSearch - adding GVarMotionMatchUseRawAnimationData to use uncompressed animation data for database indexing

#Synced-CL 51737212

diff --git a/Engine/Plugins/Animation/PoseSearch/Source/Runtime/Private/PoseSearchAssetSamplerLibrary.cpp b/Engine/Plugins/Animation/PoseSearch/Source/Runtime/Private/PoseSearchAssetSamplerLibrary.cpp
--- a/Engine/Plugins/Animation/PoseSearch/Source/Runtime/Private/PoseSearchAssetSamplerLibrary.cpp
+++ b/Engine/Plugins/Animation/PoseSearch/Source/Runtime/Private/PoseSearchAssetSamplerLibrary.cpp
@@ -52,11 +52,16 @@ FPoseSearchAssetSamplerPose UPoseSearchAssetSamplerLibrary::SamplePose(const FBo
 
 	FMemMark Mark(FMemStack::Get());
 
-	bool bPreProcessRootTransform = true;
-	const FAnimationAssetSampler Sampler(Input.Animation, Input.RootTransformOrigin, Input.BlendParameters, Input.RootTransformSamplingRate, bPreProcessRootTransform, false);
+	// @todo: should we expose those parameters?
+	constexpr bool bPreProcessRootTransform = true;
+	constexpr bool bEnforceCompressedDataSampling = false;
+	constexpr bool bUseRAWData = false;
+
+	const FAnimationAssetSampler Sampler(Input.Animation, Input.RootTransformOrigin, Input.BlendParameters, FAnimationAssetSampler::DefaultRootTransformSamplingRate, bPreProcessRootTransform, bEnforceCompressedDataSampling);
 
 	FPoseSearchAssetSamplerPose AssetSamplerPose;
 	AssetSamplerPose.BoneContainer = BoneContainer;
+	AssetSamplerPose.BoneContainer.SetUseRAWData(bUseRAWData);
 
 	FBlendedCurve Curve;
 	FCompactPose Pose;
diff --git a/Engine/Plugins/Animation/PoseSearch/Source/Runtime/Private/PoseSearchDerivedData.cpp b/Engine/Plugins/Animation/PoseSearch/Source/Runtime/Private/PoseSearchDerivedData.cpp
--- a/Engine/Plugins/Animation/PoseSearch/Source/Runtime/Private/PoseSearchDerivedData.cpp
+++ b/Engine/Plugins/Animation/PoseSearch/Source/Runtime/Private/PoseSearchDerivedData.cpp
@@ -97,6 +97,12 @@ static FAutoConsoleVariableRef CVarMotionMatchReindexCancelledDatabases(TEXT("a.
 
 #endif // ENABLE_ANIM_DEBUG
 
+static bool GVarMotionMatchUseRawAnimationData = false;
+static FAutoConsoleVariableRef CVarMotionMatchUseRawAnimationData(TEXT("a.MotionMatch.UseRawAnimationData"), GVarMotionMatchUseRawAnimationData, TEXT("Use Raw Animation Data"), FConsoleVariableDelegate::CreateLambda([](IConsoleVariable* InVariable)
+	{
+		FAsyncPoseSearchDatabasesManagement::Clear();
+	}));
+
 static bool GVarMotionMatchReindexAllReferencedDatabases = true;
 static FAutoConsoleVariableRef CVarMotionMatchReindexAllReferencedDatabases(TEXT("a.MotionMatch.ReindexAllReferencedDatabases"), GVarMotionMatchReindexAllReferencedDatabases, TEXT("Reindex All Referenced Databases"));
 
@@ -1098,6 +1104,7 @@ static void InitRoledBoneContainers(TMap<FRole, FBoneContainer>& RoledBoneContai
 		FAnimationRuntime::EnsureParentsPresent(BoneIndicesWithParents, RoledSkeleton.Skeleton->GetReferenceSkeleton());
 
 		RoledBoneContainer.InitializeTo(BoneIndicesWithParents, CurveFilterSettings, *RoledSkeleton.Skeleton);
+		RoledBoneContainer.SetUseRAWData(GVarMotionMatchUseRawAnimationData);
 	}
 }
 
@@ -1134,7 +1141,7 @@ static bool IndexDatabase(FSearchIndexBase& SearchIndexBase, const UPoseSearchDa
 			if (!SamplerMap.Contains(SamplerMapKey))
 			{
 				SamplerMap.Add(SamplerMapKey, Samplers.Num());
-				Samplers.Emplace(AnimationAsset, RootTransformOrigin, BlendParameters, FAnimationAssetSampler::DefaultRootTransformSamplingRate, false);
+				Samplers.Emplace(AnimationAsset, RootTransformOrigin, BlendParameters, FAnimationAssetSampler::DefaultRootTransformSamplingRate, false, !GVarMotionMatchUseRawAnimationData);
 			}
 		}
 	}
@@ -1796,7 +1803,7 @@ void FPoseSearchDatabaseAsyncCacheTask::StartNewRequestIfNeeded(bool bPerformCon
 	Owner.Cancel();
 
 	// composing the key
-	FKeyBuilder KeyBuilder(Database.Get(), true, bPerformConditionalPostLoadIfRequired, &PartialKeyHashes);
+	FKeyBuilder KeyBuilder(Database.Get(), true, bPerformConditionalPostLoadIfRequired, &PartialKeyHashes, !GVarMotionMatchUseRawAnimationData);
 	if (KeyBuilder.AnyAssetNotFullyLoaded())
 	{
 		DerivedDataKey = FIoHash::Zero;
@@ -2646,7 +2653,8 @@ void FAsyncPoseSearchDatabasesManagement::Tick(float DeltaTime)
 				FMemMark Mark(FMemStack::Get());
 
 				FTestSample& TestSample = TestSamples[TestSampleIndex];
-				const FAnimationAssetSampler AssetSampler(TestSample.AnimationAsset);
+
+				const FAnimationAssetSampler AssetSampler(TestSample.AnimationAsset, FTransform::Identity, FVector::ZeroVector, FAnimationAssetSampler::DefaultRootTransformSamplingRate, true, !GVarMotionMatchUseRawAnimationData);
 				const float SampleTime = AssetSampler.GetPlayLength() * TestSample.SampleNormalizedTime;
 
 				FCompactPose Pose;
@@ -2819,12 +2827,12 @@ void FAsyncPoseSearchDatabasesManagement::Tick(float DeltaTime)
 		{
 			if (const UPoseSearchDatabase* Database = Tasks[TaskIndex]->GetDatabase())
 			{
-				FKeyBuilder KeyBuilder(Database, false, false);
+				FKeyBuilder KeyBuilder(Database, false, false, nullptr, !GVarMotionMatchUseRawAnimationData);
 				const FIoHash IoHash = KeyBuilder.Finalize();
 
 				for (int32 IterationIndex = 0; IterationIndex < GVarMotionMatchTestNumIterations; ++IterationIndex)
 				{
-					FKeyBuilder TestKeyBuilder(Database, false, false, &PartialKeyHashes);
+					FKeyBuilder TestKeyBuilder(Database, false, false, &PartialKeyHashes, !GVarMotionMatchUseRawAnimationData);
 					const FIoHash TestIoHash = TestKeyBuilder.Finalize();
 
 					if (!KeyBuilder.ValidateAgainst(TestKeyBuilder))
diff --git a/Engine/Plugins/Animation/PoseSearch/Source/Runtime/Private/PoseSearchDerivedDataKey.cpp b/Engine/Plugins/Animation/PoseSearch/Source/Runtime/Private/PoseSearchDerivedDataKey.cpp
--- a/Engine/Plugins/Animation/PoseSearch/Source/Runtime/Private/PoseSearchDerivedDataKey.cpp
+++ b/Engine/Plugins/Animation/PoseSearch/Source/Runtime/Private/PoseSearchDerivedDataKey.cpp
@@ -52,7 +52,7 @@ FKeyBuilder::FKeyBuilder()
 	SetIsSaving(true);
 }
 
-FKeyBuilder::FKeyBuilder(const UObject* Object, bool bUseDataVer, bool bPerformConditionalPostLoadIfRequired, FPartialKeyHashes* InPartialKeyHashes)
+FKeyBuilder::FKeyBuilder(const UObject* Object, bool bUseDataVer, bool bPerformConditionalPostLoadIfRequired, FPartialKeyHashes* InPartialKeyHashes, bool bInPerformSequenceResidencyRequests)
 : FKeyBuilder()
 {
 #if ENABLE_KEY_PROFILING
@@ -70,6 +70,8 @@ FKeyBuilder::FKeyBuilder(const UObject* Object, bool bUseDataVer, bool bPerformC
 	ObjectBeingSerializedDependencies.Reserve(256);
 
 	bPerformConditionalPostLoad = bPerformConditionalPostLoadIfRequired;
+	bPerformSequenceResidencyRequests = bInPerformSequenceResidencyRequests;
+
 	ExternalPartialKeyHashes = InPartialKeyHashes;
 
 	UObject* NonPackageOuter = GetNonPackageOuter(KeyOwner);
@@ -365,39 +367,41 @@ bool FKeyBuilder::AnyAssetNotFullyLoaded() const
 
 bool FKeyBuilder::AnyAssetNotReady() const
 {
-	TArray<UAnimSequence*, TInlineAllocator<64>> SequencesToWaitFor;	
-	const ITargetPlatform* TargetPlatform = nullptr;
-	for (const UObject* Dependency : Dependencies)
+	if (bPerformSequenceResidencyRequests)
 	{
-		if (Dependency->IsA<UAnimSequence>())
+		TArray<UAnimSequence*, TInlineAllocator<64>> SequencesToWaitFor;
+		const ITargetPlatform* TargetPlatform = nullptr;
+		for (const UObject* Dependency : Dependencies)
 		{
-			UAnimSequence* AnimSequence = Cast<UAnimSequence>(const_cast<UObject*>(Dependency));
-			check(AnimSequence);
-
-			// initializing TargetPlatform lazily when the first UAnimSequence has been found
-			if (SequencesToWaitFor.IsEmpty())
+			if (Dependency->IsA<UAnimSequence>())
 			{
-				TargetPlatform = GetTargetPlatformManager()->GetRunningTargetPlatform();
-			}
+				UAnimSequence* AnimSequence = Cast<UAnimSequence>(const_cast<UObject*>(Dependency));
+				check(AnimSequence);
 
-			AnimSequence->RequestResidency(TargetPlatform, GetTypeHash(KeyOwner));
-			SequencesToWaitFor.Add(AnimSequence);
-		}
-	}
+				// initializing TargetPlatform lazily when the first UAnimSequence has been found
+				if (SequencesToWaitFor.IsEmpty())
+				{
+					TargetPlatform = GetTargetPlatformManager()->GetRunningTargetPlatform();
+				}
 
-	if (!SequencesToWaitFor.IsEmpty())
-	{
-		Anim::IAnimSequenceCompilingManager::FinishCompilation(MakeArrayView(SequencesToWaitFor));
+				AnimSequence->RequestResidency(TargetPlatform, GetTypeHash(KeyOwner));
+				SequencesToWaitFor.Add(AnimSequence);
+			}
+		}
 
-		for (UAnimSequence* AnimSequence : SequencesToWaitFor)
+		if (!SequencesToWaitFor.IsEmpty())
 		{
-			if (!AnimSequence->HasCompressedDataForPlatform(TargetPlatform))
+			Anim::IAnimSequenceCompilingManager::FinishCompilation(MakeArrayView(SequencesToWaitFor));
+
+			for (UAnimSequence* AnimSequence : SequencesToWaitFor)
 			{
-				return true;
+				if (!AnimSequence->HasCompressedDataForPlatform(TargetPlatform))
+				{
+					return true;
+				}
 			}
 		}
 	}
-	
 	return false;
 }
 
@@ -446,6 +450,11 @@ FIoHash FKeyBuilder::Finalize()
 	*this << AnimationCompressionVersionString;
 	*this << NonConstDatabaseIndexDerivedDataCacheKeyVersion;
 
+	if (bPerformSequenceResidencyRequests)
+	{
+		*this << bPerformSequenceResidencyRequests;
+	}
+		
 	// iterating over all ObjectKeysToVisit to finalize the composition the final hash while continuing adding their unvisited dependencies
 	while (!ObjectKeysToVisit.IsEmpty())
 	{
diff --git a/Engine/Plugins/Animation/PoseSearch/Source/Runtime/Public/PoseSearch/PoseSearchAssetSamplerLibrary.h b/Engine/Plugins/Animation/PoseSearch/Source/Runtime/Public/PoseSearch/PoseSearchAssetSamplerLibrary.h
--- a/Engine/Plugins/Animation/PoseSearch/Source/Runtime/Public/PoseSearch/PoseSearchAssetSamplerLibrary.h
+++ b/Engine/Plugins/Animation/PoseSearch/Source/Runtime/Public/PoseSearch/PoseSearchAssetSamplerLibrary.h
@@ -38,10 +38,6 @@ struct FPoseSearchAssetSamplerInput
 	// blend parameters if Animation is a blend space
 	UPROPERTY(EditAnywhere, BlueprintReadWrite, Category=State)
 	FVector BlendParameters = FVector::ZeroVector;
-	
-	// frequency of sampling while sampling the root transform of blend spaces
-	UPROPERTY(EditAnywhere, BlueprintReadWrite, Category=State)
-	int32 RootTransformSamplingRate = UE::PoseSearch::FAnimationAssetSampler::DefaultRootTransformSamplingRate;
 };
 
 USTRUCT(Experimental, BlueprintType, Category="Animation|Pose Search")
diff --git a/Engine/Plugins/Animation/PoseSearch/Source/Runtime/Public/PoseSearch/PoseSearchDerivedDataKey.h b/Engine/Plugins/Animation/PoseSearch/Source/Runtime/Public/PoseSearch/PoseSearchDerivedDataKey.h
--- a/Engine/Plugins/Animation/PoseSearch/Source/Runtime/Public/PoseSearch/PoseSearchDerivedDataKey.h
+++ b/Engine/Plugins/Animation/PoseSearch/Source/Runtime/Public/PoseSearch/PoseSearchDerivedDataKey.h
@@ -155,7 +155,7 @@ public:
 	POSESEARCH_API FKeyBuilder();
 	
 	// Experimental, this feature might be removed without warning, not for production use
-	POSESEARCH_API FKeyBuilder(const UObject* Object, bool bUseDataVer, bool bPerformConditionalPostLoadIfRequired, FPartialKeyHashes* PartialKeyHashes = nullptr);
+	POSESEARCH_API FKeyBuilder(const UObject* Object, bool bUseDataVer, bool bPerformConditionalPostLoadIfRequired, FPartialKeyHashes* PartialKeyHashes = nullptr, bool bInPerformSequenceResidencyRequests = true);
 
 	// Experimental, this feature might be removed without warning, not for production use
 	POSESEARCH_API bool ValidateAgainst(const FKeyBuilder& Other) const;
@@ -203,6 +203,8 @@ protected:
 	bool bPerformConditionalPostLoad = false;
 
 private:
+	bool bPerformSequenceResidencyRequests = false;
+	
 	void SerializeObjectInternal(UObject* Object);
 	FArchive& TryAddDependency(UObject* Object);
 
