PoseSearch - allowing PS debugger track to be children of any object track

#Synced-CL 51737212

diff --git a/Engine/Plugins/Animation/PoseSearch/Source/Editor/Private/PoseSearchDebuggerTrackCreator.cpp b/Engine/Plugins/Animation/PoseSearch/Source/Editor/Private/PoseSearchDebuggerTrackCreator.cpp
--- a/Engine/Plugins/Animation/PoseSearch/Source/Editor/Private/PoseSearchDebuggerTrackCreator.cpp
+++ b/Engine/Plugins/Animation/PoseSearch/Source/Editor/Private/PoseSearchDebuggerTrackCreator.cpp
@@ -17,7 +17,7 @@ namespace UE::PoseSearch
 
 FName FDebuggerTrackCreator::GetTargetTypeNameInternal() const
 {
-	return UAnimInstance::StaticClass()->GetFName();
+	return UObject::StaticClass()->GetFName();
 }
 
 void FDebuggerTrackCreator::GetTrackTypesInternal(TArray<RewindDebugger::FRewindDebuggerTrackType>& Types) const
