PoseSearch - implementing BP overridable UFUNCTION for multiple schema channels to gather data to construct the query the MM query passing the more versatile FChooserEvaluationContext as input parameter

#Synced-CL 51737212

diff --git a/Engine/Plugins/Animation/PoseSearch/Source/Runtime/Private/PoseSearchFeatureChannel_Curve.cpp b/Engine/Plugins/Animation/PoseSearch/Source/Runtime/Private/PoseSearchFeatureChannel_Curve.cpp
--- a/Engine/Plugins/Animation/PoseSearch/Source/Runtime/Private/PoseSearchFeatureChannel_Curve.cpp
+++ b/Engine/Plugins/Animation/PoseSearch/Source/Runtime/Private/PoseSearchFeatureChannel_Curve.cpp
@@ -30,23 +30,31 @@ bool UPoseSearchFeatureChannel_Curve::Finalize(UPoseSearchSchema* Schema)
 	return true;
 }
 
+float UPoseSearchFeatureChannel_Curve::BP_GetCurveValueFromContext_Implementation(FChooserEvaluationContext& Context) const
+{
+	using namespace UE::PoseSearch;
+
+	if (const UAnimInstance* AnimInstance = Cast<UAnimInstance>(GetAnimContext(&Context)))
+	{
+		PRAGMA_DISABLE_DEPRECATION_WARNINGS
+		return BP_GetCurveValue(AnimInstance);
+		PRAGMA_ENABLE_DEPRECATION_WARNINGS
+	}
+	
+	UE_LOGF(LogPoseSearch, Warning, "UPoseSearchFeatureChannel_Position::BP_GetCurveValueFromContext_Implementation - unsupported null UAnimInstance: WIP support for AnimNext!");
+	return 0.f;
+}
+
+
 void UPoseSearchFeatureChannel_Curve::BuildQuery(UE::PoseSearch::FSearchContext& SearchContext) const
 {
 	using namespace UE::PoseSearch;
 
-	float CurveValue = 0.0f;
 	if (bUseBlueprintQueryOverride)
 	{
-		if (const UAnimInstance* AnimInstance = Cast<UAnimInstance>(GetAnimContext(SearchContext.GetContext(SampleRole))))
-		{
-			CurveValue = BP_GetCurveValue(AnimInstance);
-			FFeatureVectorHelper::EncodeFloat(SearchContext.EditFeatureVector(), ChannelDataOffset, CurveValue);
-		}
-		else
-		{
-			// @todo: support non UAnimInstance anim contexts for AnimNext
-			UE_LOGF(LogPoseSearch, Warning, "UPoseSearchFeatureChannel_Curve::BuildQuery - unsupported null UAnimInstance: WIP support for AnimNext!");
-		}
+		FChooserEvaluationContext* ChooserContext = SearchContext.GetContext(SampleRole);
+		const float CurveValue = ensure(ChooserContext) ? BP_GetCurveValueFromContext(*ChooserContext) : 0.f;
+		FFeatureVectorHelper::EncodeFloat(SearchContext.EditFeatureVector(), ChannelDataOffset, CurveValue);
 		return;
 	}
 	
@@ -95,7 +103,7 @@ void UPoseSearchFeatureChannel_Curve::BuildQuery(UE::PoseSearch::FSearchContext&
 		return;
 	}
 	
-	CurveValue = SearchContext.GetSampleCurveValue(SampleTimeOffset, CurveName, SampleRole);
+	const float CurveValue = SearchContext.GetSampleCurveValue(SampleTimeOffset, CurveName, SampleRole);
 	FFeatureVectorHelper::EncodeFloat(SearchContext.EditFeatureVector(), ChannelDataOffset, CurveValue);
 }
 
diff --git a/Engine/Plugins/Animation/PoseSearch/Source/Runtime/Private/PoseSearchFeatureChannel_Heading.cpp b/Engine/Plugins/Animation/PoseSearch/Source/Runtime/Private/PoseSearchFeatureChannel_Heading.cpp
--- a/Engine/Plugins/Animation/PoseSearch/Source/Runtime/Private/PoseSearchFeatureChannel_Heading.cpp
+++ b/Engine/Plugins/Animation/PoseSearch/Source/Runtime/Private/PoseSearchFeatureChannel_Heading.cpp
@@ -116,23 +116,31 @@ FVector UPoseSearchFeatureChannel_Heading::GetAxis(const FQuat& Rotation) const
 	return FVector::XAxisVector;
 }
 
+FQuat UPoseSearchFeatureChannel_Heading::BP_GetWorldRotationFromContext_Implementation(FChooserEvaluationContext& Context) const
+{
+	using namespace UE::PoseSearch;
+
+	if (const UAnimInstance* AnimInstance = Cast<UAnimInstance>(GetAnimContext(&Context)))
+	{
+		PRAGMA_DISABLE_DEPRECATION_WARNINGS
+		return BP_GetWorldRotation(AnimInstance);
+		PRAGMA_ENABLE_DEPRECATION_WARNINGS
+	}
+	
+	UE_LOGF(LogPoseSearch, Warning, "UPoseSearchFeatureChannel_Heading::BP_GetWorldRotationFromContext_Implementation - unsupported null UAnimInstance: WIP support for AnimNext!");
+	return FQuat::Identity;
+}
+
 void UPoseSearchFeatureChannel_Heading::BuildQuery(UE::PoseSearch::FSearchContext& SearchContext) const
 {
 	using namespace UE::PoseSearch;
 
 	if (bUseBlueprintQueryOverride)
 	{
-		if (const UAnimInstance* AnimInstance = Cast<UAnimInstance>(GetAnimContext(SearchContext.GetContext(SampleRole))))
-		{
-			const FQuat BoneRotationWorld = BP_GetWorldRotation(AnimInstance);
-			const FQuat BoneRotation = SearchContext.GetSampleRotation(SampleTimeOffset, OriginTimeOffset, SchemaBoneIdx, SchemaOriginBoneIdx, SampleRole, OriginRole, EPermutationTimeType::UseSampleTime, &BoneRotationWorld);
-			FFeatureVectorHelper::EncodeVector(SearchContext.EditFeatureVector(), ChannelDataOffset, GetAxis(BoneRotation), ComponentStripping, true);
-		}
-		else
-		{
-			// @todo: support non UAnimInstance anim contexts for AnimNext
-			UE_LOGF(LogPoseSearch, Warning, "UPoseSearchFeatureChannel_Heading::BuildQuery - unsupported null UAnimInstance: WIP support for AnimNext!");
-		}
+		FChooserEvaluationContext* ChooserContext = SearchContext.GetContext(SampleRole);
+		const FQuat BoneRotationWorld = ensure(ChooserContext) ? BP_GetWorldRotationFromContext(*ChooserContext) : FQuat::Identity;
+		const FQuat BoneRotation = SearchContext.GetSampleRotation(SampleTimeOffset, OriginTimeOffset, SchemaBoneIdx, SchemaOriginBoneIdx, SampleRole, OriginRole, EPermutationTimeType::UseSampleTime, &BoneRotationWorld);
+		FFeatureVectorHelper::EncodeVector(SearchContext.EditFeatureVector(), ChannelDataOffset, GetAxis(BoneRotation), ComponentStripping, true);
 		return;
 	}
 	
diff --git a/Engine/Plugins/Animation/PoseSearch/Source/Runtime/Private/PoseSearchFeatureChannel_Position.cpp b/Engine/Plugins/Animation/PoseSearch/Source/Runtime/Private/PoseSearchFeatureChannel_Position.cpp
--- a/Engine/Plugins/Animation/PoseSearch/Source/Runtime/Private/PoseSearchFeatureChannel_Position.cpp
+++ b/Engine/Plugins/Animation/PoseSearch/Source/Runtime/Private/PoseSearchFeatureChannel_Position.cpp
@@ -95,23 +95,31 @@ void UPoseSearchFeatureChannel_Position::AddDependentChannels(UPoseSearchSchema*
 	}
 }
 
+FVector UPoseSearchFeatureChannel_Position::BP_GetWorldPositionFromContext_Implementation(FChooserEvaluationContext& Context) const
+{
+	using namespace UE::PoseSearch;
+
+	if (const UAnimInstance* AnimInstance = Cast<UAnimInstance>(GetAnimContext(&Context)))
+	{
+		PRAGMA_DISABLE_DEPRECATION_WARNINGS
+		return BP_GetWorldPosition(AnimInstance);
+		PRAGMA_ENABLE_DEPRECATION_WARNINGS
+	}
+	
+	UE_LOGF(LogPoseSearch, Warning, "UPoseSearchFeatureChannel_Position::BP_GetWorldPositionFromContext_Implementation - unsupported null UAnimInstance: WIP support for AnimNext!");
+	return FVector::ZeroVector;
+}
+
 void UPoseSearchFeatureChannel_Position::BuildQuery(UE::PoseSearch::FSearchContext& SearchContext) const
 {
 	using namespace UE::PoseSearch;
 
 	if (bUseBlueprintQueryOverride)
 	{
-		if (const UAnimInstance* AnimInstance = Cast<UAnimInstance>(GetAnimContext(SearchContext.GetContext(SampleRole))))
-		{
-			const FVector BonePositionWorld = BP_GetWorldPosition(AnimInstance);
-			const FVector BonePosition = SearchContext.GetSamplePosition(SampleTimeOffset, OriginTimeOffset, SchemaBoneIdx, SchemaOriginBoneIdx, SampleRole, OriginRole, PermutationTimeType, &BonePositionWorld);
-  			FFeatureVectorHelper::EncodeVector(SearchContext.EditFeatureVector(), ChannelDataOffset, BonePosition, ComponentStripping, false);
-		}
-		else
-		{
-			// @todo: support non UAnimInstance anim contexts for AnimNext
-			UE_LOGF(LogPoseSearch, Warning, "UPoseSearchFeatureChannel_Position::BuildQuery - unsupported null UAnimInstance: WIP support for AnimNext!");
-		}
+		FChooserEvaluationContext* ChooserContext = SearchContext.GetContext(SampleRole);
+		const FVector BonePositionWorld = ensure(ChooserContext) ? BP_GetWorldPositionFromContext(*ChooserContext) : FVector::ZeroVector;
+		const FVector BonePosition = SearchContext.GetSamplePosition(SampleTimeOffset, OriginTimeOffset, SchemaBoneIdx, SchemaOriginBoneIdx, SampleRole, OriginRole, PermutationTimeType, &BonePositionWorld);
+  		FFeatureVectorHelper::EncodeVector(SearchContext.EditFeatureVector(), ChannelDataOffset, BonePosition, ComponentStripping, false);
 		return;
 	}
 
diff --git a/Engine/Plugins/Animation/PoseSearch/Source/Runtime/Private/PoseSearchFeatureChannel_TimeToEvent.cpp b/Engine/Plugins/Animation/PoseSearch/Source/Runtime/Private/PoseSearchFeatureChannel_TimeToEvent.cpp
--- a/Engine/Plugins/Animation/PoseSearch/Source/Runtime/Private/PoseSearchFeatureChannel_TimeToEvent.cpp
+++ b/Engine/Plugins/Animation/PoseSearch/Source/Runtime/Private/PoseSearchFeatureChannel_TimeToEvent.cpp
@@ -22,6 +22,21 @@ bool UPoseSearchFeatureChannel_TimeToEvent::Finalize(UPoseSearchSchema* Schema)
 	return true;
 }
 
+float UPoseSearchFeatureChannel_TimeToEvent::BP_GetTimeToEventFromContext_Implementation(FChooserEvaluationContext& Context) const
+{
+	using namespace UE::PoseSearch;
+
+	if (const UAnimInstance* AnimInstance = Cast<UAnimInstance>(GetAnimContext(&Context)))
+	{
+		PRAGMA_DISABLE_DEPRECATION_WARNINGS
+		return BP_GetTimeToEvent(AnimInstance);
+		PRAGMA_ENABLE_DEPRECATION_WARNINGS
+	}
+	
+	UE_LOGF(LogPoseSearch, Warning, "UPoseSearchFeatureChannel_TimeToEvent::BP_GetTimeToEventFromContext_Implementation - unsupported null UAnimInstance: WIP support for AnimNext!");
+	return 0.f;
+}
+
 void UPoseSearchFeatureChannel_TimeToEvent::BuildQuery(UE::PoseSearch::FSearchContext& SearchContext) const
 {
 	using namespace UE::PoseSearch;
@@ -38,15 +53,9 @@ void UPoseSearchFeatureChannel_TimeToEvent::BuildQuery(UE::PoseSearch::FSearchCo
 		return;
 	}
 	
-	if (const UAnimInstance* AnimInstance = Cast<UAnimInstance>(GetAnimContext(SearchContext.GetContexts()[0])))
-	{
-		const float TimeToEvent = BP_GetTimeToEvent(AnimInstance);
-		FFeatureVectorHelper::EncodeFloat(SearchContext.EditFeatureVector(), ChannelDataOffset, TimeToEvent);
-	}
-	else
-	{
-		UE_LOGF(LogPoseSearch, Warning, "UPoseSearchFeatureChannel_TimeToEvent::BuildQuery - UPoseSearchFeatureChannel_TimeToEvent is designed to work only as BluePrint overridden class. The query TimeToEvent value will be defaulted to 0");
-	}
+	FChooserEvaluationContext* ChooserContext = SearchContext.GetContexts()[0];
+	const float TimeToEvent = ensure(ChooserContext) ? BP_GetTimeToEventFromContext(*ChooserContext) : 0.f;
+	FFeatureVectorHelper::EncodeFloat(SearchContext.EditFeatureVector(), ChannelDataOffset, TimeToEvent);
 }
 
 #if WITH_EDITOR
diff --git a/Engine/Plugins/Animation/PoseSearch/Source/Runtime/Private/PoseSearchFeatureChannel_Velocity.cpp b/Engine/Plugins/Animation/PoseSearch/Source/Runtime/Private/PoseSearchFeatureChannel_Velocity.cpp
--- a/Engine/Plugins/Animation/PoseSearch/Source/Runtime/Private/PoseSearchFeatureChannel_Velocity.cpp
+++ b/Engine/Plugins/Animation/PoseSearch/Source/Runtime/Private/PoseSearchFeatureChannel_Velocity.cpp
@@ -44,27 +44,36 @@ void UPoseSearchFeatureChannel_Velocity::AddDependentChannels(UPoseSearchSchema*
 	}
 }
 
+FVector UPoseSearchFeatureChannel_Velocity::BP_GetWorldVelocityFromContext_Implementation(FChooserEvaluationContext& Context) const
+{
+	using namespace UE::PoseSearch;
+
+	if (const UAnimInstance* AnimInstance = Cast<UAnimInstance>(GetAnimContext(&Context)))
+	{
+		PRAGMA_DISABLE_DEPRECATION_WARNINGS
+		return BP_GetWorldVelocity(AnimInstance);
+		PRAGMA_ENABLE_DEPRECATION_WARNINGS
+	}
+	
+	UE_LOGF(LogPoseSearch, Warning, "UPoseSearchFeatureChannel_Velocity::BP_GetWorldVelocityFromContext_Implementation - unsupported null UAnimInstance: WIP support for AnimNext!");
+	return FVector::ZeroVector;
+}
+
 void UPoseSearchFeatureChannel_Velocity::BuildQuery(UE::PoseSearch::FSearchContext& SearchContext) const
 {
 	using namespace UE::PoseSearch;
 
 	if (bUseBlueprintQueryOverride)
 	{
-		if (const UAnimInstance* AnimInstance = Cast<UAnimInstance>(GetAnimContext(SearchContext.GetContext(SampleRole))))
-		{
-			const FVector LinearVelocityWorld = BP_GetWorldVelocity(AnimInstance);
-			FVector LinearVelocity = SearchContext.GetSampleVelocity(SampleTimeOffset, OriginTimeOffset, SchemaBoneIdx, SchemaOriginBoneIdx, SampleRole, OriginRole, bUseCharacterSpaceVelocities, EPermutationTimeType::UseSampleTime, &LinearVelocityWorld);
-			if (bNormalize)
-			{
-				LinearVelocity = LinearVelocity.GetClampedToMaxSize(1.f);
-			}
-			FFeatureVectorHelper::EncodeVector(SearchContext.EditFeatureVector(), ChannelDataOffset, LinearVelocity, ComponentStripping, false);
-		}
-		else
+		FChooserEvaluationContext* ChooserContext = SearchContext.GetContext(SampleRole);
+		const FVector LinearVelocityWorld = ensure(ChooserContext) ? BP_GetWorldVelocityFromContext(*ChooserContext) : FVector::ZeroVector;
+	
+		FVector LinearVelocity = SearchContext.GetSampleVelocity(SampleTimeOffset, OriginTimeOffset, SchemaBoneIdx, SchemaOriginBoneIdx, SampleRole, OriginRole, bUseCharacterSpaceVelocities, EPermutationTimeType::UseSampleTime, &LinearVelocityWorld);
+		if (bNormalize)
 		{
-			// @todo: support non UAnimInstance anim contexts for AnimNext
-			UE_LOGF(LogPoseSearch, Warning, "UPoseSearchFeatureChannel_Velocity::BuildQuery - unsupported null UAnimInstance: WIP support for AnimNext!");
+			LinearVelocity = LinearVelocity.GetClampedToMaxSize(1.f);
 		}
+		FFeatureVectorHelper::EncodeVector(SearchContext.EditFeatureVector(), ChannelDataOffset, LinearVelocity, ComponentStripping, false);
 		return;
 	}
 	
diff --git a/Engine/Plugins/Animation/PoseSearch/Source/Runtime/Public/PoseSearch/PoseSearchFeatureChannel_Curve.h b/Engine/Plugins/Animation/PoseSearch/Source/Runtime/Public/PoseSearch/PoseSearchFeatureChannel_Curve.h
--- a/Engine/Plugins/Animation/PoseSearch/Source/Runtime/Public/PoseSearch/PoseSearchFeatureChannel_Curve.h
+++ b/Engine/Plugins/Animation/PoseSearch/Source/Runtime/Public/PoseSearch/PoseSearchFeatureChannel_Curve.h
@@ -49,9 +49,13 @@ public:
 	FName NormalizationGroup;
 #endif //WITH_EDITORONLY_DATA
 
+	UE_DEPRECATED(5.8, "Use BP_GetCurveValueFromContext instead")
 	UFUNCTION(BlueprintPure, BlueprintImplementableEvent, meta=(BlueprintThreadSafe, DisplayName = "Get World Curve"), Category = "Settings")
 	UE_API float BP_GetCurveValue(const UAnimInstance* AnimInstance) const;
 
+	UFUNCTION(BlueprintPure, BlueprintNativeEvent, meta=(BlueprintThreadSafe, DisplayName = "Get World Curve From Context"), Category = "Settings")
+	UE_API float BP_GetCurveValueFromContext(UPARAM(Ref) FChooserEvaluationContext& Context) const;
+
 	bool bUseBlueprintQueryOverride = false;
 
 	UE_API UPoseSearchFeatureChannel_Curve();
diff --git a/Engine/Plugins/Animation/PoseSearch/Source/Runtime/Public/PoseSearch/PoseSearchFeatureChannel_Heading.h b/Engine/Plugins/Animation/PoseSearch/Source/Runtime/Public/PoseSearch/PoseSearchFeatureChannel_Heading.h
--- a/Engine/Plugins/Animation/PoseSearch/Source/Runtime/Public/PoseSearch/PoseSearchFeatureChannel_Heading.h
+++ b/Engine/Plugins/Animation/PoseSearch/Source/Runtime/Public/PoseSearch/PoseSearchFeatureChannel_Heading.h
@@ -96,9 +96,13 @@ public:
 	FName NormalizationGroup;
 #endif //WITH_EDITORONLY_DATA
 
+	UE_DEPRECATED(5.8, "Use BP_GetWorldRotationFromContext instead")
 	UFUNCTION(BlueprintPure, BlueprintImplementableEvent, meta=(BlueprintThreadSafe, DisplayName = "Get World Rotation"), Category = "Settings")
 	UE_API FQuat BP_GetWorldRotation(const UAnimInstance* AnimInstance) const;
 
+	UFUNCTION(BlueprintPure, BlueprintNativeEvent, meta=(BlueprintThreadSafe, DisplayName = "Get World Rotation From Context"), Category = "Settings")
+	UE_API FQuat BP_GetWorldRotationFromContext(UPARAM(Ref) FChooserEvaluationContext& Context) const;
+
 	bool bUseBlueprintQueryOverride = false;
 
 	UE_API UPoseSearchFeatureChannel_Heading();
diff --git a/Engine/Plugins/Animation/PoseSearch/Source/Runtime/Public/PoseSearch/PoseSearchFeatureChannel_Position.h b/Engine/Plugins/Animation/PoseSearch/Source/Runtime/Public/PoseSearch/PoseSearchFeatureChannel_Position.h
--- a/Engine/Plugins/Animation/PoseSearch/Source/Runtime/Public/PoseSearch/PoseSearchFeatureChannel_Position.h
+++ b/Engine/Plugins/Animation/PoseSearch/Source/Runtime/Public/PoseSearch/PoseSearchFeatureChannel_Position.h
@@ -91,9 +91,13 @@ public:
 	FName NormalizationGroup;
 #endif //WITH_EDITORONLY_DATA
 
+	UE_DEPRECATED(5.8, "Use BP_GetWorldPositionFromContext instead")
 	UFUNCTION(BlueprintPure, BlueprintImplementableEvent, meta=(BlueprintThreadSafe, DisplayName = "Get World Position"), Category = "Settings")
 	UE_API FVector BP_GetWorldPosition(const UAnimInstance* AnimInstance) const;
 
+	UFUNCTION(BlueprintPure, BlueprintNativeEvent, meta=(BlueprintThreadSafe, DisplayName = "Get World Position From Context"), Category = "Settings")
+	UE_API FVector BP_GetWorldPositionFromContext(UPARAM(Ref) FChooserEvaluationContext& Context) const;
+
 	bool bUseBlueprintQueryOverride = false;
 
 	UE_API UPoseSearchFeatureChannel_Position();
diff --git a/Engine/Plugins/Animation/PoseSearch/Source/Runtime/Public/PoseSearch/PoseSearchFeatureChannel_TimeToEvent.h b/Engine/Plugins/Animation/PoseSearch/Source/Runtime/Public/PoseSearch/PoseSearchFeatureChannel_TimeToEvent.h
--- a/Engine/Plugins/Animation/PoseSearch/Source/Runtime/Public/PoseSearch/PoseSearchFeatureChannel_TimeToEvent.h
+++ b/Engine/Plugins/Animation/PoseSearch/Source/Runtime/Public/PoseSearch/PoseSearchFeatureChannel_TimeToEvent.h
@@ -19,9 +19,13 @@ public:
 	UPROPERTY(EditAnywhere, Category = "Settings")
 	int32 SamplingAttributeId = 0;
 
+	UE_DEPRECATED(5.8, "Use BP_GetTimeToEventFromContext instead")
 	UFUNCTION(BlueprintPure, BlueprintImplementableEvent, meta=(BlueprintThreadSafe, DisplayName = "Get Time To Event"), Category = "Settings")
 	UE_API float BP_GetTimeToEvent(const UAnimInstance* AnimInstance) const;
 
+	UFUNCTION(BlueprintPure, BlueprintNativeEvent, meta=(BlueprintThreadSafe, DisplayName = "Get Time To Event From Context"), Category = "Settings")
+	UE_API float BP_GetTimeToEventFromContext(UPARAM(Ref) FChooserEvaluationContext& Context) const;
+
 	bool bUseBlueprintQueryOverride = false;
 
 	UE_API UPoseSearchFeatureChannel_TimeToEvent();
diff --git a/Engine/Plugins/Animation/PoseSearch/Source/Runtime/Public/PoseSearch/PoseSearchFeatureChannel_Velocity.h b/Engine/Plugins/Animation/PoseSearch/Source/Runtime/Public/PoseSearch/PoseSearchFeatureChannel_Velocity.h
--- a/Engine/Plugins/Animation/PoseSearch/Source/Runtime/Public/PoseSearch/PoseSearchFeatureChannel_Velocity.h
+++ b/Engine/Plugins/Animation/PoseSearch/Source/Runtime/Public/PoseSearch/PoseSearchFeatureChannel_Velocity.h
@@ -91,9 +91,13 @@ public:
 	FName NormalizationGroup;
 #endif //WITH_EDITORONLY_DATA
 
+	UE_DEPRECATED(5.8, "Use BP_GetWorldVelocityFromContext instead")
 	UFUNCTION(BlueprintPure, BlueprintImplementableEvent, meta=(BlueprintThreadSafe, DisplayName = "Get World Velocity"), Category = "Settings")
 	UE_API FVector BP_GetWorldVelocity(const UAnimInstance* AnimInstance) const;
 
+	UFUNCTION(BlueprintPure, BlueprintNativeEvent, meta=(BlueprintThreadSafe, DisplayName = "Get World Velocity From Context"), Category = "Settings")
+	UE_API FVector BP_GetWorldVelocityFromContext(UPARAM(Ref) FChooserEvaluationContext& Context) const;
+
 	bool bUseBlueprintQueryOverride = false;
 
 	UE_API UPoseSearchFeatureChannel_Velocity();
