PoseSearch - fix for UPoseSearchFeatureChannel::CanBeNormalizedWith

#Synced-CL 51737212

diff --git a/Engine/Plugins/Animation/PoseSearch/Source/Runtime/Private/PoseSearchDerivedDataKey.cpp b/Engine/Plugins/Animation/PoseSearch/Source/Runtime/Private/PoseSearchDerivedDataKey.cpp
--- a/Engine/Plugins/Animation/PoseSearch/Source/Runtime/Private/PoseSearchDerivedDataKey.cpp
+++ b/Engine/Plugins/Animation/PoseSearch/Source/Runtime/Private/PoseSearchDerivedDataKey.cpp
@@ -449,11 +449,7 @@ FIoHash FKeyBuilder::Finalize()
 	*this << VersionGuid;
 	*this << AnimationCompressionVersionString;
 	*this << NonConstDatabaseIndexDerivedDataCacheKeyVersion;
-
-	if (bPerformSequenceResidencyRequests)
-	{
-		*this << bPerformSequenceResidencyRequests;
-	}
+	*this << bPerformSequenceResidencyRequests;
 		
 	// iterating over all ObjectKeysToVisit to finalize the composition the final hash while continuing adding their unvisited dependencies
 	while (!ObjectKeysToVisit.IsEmpty())
diff --git a/Engine/Plugins/Animation/PoseSearch/Source/Runtime/Private/PoseSearchFeatureChannel.cpp b/Engine/Plugins/Animation/PoseSearch/Source/Runtime/Private/PoseSearchFeatureChannel.cpp
--- a/Engine/Plugins/Animation/PoseSearch/Source/Runtime/Private/PoseSearchFeatureChannel.cpp
+++ b/Engine/Plugins/Animation/PoseSearch/Source/Runtime/Private/PoseSearchFeatureChannel.cpp
@@ -120,7 +120,7 @@ bool UPoseSearchFeatureChannel::CanBeNormalizedWith(const UPoseSearchFeatureChan
 	}
 
 	TLabelBuilder ThisLabelBuilder, OtherLabelBuilder;
-	if (GetLabel(ThisLabelBuilder).ToString() != Other->GetLabel(OtherLabelBuilder).ToString())
+	if (FCString::Strcmp(GetLabel(ThisLabelBuilder).ToString(), Other->GetLabel(OtherLabelBuilder).ToString()) != 0)
 	{
 		return false;
 	}
diff --git a/Engine/Plugins/Animation/PoseSearch/Source/Runtime/Public/PoseSearch/PoseSearchDefines.h b/Engine/Plugins/Animation/PoseSearch/Source/Runtime/Public/PoseSearch/PoseSearchDefines.h
--- a/Engine/Plugins/Animation/PoseSearch/Source/Runtime/Public/PoseSearch/PoseSearchDefines.h
+++ b/Engine/Plugins/Animation/PoseSearch/Source/Runtime/Public/PoseSearch/PoseSearchDefines.h
@@ -43,7 +43,7 @@ static constexpr int32 PreallocatedSearchesNum = 4;
 static constexpr float FiniteDelta = 1 / 60.0f; // Time delta used for computing pose derivatives
 static constexpr int32 MaxNumberOfCollectedPoseCandidatesPerDatabase = 256;
 static constexpr int32 InvalidSearchId = INDEX_NONE; // @todo: it should be 0, since the search id is an hash!
-static constexpr int32 DatabaseIndexDerivedDataCacheKeyVersion = 2727; // bump this value to force reindex ALL the UPoseSearchDatabases, without having to change the global FDevSystemGuids::POSESEARCHDB_DERIVEDDATA_VER
+static constexpr int32 DatabaseIndexDerivedDataCacheKeyVersion = 2731; // bump this value to force reindex ALL the UPoseSearchDatabases, without having to change the global FDevSystemGuids::POSESEARCHDB_DERIVEDDATA_VER
 static constexpr int32 DefaultTickPriority = 100;
 
 template<uint32 Alignment = DEFAULT_ALIGNMENT> class TMemStackSetAllocator : public TSetAllocator<TSparseArrayAllocator<TMemStackAllocator<Alignment>, TMemStackAllocator<Alignment>>, TMemStackAllocator<Alignment>> {};
