PoseSearch - fix for UPoseSearchDatabase::PopulateSelectableAssetIdx calling SelectableAssetIdx.Reset() correctly

#Synced-CL 51737212

diff --git a/Engine/Plugins/Animation/PoseSearch/Source/Runtime/Private/PoseSearchDatabase.cpp b/Engine/Plugins/Animation/PoseSearch/Source/Runtime/Private/PoseSearchDatabase.cpp
--- a/Engine/Plugins/Animation/PoseSearch/Source/Runtime/Private/PoseSearchDatabase.cpp
+++ b/Engine/Plugins/Animation/PoseSearch/Source/Runtime/Private/PoseSearchDatabase.cpp
@@ -1651,28 +1651,25 @@ void UPoseSearchDatabase::PopulateSelectableAssetIdx(FSelectableAssetIdx& Select
 
 void UPoseSearchDatabase::PopulateSelectableAssetIdx(FSelectableAssetIdx& SelectableAssetIdx, const UE::PoseSearch::FStackAssetSet* AssetsToConsider) const
 {
-	if (ensure(AssetsToConsider))
+	SelectableAssetIdx.Reset();
+	if (ensure(AssetsToConsider) && !AssetsToConsider->IsEmpty())
 	{
-		SelectableAssetIdx.Reset();
-		if (!AssetsToConsider->IsEmpty())
+		for (const UObject* AssetToConsider : *AssetsToConsider)
+		{
+			SelectableAssetIdx.Append(GetAssetIndexesForSourceAsset(AssetToConsider));
+		}
+
+		if (!SelectableAssetIdx.IsEmpty())
 		{
-			for (const UObject* AssetToConsider : *AssetsToConsider)
+			if (SelectableAssetIdx.Num() != GetSearchIndex().Assets.Num())
 			{
-				SelectableAssetIdx.Append(GetAssetIndexesForSourceAsset(AssetToConsider));
+				SelectableAssetIdx.Sort();
 			}
-
-			if (!SelectableAssetIdx.IsEmpty())
+			else
 			{
-				if (SelectableAssetIdx.Num() != GetSearchIndex().Assets.Num())
-				{
-					SelectableAssetIdx.Sort();
-				}
-				else
-				{
-					// SelectableAssetIdx contains ALL the Database->GetSearchIndex().Assets. 
-					// We reset SelectableAssetIdx since it has the same meaning, and it'll perform better
-					SelectableAssetIdx.Reset();
-				}
+				// SelectableAssetIdx contains ALL the Database->GetSearchIndex().Assets. 
+				// We reset SelectableAssetIdx since it has the same meaning, and it'll perform better
+				SelectableAssetIdx.Reset();
 			}
 		}
 	}
@@ -1682,13 +1679,12 @@ void UPoseSearchDatabase::PopulateSelectableAssetIdx(FSelectableAssetIdx& Select
 {
 	using namespace UE::PoseSearch;
 	
+	SelectableAssetIdx.Reset();
 	if (ensure(AssetIndexesToConsiderPerDatabase))
 	{
-		const int32 NumAnimationAssets = GetNumAnimationAssets();
-
-		SelectableAssetIdx.Reset();
 		if (const FStackAssetIndexes* AssetIndexesToConsider = AssetIndexesToConsiderPerDatabase->Find(this))
 		{
+			const int32 NumAnimationAssets = GetNumAnimationAssets();
 			ensure(AssetIndexesToConsider->Num() <= NumAnimationAssets);
 
 			// if AssetIndexesToConsider->Num() == NumAnimationAssets we want to search the full database and and empty SelectableAssetIdx does that
