PoseSearch - adding missing EditAnywhere to UPoseSearchInteractionAsset Debug properties

#Synced-CL 51930448

diff --git a/Engine/Plugins/Animation/PoseSearch/Source/Runtime/Public/PoseSearch/PoseSearchInteractionAsset.h b/Engine/Plugins/Animation/PoseSearch/Source/Runtime/Public/PoseSearch/PoseSearchInteractionAsset.h
--- a/Engine/Plugins/Animation/PoseSearch/Source/Runtime/Public/PoseSearch/PoseSearchInteractionAsset.h
+++ b/Engine/Plugins/Animation/PoseSearch/Source/Runtime/Public/PoseSearch/PoseSearchInteractionAsset.h
@@ -62,7 +62,7 @@ public:
 
 #if WITH_EDITORONLY_DATA
 
-	UPROPERTY(Transient, BlueprintReadWrite, Category = "Debug", meta = (EditCondition=bEnableDebugWarp, EditConditionHides))
+	UPROPERTY(Transient, EditAnywhere, BlueprintReadWrite, Category = "Debug", meta = (EditCondition=bEnableDebugWarp, EditConditionHides))
 	TArray<FTransform> DebugWarpOffsets;
 
 	// used to test warping: 0 - no warping applied, 1 - full warping/alignment applied
@@ -70,14 +70,14 @@ public:
 	// UMultiAnimAsset::GetOrigin() definition and warped accordingly with CalculateWarpTransforms
 	// following the rotation and translation weights defined in Items::WarpingWeightRotation and
 	// Items::WarpingWeightTranslation as relative weights between the Items (they'll be normalized at runtime)
-	UPROPERTY(Transient, BlueprintReadWrite, Category = "Debug", meta = (ClampMin = "0", ClampMax = "1", EditCondition=bEnableDebugWarp, EditConditionHides))
+	UPROPERTY(Transient, EditAnywhere, BlueprintReadWrite, Category = "Debug", meta = (ClampMin = "0", ClampMax = "1", EditCondition=bEnableDebugWarp, EditConditionHides))
 	float DebugWarpAmount = 0.f;
 
 	// useful to test TimeOffset
-	UPROPERTY(Transient, BlueprintReadWrite, Category = "Debug", meta = (ClampMin = "0", EditCondition=bEnableDebugWarp, EditConditionHides))
+	UPROPERTY(Transient, EditAnywhere, BlueprintReadWrite, Category = "Debug", meta = (ClampMin = "0", EditCondition=bEnableDebugWarp, EditConditionHides))
 	float DebugTimeOffset = 0.f;
 
-	UPROPERTY(Transient, BlueprintReadWrite, Category = "Debug")
+	UPROPERTY(Transient, EditAnywhere, BlueprintReadWrite, Category = "Debug")
 	bool bEnableDebugWarp = false;
 #endif // WITH_EDITORONLY_DATA
 
