Pose Search - MM col new/select database improvements

#Synced-CL 51930448

diff --git a/Engine/Plugins/Animation/PoseSearch/Source/Editor/Private/Chooser/PoseSearchColumnEditor.cpp b/Engine/Plugins/Animation/PoseSearch/Source/Editor/Private/Chooser/PoseSearchColumnEditor.cpp
--- a/Engine/Plugins/Animation/PoseSearch/Source/Editor/Private/Chooser/PoseSearchColumnEditor.cpp
+++ b/Engine/Plugins/Animation/PoseSearch/Source/Editor/Private/Chooser/PoseSearchColumnEditor.cpp
@@ -15,7 +15,7 @@
 #include "SPropertyAccessChainWidget.h"
 #include "Subsystems/AssetEditorSubsystem.h"
 #include "Widgets/Input/SNumericEntryBox.h"
-#include "Widgets/Input/STextEntryPopup.h"
+#include "Widgets/Input/SEditableTextBox.h"
 #include "Widgets/SBoxPanel.h"
 
 #define LOCTEXT_NAMESPACE "PoseSearchColumnEditor"
@@ -79,119 +79,105 @@ TSharedRef<SWidget> CreatePoseSearchColumnWidget(UChooserTable* Chooser, FChoose
 								Chooser->PostEditChange();
 							})));
 
-					MenuBuilder.AddMenuEntry(
-						LOCTEXT("New Database", "New Database"),
-						LOCTEXT("NewDatabase_Tooltip", "Create a new database used to index this row asset"),
-						FSlateIcon(),
-						FUIAction(FExecuteAction::CreateLambda([Button, Chooser, PoseSearchColumn, RowIndex]()
-							{
-								Button->SetIsOpen(false);
-
-								// picking a schema for the newly created database
-												
-								// Load the content browser module to display an asset picker
-								FContentBrowserModule& ContentBrowserModule = FModuleManager::LoadModuleChecked<FContentBrowserModule>("ContentBrowser");
-
-								FAssetPickerConfig AssetPickerConfig;
-								TSharedPtr<SWindow> PickerWindow;
-								UPoseSearchSchema* TargetSchema = nullptr;
-
-								// The asset picker will only show UPoseSearchSchema(s)
-								AssetPickerConfig.Filter.ClassPaths.Add(UPoseSearchSchema::StaticClass()->GetClassPathName());
-								AssetPickerConfig.Filter.bRecursiveClasses = true;
-
-								// The delegate that fires when an asset was selected
-								AssetPickerConfig.OnAssetSelected = FOnAssetSelected::CreateLambda([&PickerWindow, &TargetSchema](const FAssetData& SelectedAsset)
-									{
-										//FSlateApplication::Get().DismissAllMenus();
-										TargetSchema = Cast<UPoseSearchSchema>(SelectedAsset.GetAsset());
-										PickerWindow->RequestDestroyWindow();
-									});
-													
-								// The default view mode should be a list view
-								AssetPickerConfig.InitialAssetViewType = EAssetViewType::List;
-
-								PickerWindow = SNew(SWindow)
-									.Title(LOCTEXT("CreateDatabaseOptions", "Pick Database Schema"))
-									.ClientSize(FVector2D(500, 600))
-									.SupportsMinimize(false)
-									.SupportsMaximize(false)
-									[
-										SNew(SBorder)
-											.BorderImage(FAppStyle::GetBrush("Menu.Background"))
-											[
-												ContentBrowserModule.Get().CreateAssetPicker(AssetPickerConfig)
-											]
-									];
-
-								GEditor->EditorAddModalWindow(PickerWindow.ToSharedRef());
-								PickerWindow.Reset();
-
-								TSharedRef<STextEntryPopup> TextEntry =
-									SNew(STextEntryPopup)
-									.Label(LOCTEXT("NewDatabaseLabel", "New Database Name"))
-									.OnTextCommitted_Lambda([Chooser, PoseSearchColumn, RowIndex, TargetSchema](FText InText, ETextCommit::Type InCommitType)
+					MenuBuilder.AddWidget(
+						SNew(SBox)
+							.ToolTipText(LOCTEXT("NewDatabase_Tooltip", "Create a new database used to index this row asset"))
+							[
+								SNew(SEditableTextBox)
+									.Font(IDetailLayoutBuilder::GetDetailFont())
+									.HintText(LOCTEXT("NewDatabaseHint", "Enter database name..."))
+									.OnTextCommitted_Lambda([Button, Chooser, PoseSearchColumn, RowIndex](const FText& InText, ETextCommit::Type InCommitType)
 										{
-											if (InCommitType == ETextCommit::OnEnter && TargetSchema)
+											if (InCommitType == ETextCommit::OnEnter)
 											{
-												FSlateApplication::Get().DismissAllMenus();
-												FName NewDatabaseName = FName(InText.ToString());
+												Button->SetIsOpen(false);
 
-												UChooserTable* RootChooser = Chooser->GetRootChooser();
-												if (NewDatabaseName == NAME_None || FindObject<UObject>(Chooser, NewDatabaseName.ToString()))
+												// picking a schema for the newly created database
+												FContentBrowserModule& ContentBrowserModule = FModuleManager::LoadModuleChecked<FContentBrowserModule>("ContentBrowser");
+
+												FAssetPickerConfig AssetPickerConfig;
+												TSharedPtr<SWindow> PickerWindow;
+												UPoseSearchSchema* TargetSchema = nullptr;
+
+												// The asset picker will only show UPoseSearchSchema(s)
+												AssetPickerConfig.Filter.ClassPaths.Add(UPoseSearchSchema::StaticClass()->GetClassPathName());
+												AssetPickerConfig.Filter.bRecursiveClasses = true;
+
+												// The delegate that fires when an asset was selected
+												AssetPickerConfig.OnAssetSelected = FOnAssetSelected::CreateLambda([&PickerWindow, &TargetSchema](const FAssetData& SelectedAsset)
+													{
+														TargetSchema = Cast<UPoseSearchSchema>(SelectedAsset.GetAsset());
+														PickerWindow->RequestDestroyWindow();
+													});
+
+												// The default view mode should be a list view
+												AssetPickerConfig.InitialAssetViewType = EAssetViewType::List;
+
+												PickerWindow = SNew(SWindow)
+													.Title(LOCTEXT("CreateDatabaseOptions", "Pick Database Schema"))
+													.ClientSize(FVector2D(500, 600))
+													.SupportsMinimize(false)
+													.SupportsMaximize(false)
+													[
+														SNew(SBorder)
+															.BorderImage(FAppStyle::GetBrush("Menu.Background"))
+															[
+																ContentBrowserModule.Get().CreateAssetPicker(AssetPickerConfig)
+															]
+													];
+
+												GEditor->EditorAddModalWindow(PickerWindow.ToSharedRef());
+												PickerWindow.Reset();
+
+												if (TargetSchema)
 												{
-													NewDatabaseName = MakeUniqueObjectName(Chooser, UChooserTable::StaticClass(), NewDatabaseName);
-												}
+													FName NewDatabaseName = FName(InText.ToString());
+													UChooserTable* RootChooser = Chooser->GetRootChooser();
+													if (NewDatabaseName == NAME_None || FindObject<UObject>(Chooser, NewDatabaseName.ToString()))
+													{
+														NewDatabaseName = MakeUniqueObjectName(Chooser, UChooserTable::StaticClass(), NewDatabaseName);
+													}
 
-												UPoseSearchDatabase* NewDatabase = NewObject<UPoseSearchDatabase>(RootChooser, UPoseSearchDatabase::StaticClass(), NewDatabaseName, RF_Transactional);
-												NewDatabase->Schema = TargetSchema;
-												// KDTree frequently fails in Chooser columns so use Brute Force instead.
-												NewDatabase->PoseSearchMode = EPoseSearchMode::BruteForce;
+													UPoseSearchDatabase* NewDatabase = NewObject<UPoseSearchDatabase>(RootChooser, UPoseSearchDatabase::StaticClass(), NewDatabaseName, RF_Transactional);
+													NewDatabase->Schema = TargetSchema;
+													// KDTree frequently fails in Chooser columns so use Brute Force instead.
+													NewDatabase->PoseSearchMode = EPoseSearchMode::BruteForce;
 
-												const FScopedTransaction Transaction(LOCTEXT("Assign New Database", "Assign New Database"));
-												RootChooser->Modify();
-												RootChooser->AddNestedObject(NewDatabase);
+													const FScopedTransaction Transaction(LOCTEXT("Assign New Database", "Assign New Database"));
+													RootChooser->Modify();
+													RootChooser->AddNestedObject(NewDatabase);
 
-												PoseSearchColumn->SwitchDatabase(NewDatabase, RowIndex);
-												RootChooser->PostEditChange();
+													PoseSearchColumn->SwitchDatabase(NewDatabase, RowIndex);
+													RootChooser->PostEditChange();
+												}
 											}
-										});
-
-								FSlateApplication& SlateApp = FSlateApplication::Get();
-								SlateApp.PushMenu(
-									SlateApp.GetInteractiveTopLevelWindows()[0],
-									FWidgetPath(),
-									TextEntry,
-									SlateApp.GetCursorPos(),
-									FPopupTransitionEffect::TypeInPopup
-								);
-							})));
+										})
+							],
+						LOCTEXT("NewDatabaseLabel", "New"),
+						/*bNoIndent=*/true);
 
-					MenuBuilder.AddSubMenu(LOCTEXT("Select Existing", "Select Existing"),
-						LOCTEXT("SelectExisting_Tooltip", "Select Existing Embedded UPoseSearchDatabase from this asset"),
-						FNewMenuDelegate::CreateLambda([Button, Chooser, PoseSearchColumn, RowIndex](FMenuBuilder& SubMenuBuilder)
-							{
-								SubMenuBuilder.BeginSection("Existing", LOCTEXT("Existing Databases", "Existing Databases"));
+				}
 
-								if (UChooserTable* RootChooser = Chooser->GetRootChooser())
-								{
-									for (UObject* Object : RootChooser->NestedObjects)
-									{
-										if (UPoseSearchDatabase* Database = Cast<UPoseSearchDatabase>(Object))
+				MenuBuilder.BeginSection("Existing", LOCTEXT("Existing Databases", "Select Existing"));
+				{
+					if (UChooserTable* RootChooser = Chooser->GetRootChooser())
+					{
+						for (UObject* Object : RootChooser->NestedObjects)
+						{
+							if (UPoseSearchDatabase* Database = Cast<UPoseSearchDatabase>(Object))
+							{
+								MenuBuilder.AddMenuEntry(FText::FromString(Database->GetName()), LOCTEXT("AddExistingObjectTooltip", "Add a reference to this existing Database."), FSlateIcon(),
+									FUIAction(FExecuteAction::CreateLambda([Button, Chooser, PoseSearchColumn, Database, RowIndex]()
 										{
-											SubMenuBuilder.AddMenuEntry(FText::FromString(Database->GetName()), LOCTEXT("AddExistingObjectTooltip", "Add a reference to this existing Database."), FSlateIcon(),
-												FUIAction(FExecuteAction::CreateLambda([Button, Chooser, PoseSearchColumn, Database, RowIndex]()
-													{
-														const FScopedTransaction Transaction(LOCTEXT("Set Database", "Set Database"));
-														Chooser->Modify();
-														Button->SetIsOpen(false);
-														PoseSearchColumn->SwitchDatabase(Database, RowIndex);
-														Chooser->PostEditChange();
-													})));
-										}
-									}
-								}
-							}));
+											const FScopedTransaction Transaction(LOCTEXT("Set Database", "Set Database"));
+											Chooser->Modify();
+											Button->SetIsOpen(false);
+											PoseSearchColumn->SwitchDatabase(Database, RowIndex);
+											Chooser->PostEditChange();
+										})));
+							}
+						}
+					}
 				}
 
 				return MenuBuilder.MakeWidget();
