PoseSearch - fix for missing PSD indexing on reverting / reloading dependent assets

#Synced-CL 51930448

diff --git a/Engine/Plugins/Animation/PoseSearch/Source/Runtime/Private/PoseSearchDerivedData.cpp b/Engine/Plugins/Animation/PoseSearch/Source/Runtime/Private/PoseSearchDerivedData.cpp
--- a/Engine/Plugins/Animation/PoseSearch/Source/Runtime/Private/PoseSearchDerivedData.cpp
+++ b/Engine/Plugins/Animation/PoseSearch/Source/Runtime/Private/PoseSearchDerivedData.cpp
@@ -1856,22 +1856,19 @@ void FPoseSearchDatabaseAsyncCacheTask::StartNewRequestIfNeeded(bool bPerformCon
 // it cancels and waits for the task to be done and set the state to PreCancelled, so no other new requests can start until the task gets cancelled
 void FPoseSearchDatabaseAsyncCacheTask::PreCancelIfDependsOn(const UObject* Object)
 {
-	if (FKeyBuilder::ShouldInclude(Object))
-	{
-		ensure(IsInGameThread());
+	ensure(IsInGameThread());
 
-		// DatabaseDependencies is updated only in StartNewRequestIfNeeded when there are no active requests, so it's thread safe to access it 
-		if (DatabaseDependencies.Contains(Object))
-		{
-			// Database can be null if the task was Ended/Failed and Database was already garbage collected, but Tick hasn't been called yet
-			FString DatabaseName = IsValid() ? *GetDatabaseAndOuterName(Database.Get()) : TEXT("Garbage Collected Database");
-			UE_LOGF(LogPoseSearch, Log, "%ls - %ls PreCancelled because of %ls", *LexToString(DerivedDataKey), *DatabaseName, *Object->GetName());
+	// DatabaseDependencies is updated only in StartNewRequestIfNeeded when there are no active requests, so it's thread safe to access it 
+	if (DatabaseDependencies.Contains(Object))
+	{
+		// Database can be null if the task was Ended/Failed and Database was already garbage collected, but Tick hasn't been called yet
+		FString DatabaseName = IsValid() ? *GetDatabaseAndOuterName(Database.Get()) : TEXT("Garbage Collected Database");
+		UE_LOGF(LogPoseSearch, Log, "%ls - %ls PreCancelled because of %ls", *LexToString(DerivedDataKey), *DatabaseName, *Object->GetName());
 
-			// Owner.Cancel must be performed before SearchIndex.Reset() in case any task is flying (launched by Owner.LaunchTask)
-			Owner.Cancel();
+		// Owner.Cancel must be performed before SearchIndex.Reset() in case any task is flying (launched by Owner.LaunchTask)
+		Owner.Cancel();
 
-			SetState(EState::PreCancelled);
-		}
+		SetState(EState::PreCancelled);
 	}
 }
 
@@ -2415,12 +2412,12 @@ void FAsyncPoseSearchDatabasesManagement::SynchronizeDatabases()
 // we're listening to OnObjectModified to cancel any pending Task indexing databases depending from Object to avoid multi threading issues
 void FAsyncPoseSearchDatabasesManagement::OnObjectModified(UObject* Object)
 {
-	PreModified(Object);
+	PreModified(Object, FKeyBuilder::ShouldIncludeFlagsToCheck);
 }
 
-void FAsyncPoseSearchDatabasesManagement::PreModified(UObject* Object)
+void FAsyncPoseSearchDatabasesManagement::PreModified(UObject* Object, EObjectFlags ObjectFlagsToCheck)
 {
-	if (FKeyBuilder::ShouldInclude(Object))
+	if (FKeyBuilder::ShouldInclude(Object, ObjectFlagsToCheck))
 	{
 		ensure(IsInGameThread());
 
@@ -2438,9 +2435,9 @@ void FAsyncPoseSearchDatabasesManagement::PreModified(UObject* Object)
 	}
 }
 
-void FAsyncPoseSearchDatabasesManagement::PostModified(UObject* Object)
+void FAsyncPoseSearchDatabasesManagement::PostModified(UObject* Object, EObjectFlags ObjectFlagsToCheck)
 {
-	if (FKeyBuilder::ShouldInclude(Object))
+	if (FKeyBuilder::ShouldInclude(Object, ObjectFlagsToCheck))
 	{
 		ensure(IsInGameThread());
 
@@ -2464,7 +2461,7 @@ void FAsyncPoseSearchDatabasesManagement::PostModified(UObject* Object)
 
 void FAsyncPoseSearchDatabasesManagement::OnObjectTransacted(UObject* Object, const FTransactionObjectEvent& TransactionObjectEvent)
 {
-	PostModified(Object);
+	PostModified(Object, FKeyBuilder::ShouldIncludeFlagsToCheck);
 }
 
 void FAsyncPoseSearchDatabasesManagement::OnPackageReloaded(const EPackageReloadPhase InPackageReloadPhase, FPackageReloadedEvent* InPackageReloadedEvent)
@@ -2479,7 +2476,7 @@ void FAsyncPoseSearchDatabasesManagement::OnPackageReloaded(const EPackageReload
 		const TMap<UObject*, UObject*>& RepointedObjects = InPackageReloadedEvent->GetRepointedObjects();
 		for (const TPair<UObject*, UObject*>& RepointPair : RepointedObjects)
 		{
-			PreModified(RepointPair.Key);
+			PreModified(RepointPair.Key, FKeyBuilder::ShouldIncludeFlagsToCheckOnReload);
 		}
 	}
 	else if (InPackageReloadPhase == EPackageReloadPhase::PostPackageFixup)
@@ -2487,19 +2484,19 @@ void FAsyncPoseSearchDatabasesManagement::OnPackageReloaded(const EPackageReload
 		const TMap<UObject*, UObject*>& RepointedObjects = InPackageReloadedEvent->GetRepointedObjects();
 		for (const TPair<UObject*, UObject*>& RepointPair : RepointedObjects)
 		{
-			PostModified(RepointPair.Key);
+			PostModified(RepointPair.Key, FKeyBuilder::ShouldIncludeFlagsToCheckOnReload);
 		}
 	}
 }
 
 void FAsyncPoseSearchDatabasesManagement::OnPreObjectPropertyChanged(UObject* InObject, const FEditPropertyChain& InPropertyChain)
 {
-	PreModified(InObject);
+	PreModified(InObject, FKeyBuilder::ShouldIncludeFlagsToCheck);
 }
 
 void FAsyncPoseSearchDatabasesManagement::OnObjectPropertyChanged(UObject* InObject, FPropertyChangedEvent& InPropertyChangedEvent)
 {
-	PostModified(InObject);
+	PostModified(InObject, FKeyBuilder::ShouldIncludeFlagsToCheck);
 }
 
 void FAsyncPoseSearchDatabasesManagement::Shutdown()
diff --git a/Engine/Plugins/Animation/PoseSearch/Source/Runtime/Private/PoseSearchDerivedDataKey.cpp b/Engine/Plugins/Animation/PoseSearch/Source/Runtime/Private/PoseSearchDerivedDataKey.cpp
--- a/Engine/Plugins/Animation/PoseSearch/Source/Runtime/Private/PoseSearchDerivedDataKey.cpp
+++ b/Engine/Plugins/Animation/PoseSearch/Source/Runtime/Private/PoseSearchDerivedDataKey.cpp
@@ -252,7 +252,7 @@ FArchive& FKeyBuilder::TryAddDependency(UObject* Object)
 
 FArchive& FKeyBuilder::operator<<(class UObject*& Object)
 {
-	if (ShouldInclude(Object))
+	if (ShouldInclude(Object, ShouldIncludeFlagsToCheck))
 	{
 		return TryAddDependency(Object);
 	}
@@ -526,7 +526,7 @@ bool FKeyBuilder::IsExcludedType(const UObject* Object)
 	return false;
 }
 
-bool FKeyBuilder::ShouldInclude(const UObject* Object)
+bool FKeyBuilder::ShouldInclude(const UObject* Object, EObjectFlags ObjectFlagsToCheck)
 {
 	if (!Object)
 	{
@@ -538,7 +538,7 @@ bool FKeyBuilder::ShouldInclude(const UObject* Object)
 		return false;
 	}
 
-	if (Object->HasAnyFlags(RF_Transient | RF_MirroredGarbage | RF_NewerVersionExists | RF_BeginDestroyed | RF_FinishDestroyed | RF_DuplicateTransient | RF_NonPIEDuplicateTransient))
+	if (Object->HasAnyFlags(ObjectFlagsToCheck))
 	{
 		return false;
 	}
diff --git a/Engine/Plugins/Animation/PoseSearch/Source/Runtime/Public/PoseSearch/PoseSearchDerivedData.h b/Engine/Plugins/Animation/PoseSearch/Source/Runtime/Public/PoseSearch/PoseSearchDerivedData.h
--- a/Engine/Plugins/Animation/PoseSearch/Source/Runtime/Public/PoseSearch/PoseSearchDerivedData.h
+++ b/Engine/Plugins/Animation/PoseSearch/Source/Runtime/Public/PoseSearch/PoseSearchDerivedData.h
@@ -55,8 +55,8 @@ namespace UE::PoseSearch
 
 		void Shutdown();
 		void StartQueuedTasks(int32 MaxActiveTasks);
-		void PreModified(UObject* Object);
-		void PostModified(UObject* Object);
+		void PreModified(UObject* Object, EObjectFlags ObjectFlagsToCheck);
+		void PostModified(UObject* Object, EObjectFlags ObjectFlagsToCheck);
 
 		// Begin FTickableGameObject
 		virtual void Tick(float DeltaTime) override;
diff --git a/Engine/Plugins/Animation/PoseSearch/Source/Runtime/Public/PoseSearch/PoseSearchDerivedDataKey.h b/Engine/Plugins/Animation/PoseSearch/Source/Runtime/Public/PoseSearch/PoseSearchDerivedDataKey.h
--- a/Engine/Plugins/Animation/PoseSearch/Source/Runtime/Public/PoseSearch/PoseSearchDerivedDataKey.h
+++ b/Engine/Plugins/Animation/PoseSearch/Source/Runtime/Public/PoseSearch/PoseSearchDerivedDataKey.h
@@ -179,7 +179,11 @@ public:
 	POSESEARCH_API const TSet<const UObject*>& GetDependencies() const;
 
 	// Experimental, this feature might be removed without warning, not for production use
-	static POSESEARCH_API bool ShouldInclude(const UObject* Object);
+	static constexpr EObjectFlags ShouldIncludeFlagsToCheckOnReload = RF_Transient | RF_MirroredGarbage | RF_BeginDestroyed | RF_FinishDestroyed | RF_DuplicateTransient | RF_NonPIEDuplicateTransient;
+	// Experimental, this feature might be removed without warning, not for production use
+	static constexpr EObjectFlags ShouldIncludeFlagsToCheck = ShouldIncludeFlagsToCheckOnReload | RF_NewerVersionExists;
+	// Experimental, this feature might be removed without warning, not for production use
+	static POSESEARCH_API bool ShouldInclude(const UObject* Object, EObjectFlags ObjectFlagsToCheck);
 
 protected:
 	// to keep the key generation lightweight, we don't hash these types
