PoseSearch - fix for previewing game generated pose search debugger data with the editor

#Synced-CL 52146275

diff --git a/Engine/Plugins/Animation/PoseSearch/Source/Editor/Private/Trace/PoseSearchTraceAnalyzer.cpp b/Engine/Plugins/Animation/PoseSearch/Source/Editor/Private/Trace/PoseSearchTraceAnalyzer.cpp
--- a/Engine/Plugins/Animation/PoseSearch/Source/Editor/Private/Trace/PoseSearchTraceAnalyzer.cpp
+++ b/Engine/Plugins/Animation/PoseSearch/Source/Editor/Private/Trace/PoseSearchTraceAnalyzer.cpp
@@ -4,7 +4,6 @@
 
 #include "HAL/LowLevelMemTracker.h"
 #include "PoseSearchTraceProvider.h"
-#include "Serialization/MemoryReader.h"
 #include "PoseSearch/Trace/PoseSearchTraceLogger.h"
 #include "TraceServices/Model/AnalysisSession.h"
 #include "PoseSearch/PoseSearchCustomVersion.h"
@@ -32,7 +31,7 @@ bool FTraceAnalyzer::OnEvent(uint16 RouteId, EStyle Style, const FOnEventContext
 
 	TraceServices::FAnalysisSessionEditScope Scope(Session);
 
-	FCustomVersionContainer CustomVersionContainer;
+	FPoseSearchCustomVersion::Type CustomVersionType = FPoseSearchCustomVersion::BeforeCustomVersionWasAdded;
 
 	switch (RouteId)
 	{
@@ -41,11 +40,11 @@ bool FTraceAnalyzer::OnEvent(uint16 RouteId, EStyle Style, const FOnEventContext
 		break;
 	case RouteId_MotionMatchingState2:
 		// Ensure we have the version that matches this RouteId event.
-		CustomVersionContainer.SetVersion(FPoseSearchCustomVersion::GUID, FPoseSearchCustomVersion::DeprecatedTrajectoryTypes, TEXT("Dev-PoseSearch-Version"));
+		CustomVersionType = FPoseSearchCustomVersion::DeprecatedTrajectoryTypes;
 		break;
 	case RouteId_MotionMatchingState3:
 		// Ensure we have the version that matches this RouteId event.
-		CustomVersionContainer.SetVersion(FPoseSearchCustomVersion::GUID, FPoseSearchCustomVersion::AddedInterruptModeToDebugger, TEXT("Dev-PoseSearch-Version"));
+		CustomVersionType = FPoseSearchCustomVersion::AddedInterruptModeToDebugger;
 		break;
 	default:
 		// Should not happen...
@@ -54,9 +53,7 @@ bool FTraceAnalyzer::OnEvent(uint16 RouteId, EStyle Style, const FOnEventContext
 	}
 
 	FTraceMotionMatchingStateMessage Message;
-	FMemoryReaderView Archive(Context.EventData.GetArrayView<uint8>("Data"));
-	Archive.SetCustomVersions(CustomVersionContainer);
-	Archive << Message;
+	Message.Input(CustomVersionType, Context.EventData.GetArrayView<uint8>("Data"));
 	TraceProvider.AppendMotionMatchingState(Message, Context.EventTime.AsSeconds(Message.Cycle));
 
 	return true;
diff --git a/Engine/Plugins/Animation/PoseSearch/Source/Runtime/Private/PoseSearchHistory.cpp b/Engine/Plugins/Animation/PoseSearch/Source/Runtime/Private/PoseSearchHistory.cpp
--- a/Engine/Plugins/Animation/PoseSearch/Source/Runtime/Private/PoseSearchHistory.cpp
+++ b/Engine/Plugins/Animation/PoseSearch/Source/Runtime/Private/PoseSearchHistory.cpp
@@ -445,8 +445,14 @@ FArchive& operator<<(FArchive& Ar, FPoseHistoryEntry& Entry)
 	Ar << Entry.CurveValues;
 	Ar << Entry.AccumulatedSeconds;
 
+	// this method is used by rewind debugger to archive data for pose search debugger and shared across game and editor.
+	// when captures come from the game DebugMeshTransform will be invalid and displaying incorrectly when previewed by the editor
+	// @todo: should we make DebugMeshTransform !WITH_EDITORONLY_DATA?
 #if WITH_EDITORONLY_DATA
 	Ar << Entry.DebugMeshTransform;
+#else // WITH_EDITORONLY_DATA
+	FTransform DebugMeshTransform = FTransform::Identity;
+	Ar << DebugMeshTransform;
 #endif // WITH_EDITORONLY_DATA
 
 	return Ar;
diff --git a/Engine/Plugins/Animation/PoseSearch/Source/Runtime/Private/Trace/PoseSearchTraceLogger.cpp b/Engine/Plugins/Animation/PoseSearch/Source/Runtime/Private/Trace/PoseSearchTraceLogger.cpp
--- a/Engine/Plugins/Animation/PoseSearch/Source/Runtime/Private/Trace/PoseSearchTraceLogger.cpp
+++ b/Engine/Plugins/Animation/PoseSearch/Source/Runtime/Private/Trace/PoseSearchTraceLogger.cpp
@@ -4,8 +4,8 @@
 #include "Animation/AnimInstanceProxy.h"
 #include "Animation/AnimNodeBase.h"
 #include "Components/SkeletalMeshComponent.h"
-#include "PoseSearch/PoseSearchCustomVersion.h"
 #include "PoseSearch/PoseSearchDatabase.h"
+#include "Serialization/MemoryReader.h"
 #include "Serialization/MemoryWriter.h"
 #include "Trace/Trace.inl"
 #include "TraceFilter.h"
@@ -92,6 +92,23 @@ FArchive& operator<<(FArchive& Ar, FTraceMotionMatchingStateMessage& State)
 	return Ar;
 }
 
+void FTraceMotionMatchingStateMessage::Input(FPoseSearchCustomVersion::Type CustomVersionType, TConstArrayView<uint8> Data)
+{
+	FCustomVersionContainer CustomVersionContainer;
+
+	FMemoryReaderView Archive(Data);
+	if (CustomVersionType != FPoseSearchCustomVersion::BeforeCustomVersionWasAdded)
+	{
+		// Ensure we have the version that matches this RouteId event.
+		CustomVersionContainer.SetVersion(FPoseSearchCustomVersion::GUID, CustomVersionType, TEXT("Dev-PoseSearch-Version"));
+	}
+
+	Archive.SetCustomVersions(CustomVersionContainer);
+	// setting SetFilterEditorOnly to have archived data sharable between game and editor
+	Archive.SetFilterEditorOnly(true);
+	Archive << *this;
+}
+
 void FTraceMotionMatchingStateMessage::Output()
 {
 #if OBJECT_TRACE_ENABLED
@@ -102,7 +119,8 @@ void FTraceMotionMatchingStateMessage::Output()
 	FCustomVersionContainer CustomVersionContainer;
 	CustomVersionContainer.SetVersion(FPoseSearchCustomVersion::GUID, FPoseSearchCustomVersion::LatestVersion, TEXT("Dev-PoseSearch-Version"));
 	Archive.SetCustomVersions(CustomVersionContainer);
-	
+	// setting SetFilterEditorOnly to have archived data sharable between game and editor
+	Archive.SetFilterEditorOnly(true);
 	Archive << *this;
 	UE_TRACE_LOG(PoseSearch, MotionMatchingState3, PoseSearchChannel) << MotionMatchingState3.Data(ArchiveData.GetData(), ArchiveData.Num());
 #endif
diff --git a/Engine/Plugins/Animation/PoseSearch/Source/Runtime/Public/PoseSearch/PoseSearchCost.h b/Engine/Plugins/Animation/PoseSearch/Source/Runtime/Public/PoseSearch/PoseSearchCost.h
--- a/Engine/Plugins/Animation/PoseSearch/Source/Runtime/Public/PoseSearch/PoseSearchCost.h
+++ b/Engine/Plugins/Animation/PoseSearch/Source/Runtime/Public/PoseSearch/PoseSearchCost.h
@@ -46,11 +46,28 @@ public:
 	{
 		Ar << PoseSearchCost.TotalCost;
 
+		// this method is used by rewind debugger to archive data for pose search debugger and shared across game and editor.
+		// when captures come from the game the cost addend breakdowns will be invalid and displaying incorrect information when previewed by the editor
+		// @todo: should we make *CostAddend !WITH_EDITORONLY_DATA?
 #if WITH_EDITORONLY_DATA
+		
 		Ar << PoseSearchCost.NotifyCostAddend;
 		Ar << PoseSearchCost.ContinuingPoseCostAddend;
 		Ar << PoseSearchCost.ContinuingInteractionCostAddend;
 		Ar << PoseSearchCost.ContinuingContextInteractionCostAddend;
+
+#else // WITH_EDITORONLY_DATA
+		
+		float NotifyCostAddend = 0.f;
+		float ContinuingPoseCostAddend = 0.f;
+		float ContinuingInteractionCostAddend = 0.f;
+		float ContinuingContextInteractionCostAddend = 0.f;
+
+		Ar << NotifyCostAddend;
+		Ar << ContinuingPoseCostAddend;
+		Ar << ContinuingInteractionCostAddend;
+		Ar << ContinuingContextInteractionCostAddend;
+
 #endif // WITH_EDITORONLY_DATA
 			
 		return Ar;
diff --git a/Engine/Plugins/Animation/PoseSearch/Source/Runtime/Public/PoseSearch/Trace/PoseSearchTraceLogger.h b/Engine/Plugins/Animation/PoseSearch/Source/Runtime/Public/PoseSearch/Trace/PoseSearchTraceLogger.h
--- a/Engine/Plugins/Animation/PoseSearch/Source/Runtime/Public/PoseSearch/Trace/PoseSearchTraceLogger.h
+++ b/Engine/Plugins/Animation/PoseSearch/Source/Runtime/Public/PoseSearch/Trace/PoseSearchTraceLogger.h
@@ -7,6 +7,7 @@
 #include "ObjectTrace.h"
 #endif
 #include "PoseSearch/PoseSearchContext.h"
+#include "PoseSearch/PoseSearchCustomVersion.h"
 #include "PoseSearch/PoseSearchLibrary.h"
 
 #define UE_API POSESEARCH_API
@@ -105,7 +106,9 @@ struct FTraceMotionMatchingStateMessage : public FTraceMessage
 
 	EPoseSearchInterruptMode InterruptMode = EPoseSearchInterruptMode::DoNotInterrupt;
 	
-	/** Output the current state info to the logger */
+	// Input the current state from EventData
+	UE_API void Input(FPoseSearchCustomVersion::Type CustomVersionType, TConstArrayView<uint8> Data);
+	// Output the current state info to the logger
 	UE_API void Output();
 
 	UE_DEPRECATED(5.7, "This API will soon be removed. Look for a FTraceMotionMatchingStateDatabaseEntry::PoseEntries with PoseCandidateFlags containg EPoseCandidateFlags::Valid_ContinuingPose for the 'current database and pose index' instead")
