PoseSearch - adding BlueprintType FPoseSearchHistory to be able to copy pose historyies across UAF, and PB to mitigate the thread safety issues

#Synced-CL 52175710

diff --git a/Engine/Plugins/Animation/PoseSearch/Source/Runtime/Private/Chooser/PoseSearchChooserColumn.cpp b/Engine/Plugins/Animation/PoseSearch/Source/Runtime/Private/Chooser/PoseSearchChooserColumn.cpp
--- a/Engine/Plugins/Animation/PoseSearch/Source/Runtime/Private/Chooser/PoseSearchChooserColumn.cpp
+++ b/Engine/Plugins/Animation/PoseSearch/Source/Runtime/Private/Chooser/PoseSearchChooserColumn.cpp
@@ -82,6 +82,7 @@ namespace UE::PoseSearch
 		TArray<FSearchResult> SearchResults;
 		const IPoseHistory* PoseHistory = nullptr;
 		const FChooserPlayerSettings* ChooserPlayerSettings = nullptr;
+		TSharedPtr<UE::PoseSearch::FPoseHistory, ESPMode::ThreadSafe> PoseHistoryPtr;
 
 #if DO_CHECK
 		const FPoseSearchColumn* DebugOwner = nullptr;
@@ -693,22 +694,29 @@ void FPoseSearchColumn::Filter(FChooserEvaluationContext& Context, const FChoose
 	ensure(PoseSearchColumnScratchArea->SearchResults.IsEmpty());
 #endif // DO_CHECK
 
-	const UE::PoseSearch::IPoseHistory* PoseHistory = nullptr;
 	FPoseHistoryReference PoseHistoryReference;
 	if (const FChooserParameterPoseHistoryBase* PoseHistoryParameter = InputValue.GetPtr<FChooserParameterPoseHistoryBase>())
 	{
+		FPoseSearchHistory PoseSearchHistory;
 		if (PoseHistoryParameter->GetValue(Context, PoseHistoryReference))
 		{
-			PoseHistory = PoseHistoryReference.PoseHistory.Get();
+			PoseSearchColumnScratchArea->PoseHistory = PoseHistoryReference.PoseHistory.Get();
+		}
+		else if (PoseHistoryParameter->GetValue(Context, PoseSearchHistory))
+		{
+			// @todo: avoid allocations...
+			PoseSearchColumnScratchArea->PoseHistoryPtr = MakeShareable(new FPoseHistory());
+			PoseSearchColumnScratchArea->PoseHistoryPtr->EditPoseSearchHistory() = MoveTemp(PoseSearchHistory);
+			PoseSearchColumnScratchArea->PoseHistory = PoseSearchColumnScratchArea->PoseHistoryPtr.Get();
 		}
 	}
 
 	PoseSearchColumnScratchArea->ChooserPlayerSettings = GetChooserPlayerSettings(Context);
-	if (!PoseHistory && PoseSearchColumnScratchArea->ChooserPlayerSettings && PoseSearchColumnScratchArea->ChooserPlayerSettings->AnimationUpdateContext)
+	if (!PoseSearchColumnScratchArea->PoseHistory && PoseSearchColumnScratchArea->ChooserPlayerSettings && PoseSearchColumnScratchArea->ChooserPlayerSettings->AnimationUpdateContext)
 	{
 		if (const FPoseHistoryProvider* PoseHistoryProvider = PoseSearchColumnScratchArea->ChooserPlayerSettings->AnimationUpdateContext->GetMessage<FPoseHistoryProvider>())
 		{
-			PoseHistory = PoseHistoryProvider->GetPoseHistoryPtr();
+			PoseSearchColumnScratchArea->PoseHistory = PoseHistoryProvider->GetPoseHistoryPtr();
 		}
 	}
 
@@ -716,14 +724,12 @@ void FPoseSearchColumn::Filter(FChooserEvaluationContext& Context, const FChoose
 	TArray<FActiveColumnCost> TracedActiveColumnCosts;
 #endif // WITH_EDITOR
 
-	if (!PoseHistory)
+	if (!PoseSearchColumnScratchArea->PoseHistory)
 	{
-		UE_LOGF(LogPoseSearch, Error, "FPoseSearchColumn::Filter, missing IPoseHistory");
+		UE_LOGF(LogPoseSearch, Error, "FPoseSearchColumn::Filter, missing IPoseHistory when filtering chooser table '%ls' with root '%ls'", *Chooser->GetName(), *Chooser->GetRootChooser()->GetName());
 	}
 	else
 	{
-		PoseSearchColumnScratchArea->PoseHistory = PoseHistory;
-
 		ensure(Chooser);
 		
 		FStackDatabaseToAssetIndexes AssetIndexesToConsiderPerDatabase;
@@ -760,8 +766,7 @@ void FPoseSearchColumn::Filter(FChooserEvaluationContext& Context, const FChoose
 #if !NO_LOGGING
 					else if (Role != Pair.Key->Schema->GetDefaultRole())
 					{
-						// @todo: improve logging
-						UE_LOGF(LogPoseSearch, Error, "FPoseSearchColumn::Filter inconsistent Roles");
+						UE_LOGF(LogPoseSearch, Error, "FPoseSearchColumn::Filter inconsistent Roles across the databases referenced by chooser table '%ls' with root '%ls' (expected '%ls', found '%ls')", *Chooser->GetName(), *Chooser->GetRootChooser()->GetName(), *Role.ToString(), *Pair.Key->Schema->GetDefaultRole().ToString());
 					}
 #endif // !NO_LOGGING
 					DatabasesToSearch.Add(Pair.Key);
@@ -772,7 +777,7 @@ void FPoseSearchColumn::Filter(FChooserEvaluationContext& Context, const FChoose
 			const FFloatInterval PoseJumpThresholdTime(0.f, 0.f);
 			FSearchContext SearchContext(0.f, PoseJumpThresholdTime, FPoseSearchEvent());
 
-			SearchContext.AddRole(Role, &Context, PoseHistory);
+			SearchContext.AddRole(Role, &Context, PoseSearchColumnScratchArea->PoseHistory);
 			SearchContext.SetAssetIndexesToConsiderPerDatabase(&AssetIndexesToConsiderPerDatabase);
 
 			FPoseSearchContinuingProperties ContinuingProperties;
diff --git a/Engine/Plugins/Animation/PoseSearch/Source/Runtime/Private/PoseSearchHistory.cpp b/Engine/Plugins/Animation/PoseSearch/Source/Runtime/Private/PoseSearchHistory.cpp
--- a/Engine/Plugins/Animation/PoseSearch/Source/Runtime/Private/PoseSearchHistory.cpp
+++ b/Engine/Plugins/Animation/PoseSearch/Source/Runtime/Private/PoseSearchHistory.cpp
@@ -714,19 +714,8 @@
 		UE_MT_SCOPED_WRITE_ACCESS(PoseDataThreadSafeCounter);
 #endif // ENABLE_ANIM_DEBUG
 
-		MaxNumPoses = Other.MaxNumPoses;
-		SamplingInterval = Other.SamplingInterval;
-
-		Trajectory = Other.Trajectory;
+		PoseSearchHistory = Other.PoseSearchHistory;
 		TrajectoryDataState = Other.TrajectoryDataState;
-		TrajectorySpeedMultiplier = Other.TrajectorySpeedMultiplier;
-
-		LastUpdateSkeleton = Other.LastUpdateSkeleton;
-		BoneToTransformMap = Other.BoneToTransformMap;
-		CollectedCurves = Other.CollectedCurves;
-		BoneToTransformMapTypeHash  = Other.BoneToTransformMapTypeHash;
-		Entries = Other.Entries;
-		PoseIndicesHistory = Other.PoseIndicesHistory;
 	}
 	return *this;
 }
@@ -740,19 +729,8 @@
 		UE_MT_SCOPED_WRITE_ACCESS(PoseDataThreadSafeCounter);
 #endif // ENABLE_ANIM_DEBUG
 
-		MaxNumPoses = MoveTemp(Other.MaxNumPoses);
-		SamplingInterval = MoveTemp(Other.SamplingInterval);
-
-		Trajectory = MoveTemp(Other.Trajectory);
+		PoseSearchHistory = MoveTemp(Other.PoseSearchHistory);
 		TrajectoryDataState = MoveTemp(Other.TrajectoryDataState);
-		TrajectorySpeedMultiplier = MoveTemp(Other.TrajectorySpeedMultiplier);
-
-		LastUpdateSkeleton = MoveTemp(Other.LastUpdateSkeleton);
-		BoneToTransformMap = MoveTemp(Other.BoneToTransformMap);
-		CollectedCurves = MoveTemp(Other.CollectedCurves);
-		BoneToTransformMapTypeHash  = MoveTemp(Other.BoneToTransformMapTypeHash);
-		Entries = MoveTemp(Other.Entries);
-		PoseIndicesHistory = MoveTemp(Other.PoseIndicesHistory);
 	}
 	return *this;
 }
@@ -760,181 +738,92 @@
 void FPoseHistory::Initialize_AnyThread(int32 InNumPoses, float InSamplingInterval)
 {
 	UE_MT_SCOPED_WRITE_ACCESS(PoseDataThreadSafeCounter);
-	check(InNumPoses >= 2);
-
-	MaxNumPoses = InNumPoses;
-	SamplingInterval = InSamplingInterval;
-	
-	Trajectory = FTransformTrajectory();
+	PoseSearchHistory.Initialize(InNumPoses, InSamplingInterval);
 	TrajectoryDataState = FPoseSearchTrajectoryData::FState();
-	TrajectorySpeedMultiplier = 1.f;
-
-	LastUpdateSkeleton = nullptr;
-	BoneToTransformMap.Reset();
-	CollectedCurves.Reset();
-	BoneToTransformMapTypeHash = 0;
-	Entries.Reset();
-	PoseIndicesHistory.Reset();
 }
 
 bool FPoseHistory::GetTransformAtTime(float Time, FTransform& OutBoneTransform, const USkeleton* BoneIndexSkeleton, FBoneIndexType BoneIndexType, FBoneIndexType ReferenceBoneIndexType, bool bExtrapolate) const
 {
 	UE_MT_SCOPED_READ_ACCESS(PoseDataThreadSafeCounter);
-
-	static_assert(RootBoneIndexType == 0 && ComponentSpaceIndexType == FBoneIndexType(-1) && WorldSpaceIndexType == FBoneIndexType(-2)); // some assumptions
-	
-	bool bSuccess = true;
-	if (BoneIndexType == ReferenceBoneIndexType)
-	{
-		OutBoneTransform = FTransform::Identity;
-		bSuccess = true;
-	}
-	else if (ReferenceBoneIndexType == WorldSpaceIndexType)
-	{
-		if (BoneIndexType == ComponentSpaceIndexType)
-		{
-			OutBoneTransform = Trajectory.GetSampleAtTime(Time, bExtrapolate).GetTransform();
-			bSuccess = true;
-		}
-		else
-		{
-			// getting BoneIndexType in ComponentSpaceIndexType and then multiplying it with the component to world transform from the trajectory
-			bSuccess = GetTransformAtTime(Time, OutBoneTransform, BoneIndexSkeleton, BoneIndexType, ComponentSpaceIndexType, bExtrapolate);
-			OutBoneTransform *= Trajectory.GetSampleAtTime(Time, bExtrapolate).GetTransform();
-		}
-	}
-	else if (BoneIndexType == WorldSpaceIndexType || BoneIndexType == ComponentSpaceIndexType)
-	{
-		bSuccess = GetTransformAtTime(Time, OutBoneTransform, BoneIndexSkeleton, ReferenceBoneIndexType, BoneIndexType, bExtrapolate);
-		OutBoneTransform = OutBoneTransform.Inverse();
-	}
-	else
-	{
-		check(BoneIndexType != ComponentSpaceIndexType && BoneIndexType != WorldSpaceIndexType && ReferenceBoneIndexType != WorldSpaceIndexType);
-
-		const int32 NumEntries = Entries.Num();
-		if (NumEntries > 0)
-		{
-			int32 NextIdx = 0;
-			int32 PrevIdx = 0;
-
-			if (NumEntries > 1)
-			{
-				const int32 LowerBoundIdx = LowerBound(Entries.begin(), Entries.end(), Time, [](const FPoseHistoryEntry& Entry, float Value) { return Value > Entry.AccumulatedSeconds; });
-				NextIdx = FMath::Clamp(LowerBoundIdx, 1, NumEntries - 1);
-				PrevIdx = NextIdx - 1;
-			}
-
-			const FPoseHistoryEntry& PrevEntry = Entries[PrevIdx];
-			const FPoseHistoryEntry& NextEntry = Entries[NextIdx];
-
-			bSuccess = LerpEntries(Time, bExtrapolate, PrevEntry, NextEntry, BoneIndexSkeleton, LastUpdateSkeleton.Get(), BoneToTransformMap, BoneIndexType, ReferenceBoneIndexType, OutBoneTransform);
-		}
-		else
-		{
-			OutBoneTransform = FTransform::Identity;
-			bSuccess = false;
-		}
-	}
-
-	// @todo: reenable this logging after implementing AnimNext PoseHistory initialization (currently spamming on actor spawning with MM active)
-//#if WITH_EDITOR
-//	// @todo: logging errors only if pose history was properly initialized for now (since AnimNext version doesn't implement initialization with ref pose yet),
-//	// and we ultimately want to catch errors of missing bone transforms only checking valid USkeleton(s) against valid PoseSearchSchema(s)
-//	if (!bSuccess && BoneIndexSkeleton)
-//	{
-//		UE_LOGF(LogPoseSearch, Error, "FPoseHistory::GetTransformAtTime - Couldn't find root bone transform at time %f for Skeleton %ls!", Time, *BoneIndexSkeleton->GetName());
-//	}
-//#endif // WITH_EDITOR
-
-	return bSuccess;
+	return PoseSearchHistory.GetTransformAtTime(Time, OutBoneTransform, BoneIndexSkeleton, BoneIndexType, ReferenceBoneIndexType, bExtrapolate);
 }
 
 bool FPoseHistory::GetCurveValueAtTime(float Time, const FName& CurveName, float& OutCurveValue, bool bExtrapolate /*= false*/) const
 {
 	UE_MT_SCOPED_READ_ACCESS(PoseDataThreadSafeCounter);
-
-	const int32 NumEntries = Entries.Num();
-	if (NumEntries > 0)
-	{
-		int32 NextIdx = 0;
-		int32 PrevIdx = 0;
-
-		if (NumEntries > 1)
-		{
-			const int32 LowerBoundIdx = LowerBound(Entries.begin(), Entries.end(), Time, [](const FPoseHistoryEntry& Entry, float Value) { return Value > Entry.AccumulatedSeconds; });
-			NextIdx = FMath::Clamp(LowerBoundIdx, 1, NumEntries - 1);
-			PrevIdx = NextIdx - 1;
-		}
-
-		const FPoseHistoryEntry& PrevEntry = Entries[PrevIdx];
-		const FPoseHistoryEntry& NextEntry = Entries[NextIdx];
-
-		return LerpEntries(Time, bExtrapolate, PrevEntry, NextEntry, CurveName, CollectedCurves, OutCurveValue);
-	}
-	
-	OutCurveValue = 0.0f;
-	return false;
+	return PoseSearchHistory.GetCurveValueAtTime(Time, CurveName, OutCurveValue, bExtrapolate);
 }
 
 const FTransformTrajectory& FPoseHistory::GetTrajectory() const
 {
 	UE_MT_SCOPED_READ_ACCESS(PoseDataThreadSafeCounter);
-	return Trajectory;
+	return PoseSearchHistory.GetTrajectory();
 }
 
 float FPoseHistory::GetTrajectorySpeedMultiplier() const
 {
 	UE_MT_SCOPED_READ_ACCESS(PoseDataThreadSafeCounter);
-	return TrajectorySpeedMultiplier;
+	return PoseSearchHistory.GetTrajectorySpeedMultiplier();
 }
 
 bool FPoseHistory::IsEmpty() const
 {
 	UE_MT_SCOPED_READ_ACCESS(PoseDataThreadSafeCounter);
-	return Entries.IsEmpty();
+	return PoseSearchHistory.IsEmpty();
 }
 
 const FBoneToTransformMap& FPoseHistory::GetBoneToTransformMap() const
 {
 	UE_MT_SCOPED_READ_ACCESS(PoseDataThreadSafeCounter);
-	return BoneToTransformMap;
+	return PoseSearchHistory.GetBoneToTransformMap();
 }
 
 const TConstArrayView<FName> FPoseHistory::GetCollectedCurves() const
 {
 	UE_MT_SCOPED_READ_ACCESS(PoseDataThreadSafeCounter);
-	return MakeConstArrayView(CollectedCurves);
+	return PoseSearchHistory.GetCollectedCurves();
 }
 
 int32 FPoseHistory::GetNumEntries() const
 {
 	UE_MT_SCOPED_READ_ACCESS(PoseDataThreadSafeCounter);
-	return Entries.Num();
+	return PoseSearchHistory.GetNumEntries();
 }
 
 const FPoseHistoryEntry& FPoseHistory::GetEntry(int32 EntryIndex) const
 {
 	UE_MT_SCOPED_READ_ACCESS(PoseDataThreadSafeCounter);
-	return Entries[EntryIndex];
+	return PoseSearchHistory.GetEntry(EntryIndex);
 }
 
 const FPoseIndicesHistory* FPoseHistory::GetPoseIndicesHistory() const
 {
 	UE_MT_SCOPED_READ_ACCESS(PoseDataThreadSafeCounter);
-	return &PoseIndicesHistory;
+	return PoseSearchHistory.GetPoseIndicesHistory();
 }
 
 int32 FPoseHistory::GetMaxNumPoses() const
 {
 	UE_MT_SCOPED_READ_ACCESS(PoseDataThreadSafeCounter);
-	return MaxNumPoses;
+	return PoseSearchHistory.GetMaxNumPoses();
 }
 
 float FPoseHistory::GetSamplingInterval() const
 {
 	UE_MT_SCOPED_READ_ACCESS(PoseDataThreadSafeCounter);
-	return SamplingInterval;
+	return PoseSearchHistory.GetSamplingInterval();
 }
 
+const FPoseSearchHistory& FPoseHistory::GetPoseSearchHistory() const
+{
+	UE_MT_SCOPED_READ_ACCESS(PoseDataThreadSafeCounter);
+	return PoseSearchHistory;
+}
+
+FPoseSearchHistory& FPoseHistory::EditPoseSearchHistory()
+{
+	UE_MT_SCOPED_WRITE_ACCESS(PoseDataThreadSafeCounter);
+	return PoseSearchHistory;
+}
+
 void FPoseHistory::GenerateTrajectory(const UObject* AnimContext, float DeltaTime, const FPoseSearchTrajectoryData& TrajectoryData, const FPoseSearchTrajectoryData::FSampling& TrajectoryDataSampling)
@@ -943,9 +832,9 @@
 	FPoseSearchTrajectoryData::FDerived TrajectoryDataDerived;
 	if (TrajectoryData.UpdateData(DeltaTime, AnimContext, TrajectoryDataDerived, TrajectoryDataState))
 	{
-		UPoseSearchTrajectoryLibrary::InitTrajectorySamples(Trajectory, TrajectoryDataDerived.Position, TrajectoryDataDerived.Facing, TrajectoryDataSampling, DeltaTime);
-		UPoseSearchTrajectoryLibrary::UpdateHistory_TransformHistory(Trajectory, TrajectoryDataDerived.Position, TrajectoryDataDerived.Velocity, TrajectoryDataSampling, DeltaTime);
-		UPoseSearchTrajectoryLibrary::UpdatePrediction_SimulateCharacterMovement(Trajectory, TrajectoryData, TrajectoryDataDerived, TrajectoryDataSampling, DeltaTime);
+		UPoseSearchTrajectoryLibrary::InitTrajectorySamples(PoseSearchHistory.EditTrajectory(), TrajectoryDataDerived.Position, TrajectoryDataDerived.Facing, TrajectoryDataSampling, DeltaTime);
+		UPoseSearchTrajectoryLibrary::UpdateHistory_TransformHistory(PoseSearchHistory.EditTrajectory(), TrajectoryDataDerived.Position, TrajectoryDataDerived.Velocity, TrajectoryDataSampling, DeltaTime);
+		UPoseSearchTrajectoryLibrary::UpdatePrediction_SimulateCharacterMovement(PoseSearchHistory.EditTrajectory(), TrajectoryData, TrajectoryDataDerived, TrajectoryDataSampling, DeltaTime);
 
 		// @todo: support TrajectorySpeedMultiplier
 		//TrajectorySpeedMultiplier = 1.f;
@@ -962,20 +851,7 @@
 	if (!InTrajectory.Samples.IsEmpty())
 	{
 		UE_MT_SCOPED_WRITE_ACCESS(PoseDataThreadSafeCounter);
-
-		// UE_MT_SCOPED_WRITE_ACCESS will assert in case of improper usage
-		Trajectory = InTrajectory;
-
-		TrajectorySpeedMultiplier = InTrajectorySpeedMultiplier;
-
-		if (!FMath::IsNearlyEqual(TrajectorySpeedMultiplier, 1.f))
-		{
-			const float TrajectorySpeedMultiplierInv = FMath::IsNearlyZero(TrajectorySpeedMultiplier) ? 1.f : 1.f / TrajectorySpeedMultiplier;
-			for (FTransformTrajectorySample& Sample : Trajectory.Samples)
-			{
-				Sample.TimeInSeconds *= TrajectorySpeedMultiplierInv;
-			}
-		}
+		PoseSearchHistory.SetTrajectory(InTrajectory, InTrajectorySpeedMultiplier);
 	}
 }
 
@@ -988,7 +864,7 @@
 #endif //WITH_EDITORONLY_DATA
 	)
 {
-	if (!ensure(MaxNumPoses >= 2))
+	if (!ensure(PoseSearchHistory.GetMaxNumPoses() >= 2))
 	{
 		UE_LOGF(LogPoseSearch, Error, "FPoseHistory::EvaluateComponentSpace_AnyThread called before FPoseHistory::Initialize_AnyThread! Please double check the implementation of any FAnimNode_Base calling this API, correctly propagating Initialize_AnyThread to their FPoseLink(s)! To mitigate issues FPoseHistory will be set to to default values");
 		Initialize_AnyThread(2, 0.f);
@@ -996,103 +872,12 @@
 
 	UE_MT_SCOPED_WRITE_ACCESS(PoseDataThreadSafeCounter);
 
-	const USkeleton* Skeleton = ComponentSpacePoseProvider.GetSkeletonAsset();
-
-	if (bCacheBones)
-	{
-		const uint32 OldBoneToTransformMapTypeHash = BoneToTransformMapTypeHash;
-
-		BoneToTransformMap.Reset();
-		CollectedCurves = InCollectedCurves;
-		if (!RequiredBones.IsEmpty())
-		{
-			// making sure we always collect the root bone transform (by construction BoneToTransformMap[0] = 0)
-			const FComponentSpaceTransformIndex ComponentSpaceTransformRootBoneIndex = 0;
-			BoneToTransformMap.Add(RootBoneIndexType) = ComponentSpaceTransformRootBoneIndex;
-
-			for (int32 i = 0; i < RequiredBones.Num(); ++i)
-			{
-				// adding only unique RequiredBones to avoid oversizing Entries::ComponentSpaceTransforms
-				if (!BoneToTransformMap.Find(RequiredBones[i]))
-				{
-					const FComponentSpaceTransformIndex ComponentSpaceTransformIndex = BoneToTransformMap.Num();
-					BoneToTransformMap.Add(RequiredBones[i]) = ComponentSpaceTransformIndex;
-				}
-			}
-		}
-
-		BoneToTransformMapTypeHash = GetTypeHash(BoneToTransformMap);
-		bNeedsReset |= (OldBoneToTransformMapTypeHash != BoneToTransformMapTypeHash);
-	}
-
-	if (LastUpdateSkeleton != Skeleton)
-	{
-		bNeedsReset = true;
-		LastUpdateSkeleton = Skeleton;
-	}
-
-	if (bNeedsReset)
-	{
-		Entries.Reset();
-		Entries.Reserve(MaxNumPoses);
-	}
-
-	FPoseHistoryEntry FutureEntryTemp;
-	if (!Entries.IsEmpty() && Entries.Last().AccumulatedSeconds > 0.f)
-	{
-		// removing the "future" root bone Entry
-		FutureEntryTemp = MoveTemp(Entries.Last());
-		Entries.Pop();
-	}
-
-	// Age our elapsed times
-	for (FPoseHistoryEntry& Entry : Entries)
-	{
-		Entry.AccumulatedSeconds -= DeltaTime;
-	}
-
-	if (Entries.Num() != MaxNumPoses)
-	{
-		// Consume every pose until the queue is full
-		Entries.Emplace();
-	}
-	else if (SamplingInterval <= 0.f || Entries[Entries.Num() - 2].AccumulatedSeconds <= -SamplingInterval)
-	{
-		FPoseHistoryEntry EntryTemp = MoveTemp(Entries.First());
-		Entries.PopFront();
-		Entries.Emplace(MoveTemp(EntryTemp));
-	}
-
-	// Regardless of the retention policy, we always update the most recent Entry
-	FPoseHistoryEntry& MostRecentEntry = Entries.Last();
-	MostRecentEntry.Update(0.f, ComponentSpacePoseProvider, BoneToTransformMap, bStoreScales, Curves, CollectedCurves
+	PoseSearchHistory.EvaluateComponentSpace(DeltaTime, ComponentSpacePoseProvider, bStoreScales, RootBoneRecoveryTime
+		, RootBoneTranslationRecoveryRatio, RootBoneRotationRecoveryRatio, bNeedsReset, bCacheBones, RequiredBones, Curves, InCollectedCurves
 #if WITH_EDITORONLY_DATA
 		, InDebugMeshTransform
 #endif //WITH_EDITORONLY_DATA
 	);
-
-	if (RootBoneRecoveryTime > 0.f && !Trajectory.Samples.IsEmpty())
-	{
-		// adding the updated "future" root bone Entry
-		const FTransform& RefRootBone = Skeleton->GetReferenceSkeleton().GetRefBonePose()[RootBoneIndexType];
-		const FQuat RootBoneRotationAtRecoveryTime = FMath::Lerp(FQuat(MostRecentEntry.ComponentSpaceRotations[RootBoneIndexType]), RefRootBone.GetRotation(), RootBoneRotationRecoveryRatio);
-
-		FVector RootBoneDeltaTranslationAtRecoveryTime = FVector::ZeroVector;
-		if (RootBoneTranslationRecoveryRatio > 0.f)
-		{
-			const FTransform WorldRootAtCurrentTime = Trajectory.GetSampleAtTime(0.f).GetTransform();
-			const FTransform WorldRootBoneAtCurrentTime = MostRecentEntry.GetComponentSpaceTransform(RootBoneIndexType) * WorldRootAtCurrentTime;
-			const FVector WorldRootBoneDeltaTranslationAtCurrentTime = (WorldRootBoneAtCurrentTime.GetTranslation() - WorldRootAtCurrentTime.GetTranslation()) * RootBoneTranslationRecoveryRatio;
-			const FTransform WorldRootAtRecoveryTime = Trajectory.GetSampleAtTime(RootBoneRecoveryTime).GetTransform();
-			RootBoneDeltaTranslationAtRecoveryTime = WorldRootAtRecoveryTime.InverseTransformVector(WorldRootBoneDeltaTranslationAtCurrentTime);
-		}
-
-		const FTransform RootBoneTransformAtRecoveryTime(RootBoneRotationAtRecoveryTime, RootBoneDeltaTranslationAtRecoveryTime, RefRootBone.GetScale3D());
-		FutureEntryTemp.SetNum(1, bStoreScales);
-		FutureEntryTemp.SetComponentSpaceTransform(RootBoneIndexType, RootBoneTransformAtRecoveryTime);
-		FutureEntryTemp.AccumulatedSeconds = RootBoneRecoveryTime;
-		Entries.Emplace(MoveTemp(FutureEntryTemp));
-	}
 }
 
 #if ENABLE_DRAW_DEBUG && ENABLE_ANIM_DEBUG
@@ -1105,17 +890,17 @@
 		const float DebugThickness = GVarAnimPoseHistoryDebugDrawTrajectoryThickness;
 		const int MaxHistorySamples = GVarAnimPoseHistoryDebugDrawTrajectoryMaxNumOfHistorySamples;
 		const int MaxPredictionSamples = GVarAnimPoseHistoryDebugDrawTrajectoryMaxNumOfPredictionSamples;
-		UTransformTrajectoryBlueprintLibrary::DebugDrawTrajectory(Trajectory, AnimInstanceProxy, DebugThickness, 0, MaxHistorySamples, MaxPredictionSamples);
+		UTransformTrajectoryBlueprintLibrary::DebugDrawTrajectory(PoseSearchHistory.GetTrajectory(), AnimInstanceProxy, DebugThickness, 0, MaxHistorySamples, MaxPredictionSamples);
 	}
 
 	if (Color.A > 0 && GVarAnimPoseHistoryDebugDrawPose)
 	{
-		const bool bValidTrajectory = !Trajectory.Samples.IsEmpty();
+		const bool bValidTrajectory = !PoseSearchHistory.GetTrajectory().Samples.IsEmpty();
 		TArray<FTransform, TInlineAllocator<128>> PrevGlobalTransforms;
 
-		for (int32 EntryIndex = 0; EntryIndex < Entries.Num(); ++EntryIndex)
+		for (int32 EntryIndex = 0; EntryIndex < PoseSearchHistory.GetNumEntries(); ++EntryIndex)
 		{
-			const FPoseHistoryEntry& Entry = Entries[EntryIndex];
+			const FPoseHistoryEntry& Entry = PoseSearchHistory.GetEntry(EntryIndex);
 
 			const int32 PrevGlobalTransformsNum = PrevGlobalTransforms.Num();
 			const int32 Max = FMath::Max(PrevGlobalTransformsNum, Entry.Num());
@@ -1124,7 +909,7 @@
 
 			for (int32 i = 0; i < Entry.Num(); ++i)
 			{
-				const FTransform RootTransform = bValidTrajectory ? Trajectory.GetSampleAtTime(Entry.AccumulatedSeconds).GetTransform() : AnimInstanceProxy.GetComponentTransform();
+				const FTransform RootTransform = bValidTrajectory ? PoseSearchHistory.GetTrajectory().GetSampleAtTime(Entry.AccumulatedSeconds).GetTransform() : AnimInstanceProxy.GetComponentTransform();
 				const FTransform GlobalTransforms = Entry.GetComponentSpaceTransform(i) * RootTransform;
 
 				if (i < PrevGlobalTransformsNum)
@@ -1457,3 +1242,256 @@
 }
 
 } // namespace UE::PoseSearch
+
+
+//////////////////////////////////////////////////////////////////////////
+// FPoseSearchHistory
+void FPoseSearchHistory::Initialize(int32 InNumPoses, float InSamplingInterval)
+{
+	check(InNumPoses >= 2);
+
+	MaxNumPoses = InNumPoses;
+	SamplingInterval = InSamplingInterval;
+	
+	Trajectory = FTransformTrajectory();
+	TrajectorySpeedMultiplier = 1.f;
+
+	LastUpdateSkeleton = nullptr;
+	BoneToTransformMap.Reset();
+	CollectedCurves.Reset();
+	BoneToTransformMapTypeHash = 0;
+	Entries.Reset();
+	PoseIndicesHistory.Reset();
+}
+
+bool FPoseSearchHistory::GetTransformAtTime(float Time, FTransform& OutBoneTransform, const USkeleton* BoneIndexSkeleton, FBoneIndexType BoneIndexType, FBoneIndexType ReferenceBoneIndexType, bool bExtrapolate) const
+{
+	using namespace UE::PoseSearch;
+
+	static_assert(RootBoneIndexType == 0 && ComponentSpaceIndexType == FBoneIndexType(-1) && WorldSpaceIndexType == FBoneIndexType(-2)); // some assumptions
+	
+	bool bSuccess = true;
+	if (BoneIndexType == ReferenceBoneIndexType)
+	{
+		OutBoneTransform = FTransform::Identity;
+		bSuccess = true;
+	}
+	else if (ReferenceBoneIndexType == WorldSpaceIndexType)
+	{
+		if (BoneIndexType == ComponentSpaceIndexType)
+		{
+			OutBoneTransform = Trajectory.GetSampleAtTime(Time, bExtrapolate).GetTransform();
+			bSuccess = true;
+		}
+		else
+		{
+			// getting BoneIndexType in ComponentSpaceIndexType and then multiplying it with the component to world transform from the trajectory
+			bSuccess = GetTransformAtTime(Time, OutBoneTransform, BoneIndexSkeleton, BoneIndexType, ComponentSpaceIndexType, bExtrapolate);
+			OutBoneTransform *= Trajectory.GetSampleAtTime(Time, bExtrapolate).GetTransform();
+		}
+	}
+	else if (BoneIndexType == WorldSpaceIndexType || BoneIndexType == ComponentSpaceIndexType)
+	{
+		bSuccess = GetTransformAtTime(Time, OutBoneTransform, BoneIndexSkeleton, ReferenceBoneIndexType, BoneIndexType, bExtrapolate);
+		OutBoneTransform = OutBoneTransform.Inverse();
+	}
+	else
+	{
+		check(BoneIndexType != ComponentSpaceIndexType && BoneIndexType != WorldSpaceIndexType && ReferenceBoneIndexType != WorldSpaceIndexType);
+
+		const int32 NumEntries = Entries.Num();
+		if (NumEntries > 0)
+		{
+			int32 NextIdx = 0;
+			int32 PrevIdx = 0;
+
+			if (NumEntries > 1)
+			{
+				const int32 LowerBoundIdx = LowerBound(Entries.begin(), Entries.end(), Time, [](const FPoseHistoryEntry& Entry, float Value) { return Value > Entry.AccumulatedSeconds; });
+				NextIdx = FMath::Clamp(LowerBoundIdx, 1, NumEntries - 1);
+				PrevIdx = NextIdx - 1;
+			}
+
+			const FPoseHistoryEntry& PrevEntry = Entries[PrevIdx];
+			const FPoseHistoryEntry& NextEntry = Entries[NextIdx];
+
+			bSuccess = LerpEntries(Time, bExtrapolate, PrevEntry, NextEntry, BoneIndexSkeleton, LastUpdateSkeleton.Get(), BoneToTransformMap, BoneIndexType, ReferenceBoneIndexType, OutBoneTransform);
+		}
+		else
+		{
+			OutBoneTransform = FTransform::Identity;
+			bSuccess = false;
+		}
+	}
+
+	// @todo: reenable this logging after implementing AnimNext PoseHistory initialization (currently spamming on actor spawning with MM active)
+//#if WITH_EDITOR
+//	// @todo: logging errors only if pose history was properly initialized for now (since AnimNext version doesn't implement initialization with ref pose yet),
+//	// and we ultimately want to catch errors of missing bone transforms only checking valid USkeleton(s) against valid PoseSearchSchema(s)
+//	if (!bSuccess && BoneIndexSkeleton)
+//	{
+//		UE_LOGF(LogPoseSearch, Error, "FPoseHistory::GetTransformAtTime - Couldn't find root bone transform at time %f for Skeleton %ls!", Time, *BoneIndexSkeleton->GetName());
+//	}
+//#endif // WITH_EDITOR
+
+	return bSuccess;
+}
+
+bool FPoseSearchHistory::GetCurveValueAtTime(float Time, const FName& CurveName, float& OutCurveValue, bool bExtrapolate /*= false*/) const
+{
+	using namespace UE::PoseSearch;
+
+	const int32 NumEntries = Entries.Num();
+	if (NumEntries > 0)
+	{
+		int32 NextIdx = 0;
+		int32 PrevIdx = 0;
+
+		if (NumEntries > 1)
+		{
+			const int32 LowerBoundIdx = LowerBound(Entries.begin(), Entries.end(), Time, [](const FPoseHistoryEntry& Entry, float Value) { return Value > Entry.AccumulatedSeconds; });
+			NextIdx = FMath::Clamp(LowerBoundIdx, 1, NumEntries - 1);
+			PrevIdx = NextIdx - 1;
+		}
+
+		const FPoseHistoryEntry& PrevEntry = Entries[PrevIdx];
+		const FPoseHistoryEntry& NextEntry = Entries[NextIdx];
+
+		return LerpEntries(Time, bExtrapolate, PrevEntry, NextEntry, CurveName, CollectedCurves, OutCurveValue);
+	}
+	
+	OutCurveValue = 0.0f;
+	return false;
+}
+
+void FPoseSearchHistory::SetTrajectory(const FTransformTrajectory& InTrajectory, float InTrajectorySpeedMultiplier)
+{
+	Trajectory = InTrajectory;
+
+	TrajectorySpeedMultiplier = InTrajectorySpeedMultiplier;
+
+	if (!FMath::IsNearlyEqual(TrajectorySpeedMultiplier, 1.f))
+	{
+		const float TrajectorySpeedMultiplierInv = FMath::IsNearlyZero(TrajectorySpeedMultiplier) ? 1.f : 1.f / TrajectorySpeedMultiplier;
+		for (FTransformTrajectorySample& Sample : Trajectory.Samples)
+		{
+			Sample.TimeInSeconds *= TrajectorySpeedMultiplierInv;
+		}
+	}
+}
+
+void FPoseSearchHistory::EvaluateComponentSpace(float DeltaTime, UE::PoseSearch::IComponentSpacePoseProvider& ComponentSpacePoseProvider, bool bStoreScales,
+	float RootBoneRecoveryTime, float RootBoneTranslationRecoveryRatio, float RootBoneRotationRecoveryRatio,
+	bool bNeedsReset, bool bCacheBones, const TArray<FBoneIndexType>& RequiredBones,
+	const FBlendedCurve& Curves, const TConstArrayView<FName>& InCollectedCurves
+#if WITH_EDITORONLY_DATA
+	, const FTransform& InDebugMeshTransform
+#endif //WITH_EDITORONLY_DATA
+	)
+{
+	using namespace UE::PoseSearch;
+
+	if (!ensure(MaxNumPoses >= 2))
+	{
+		UE_LOGF(LogPoseSearch, Error, "FPoseSearchHistory::EvaluateComponentSpace called before FPoseSearchHistory::Initialize! Please double check the implementation of any FAnimNode_Base calling this API, correctly propagating Initialize_AnyThread to their FPoseLink(s)! To mitigate issues FPoseHistory will be set to to default values");
+		Initialize(2, 0.f);
+	}
+
+	const USkeleton* Skeleton = ComponentSpacePoseProvider.GetSkeletonAsset();
+
+	if (bCacheBones)
+	{
+		const uint32 OldBoneToTransformMapTypeHash = BoneToTransformMapTypeHash;
+
+		BoneToTransformMap.Reset();
+		CollectedCurves = InCollectedCurves;
+		if (!RequiredBones.IsEmpty())
+		{
+			// making sure we always collect the root bone transform (by construction BoneToTransformMap[0] = 0)
+			const FComponentSpaceTransformIndex ComponentSpaceTransformRootBoneIndex = 0;
+			BoneToTransformMap.Add(RootBoneIndexType) = ComponentSpaceTransformRootBoneIndex;
+
+			for (int32 i = 0; i < RequiredBones.Num(); ++i)
+			{
+				// adding only unique RequiredBones to avoid oversizing Entries::ComponentSpaceTransforms
+				if (!BoneToTransformMap.Find(RequiredBones[i]))
+				{
+					const FComponentSpaceTransformIndex ComponentSpaceTransformIndex = BoneToTransformMap.Num();
+					BoneToTransformMap.Add(RequiredBones[i]) = ComponentSpaceTransformIndex;
+				}
+			}
+		}
+
+		BoneToTransformMapTypeHash = GetTypeHash(BoneToTransformMap);
+		bNeedsReset |= (OldBoneToTransformMapTypeHash != BoneToTransformMapTypeHash);
+	}
+
+	if (LastUpdateSkeleton != Skeleton)
+	{
+		bNeedsReset = true;
+		LastUpdateSkeleton = Skeleton;
+	}
+
+	if (bNeedsReset)
+	{
+		Entries.Reset();
+		Entries.Reserve(MaxNumPoses);
+	}
+
+	FPoseHistoryEntry FutureEntryTemp;
+	if (!Entries.IsEmpty() && Entries.Last().AccumulatedSeconds > 0.f)
+	{
+		// removing the "future" root bone Entry
+		FutureEntryTemp = MoveTemp(Entries.Last());
+		Entries.Pop();
+	}
+
+	// Age our elapsed times
+	for (FPoseHistoryEntry& Entry : Entries)
+	{
+		Entry.AccumulatedSeconds -= DeltaTime;
+	}
+
+	if (Entries.Num() != MaxNumPoses)
+	{
+		// Consume every pose until the queue is full
+		Entries.Emplace();
+	}
+	else if (SamplingInterval <= 0.f || Entries[Entries.Num() - 2].AccumulatedSeconds <= -SamplingInterval)
+	{
+		FPoseHistoryEntry EntryTemp = MoveTemp(Entries.First());
+		Entries.PopFront();
+		Entries.Emplace(MoveTemp(EntryTemp));
+	}
+
+	// Regardless of the retention policy, we always update the most recent Entry
+	FPoseHistoryEntry& MostRecentEntry = Entries.Last();
+	MostRecentEntry.Update(0.f, ComponentSpacePoseProvider, BoneToTransformMap, bStoreScales, Curves, CollectedCurves
+#if WITH_EDITORONLY_DATA
+		, InDebugMeshTransform
+#endif //WITH_EDITORONLY_DATA
+	);
+
+	if (RootBoneRecoveryTime > 0.f && !Trajectory.Samples.IsEmpty())
+	{
+		// adding the updated "future" root bone Entry
+		const FTransform& RefRootBone = Skeleton->GetReferenceSkeleton().GetRefBonePose()[RootBoneIndexType];
+		const FQuat RootBoneRotationAtRecoveryTime = FMath::Lerp(FQuat(MostRecentEntry.ComponentSpaceRotations[RootBoneIndexType]), RefRootBone.GetRotation(), RootBoneRotationRecoveryRatio);
+
+		FVector RootBoneDeltaTranslationAtRecoveryTime = FVector::ZeroVector;
+		if (RootBoneTranslationRecoveryRatio > 0.f)
+		{
+			const FTransform WorldRootAtCurrentTime = Trajectory.GetSampleAtTime(0.f).GetTransform();
+			const FTransform WorldRootBoneAtCurrentTime = MostRecentEntry.GetComponentSpaceTransform(RootBoneIndexType) * WorldRootAtCurrentTime;
+			const FVector WorldRootBoneDeltaTranslationAtCurrentTime = (WorldRootBoneAtCurrentTime.GetTranslation() - WorldRootAtCurrentTime.GetTranslation()) * RootBoneTranslationRecoveryRatio;
+			const FTransform WorldRootAtRecoveryTime = Trajectory.GetSampleAtTime(RootBoneRecoveryTime).GetTransform();
+			RootBoneDeltaTranslationAtRecoveryTime = WorldRootAtRecoveryTime.InverseTransformVector(WorldRootBoneDeltaTranslationAtCurrentTime);
+		}
+
+		const FTransform RootBoneTransformAtRecoveryTime(RootBoneRotationAtRecoveryTime, RootBoneDeltaTranslationAtRecoveryTime, RefRootBone.GetScale3D());
+		FutureEntryTemp.SetNum(1, bStoreScales);
+		FutureEntryTemp.SetComponentSpaceTransform(RootBoneIndexType, RootBoneTransformAtRecoveryTime);
+		FutureEntryTemp.AccumulatedSeconds = RootBoneRecoveryTime;
+		Entries.Emplace(MoveTemp(FutureEntryTemp));
+	}
+}
diff --git a/Engine/Plugins/Animation/PoseSearch/Source/Runtime/Public/PoseSearch/Chooser/ChooserParameterPoseHistoryBase.h b/Engine/Plugins/Animation/PoseSearch/Source/Runtime/Public/PoseSearch/Chooser/ChooserParameterPoseHistoryBase.h
--- a/Engine/Plugins/Animation/PoseSearch/Source/Runtime/Public/PoseSearch/Chooser/ChooserParameterPoseHistoryBase.h
+++ b/Engine/Plugins/Animation/PoseSearch/Source/Runtime/Public/PoseSearch/Chooser/ChooserParameterPoseHistoryBase.h
@@ -13,5 +13,6 @@ struct FChooserParameterPoseHistoryBase : public FChooserParameterBase
 {
 	GENERATED_BODY()
 	virtual bool GetValue(FChooserEvaluationContext& Context, FPoseHistoryReference& OutResult) const { return false; }
+	virtual bool GetValue(FChooserEvaluationContext& Context, FPoseSearchHistory& OutResult) const { return false; }
 	virtual bool IsBound() const { return false; }
 };
\ No newline at end of file
diff --git a/Engine/Plugins/Animation/PoseSearch/Source/Runtime/Public/PoseSearch/Chooser/PoseSearchChooserColumn.h b/Engine/Plugins/Animation/PoseSearch/Source/Runtime/Public/PoseSearch/Chooser/PoseSearchChooserColumn.h
--- a/Engine/Plugins/Animation/PoseSearch/Source/Runtime/Public/PoseSearch/Chooser/PoseSearchChooserColumn.h
+++ b/Engine/Plugins/Animation/PoseSearch/Source/Runtime/Public/PoseSearch/Chooser/PoseSearchChooserColumn.h
@@ -76,6 +76,11 @@ struct FPoseHistoryContextProperty :  public FChooserParameterPoseHistoryBase
 		return Binding.GetStructValue(Context, OutResult);
 	}
 
+	virtual bool GetValue(FChooserEvaluationContext& Context, FPoseSearchHistory& OutResult) const override
+	{
+		return Binding.GetStructValue(Context, OutResult);
+	}
+
 	virtual bool IsBound() const override
 	{
 		return Binding.IsBoundToRoot || !Binding.PropertyBindingChain.IsEmpty();
diff --git a/Engine/Plugins/Animation/PoseSearch/Source/Runtime/Public/PoseSearch/PoseSearchHistory.h b/Engine/Plugins/Animation/PoseSearch/Source/Runtime/Public/PoseSearch/PoseSearchHistory.h
--- a/Engine/Plugins/Animation/PoseSearch/Source/Runtime/Public/PoseSearch/PoseSearchHistory.h
+++ b/Engine/Plugins/Animation/PoseSearch/Source/Runtime/Public/PoseSearch/PoseSearchHistory.h
@@ -113,6 +113,86 @@ struct FPoseHistoryEntry
 };
 FArchive& operator<<(FArchive& Ar, FPoseHistoryEntry& Entry);
 
+} // namespace UE::PoseSearch
+
+// Experimental, this feature might be removed without warning, not for production use
+USTRUCT(Experimental, BlueprintType)
+struct FPoseSearchHistory
+{
+	GENERATED_BODY()
+
+	UE_API void Initialize(int32 InNumPoses, float InSamplingInterval);
+	UE_API void EvaluateComponentSpace(float DeltaTime, UE::PoseSearch::IComponentSpacePoseProvider& ComponentSpacePoseProvider, bool bStoreScales,
+		float RootBoneRecoveryTime, float RootBoneTranslationRecoveryRatio, float RootBoneRotationRecoveryRatio,
+		bool bNeedsReset, bool bCacheBones, const TArray<FBoneIndexType>& RequiredBones, 
+		const FBlendedCurve& Curves = FBlendedCurve(), const TConstArrayView<FName>& InCollectedCurves = TConstArrayView<FName>()
+#if WITH_EDITORONLY_DATA
+		, const FTransform& InDebugMeshTransform = FTransform::Identity
+#endif //WITH_EDITORONLY_DATA
+	);
+
+	UE_API bool GetTransformAtTime(float Time, FTransform& OutBoneTransform, const USkeleton* BoneIndexSkeleton = nullptr, FBoneIndexType BoneIndexType = UE::PoseSearch::RootBoneIndexType, FBoneIndexType ReferenceBoneIndexType = UE::PoseSearch::ComponentSpaceIndexType, bool bExtrapolate = false) const;
+	UE_API bool GetCurveValueAtTime(float Time, const FName& CurveName, float& outcurvevalue, bool bExtrapolate = false) const;
+	UE_API void SetTrajectory(const FTransformTrajectory& InTrajectory, float InTrajectorySpeedMultiplier = 1.f);
+	UE_API void GenerateTrajectory(const UObject* AnimContext, float DeltaTime);
+	UE_API float GetTrajectorySpeedMultiplier() const { return TrajectorySpeedMultiplier; }
+	UE_API bool IsEmpty() const { return Entries.IsEmpty(); }
+	const UE::PoseSearch::FBoneToTransformMap& GetBoneToTransformMap() const { return BoneToTransformMap; }
+	const TConstArrayView<FName> GetCollectedCurves() const { return MakeConstArrayView(CollectedCurves); }
+
+	int32 GetNumEntries() const { return Entries.Num(); }
+	const UE::PoseSearch::FPoseHistoryEntry& GetEntry(int32 EntryIndex) const { return Entries[EntryIndex]; }
+	const UE::PoseSearch::FPoseIndicesHistory* GetPoseIndicesHistory() const { return &PoseIndicesHistory; }
+
+	int32 GetMaxNumPoses() const { return MaxNumPoses; }
+	float GetSamplingInterval() const { return SamplingInterval; }
+
+	const FTransformTrajectory& GetTrajectory() const { return Trajectory; }
+	FTransformTrajectory& EditTrajectory() { return Trajectory; }
+
+private:
+	UPROPERTY(EditAnywhere, Category=Settings, meta=(AllowPrivateAccess=true))
+	FTransformTrajectory Trajectory;
+
+	// max number of stored poses in this pose history
+	UPROPERTY(EditAnywhere, Category=Settings, meta=(AllowPrivateAccess=true))
+	int32 MaxNumPoses = 2;
+
+	// frequency of sampling of poses in this pose history. negative numbers will store a new pose on every update
+	UPROPERTY(EditAnywhere, Category=Settings, meta=(AllowPrivateAccess=true))
+	float SamplingInterval = -1.f;
+	
+	UPROPERTY(EditAnywhere, Category=Settings, meta=(AllowPrivateAccess=true))
+	float TrajectorySpeedMultiplier = 1.f;
+
+	// skeleton from the last Update, to keep tracking skeleton changes, and support compatible skeletons
+	TWeakObjectPtr<const USkeleton> LastUpdateSkeleton;
+
+	// Map from skeleton bone indexes (encoded as FBoneIndexType) to internal pose history transform index (FComponentSpaceTransformIndex).
+	// If Empty all the bones get collected
+	
+	// @todo: make it a UPROPERTY
+	// typedef TMap<FBoneIndexType, FComponentSpaceTransformIndex> FBoneToTransformMap;
+	UE::PoseSearch::FBoneToTransformMap BoneToTransformMap;
+		
+	// list of curves that we want to collect into our history.
+	TArray<FName> CollectedCurves;
+
+	// GetTypeHash for BoneToTransformMap
+	uint32 BoneToTransformMapTypeHash = 0;
+
+	// @todo: make it a UPROPERTY
+	// ring buffer of collected bones
+	TRingBuffer<UE::PoseSearch::FPoseHistoryEntry> Entries;
+
+	// @todo: make it a UPROPERTY
+	// data structure containing the previously played pose indexes with their respective datababases
+	UE::PoseSearch::FPoseIndicesHistory PoseIndicesHistory;
+};
+
+namespace UE::PoseSearch
+{
+
 struct IPoseHistory : public TSharedFromThis<IPoseHistory, ESPMode::ThreadSafe>
 {
 	virtual ~IPoseHistory() = default;
@@ -234,37 +314,14 @@ struct FPoseHistory : public IPoseHistory
 	UE_API int32 GetMaxNumPoses() const;
 	UE_API float GetSamplingInterval() const;
 
-private:
-	// caching MaxNumPoses, since FData::Entries.Max() is a padded number
-	int32 MaxNumPoses = 0;
+	// Experimental, this feature might be removed without warning, not for production use
+	UE_API const FPoseSearchHistory& GetPoseSearchHistory() const;
+	// Experimental, this feature might be removed without warning, not for production use
+	UE_API FPoseSearchHistory& EditPoseSearchHistory();
 
-	float SamplingInterval = 0.f;
-	
-	FTransformTrajectory Trajectory;
-	
+private:
+	FPoseSearchHistory PoseSearchHistory;
 	FPoseSearchTrajectoryData::FState TrajectoryDataState;
-	// @todo: deprecate this member and expose it via blue print logic or as global query scaling multiplier
-	float TrajectorySpeedMultiplier = 1.f;
-
-	// skeleton from the last Update, to keep tracking skeleton changes, and support compatible skeletons
-	TWeakObjectPtr<const USkeleton> LastUpdateSkeleton;
-
-	// Map from skeleton bone indexes (encoded as FBoneIndexType) to internal pose history transform index (FComponentSpaceTransformIndex).
-	// If Empty all the bones get collected
-	FBoneToTransformMap BoneToTransformMap;
-		
-	// list of curves that we want to collect into our history.
-	TArray<FName> CollectedCurves;
-
-	// GetTypeHash for BoneToTransformMap
-	uint32 BoneToTransformMapTypeHash = 0;
-
-	// ring buffer of collected bones
-	TRingBuffer<FPoseHistoryEntry> Entries;
-
-	// datra structure containing the previously played pose indexes with their respective datababases
-	FPoseIndicesHistory PoseIndicesHistory;
-
 	UE_MT_DECLARE_TS_RW_ACCESS_DETECTOR(PoseDataThreadSafeCounter);
 };
 
@@ -338,4 +395,6 @@ struct FPoseHistoryReference
 	TSharedPtr<UE::PoseSearch::IPoseHistory> PoseHistory;
 };
 
+
+
 #undef UE_API
diff --git a/Engine/Plugins/Experimental/UAF/UAFPoseSearch/Source/UAFPoseSearch/Private/HistoryCollectorTrait.cpp b/Engine/Plugins/Experimental/UAF/UAFPoseSearch/Source/UAFPoseSearch/Private/HistoryCollectorTrait.cpp
--- a/Engine/Plugins/Experimental/UAF/UAFPoseSearch/Source/UAFPoseSearch/Private/HistoryCollectorTrait.cpp
+++ b/Engine/Plugins/Experimental/UAF/UAFPoseSearch/Source/UAFPoseSearch/Private/HistoryCollectorTrait.cpp
@@ -218,7 +218,7 @@ namespace UE::UAF
 		}
 		else
 		{
-			UE_LOG(LogPoseSearch, Error, TEXT("FHistoryCollectorTrait::PreUpdate, PostEvaluate has not being called last frame! Some trait in the TraitStack didn't propagate correctly the PostEvaluate!"));
+			UE_LOGF(LogPoseSearch, Error, "FHistoryCollectorTrait::PreUpdate, PostEvaluate has not being called last frame! Some trait in the TraitStack didn't propagate correctly the PostEvaluate!");
 		}
 #endif // WITH_EDITOR
 
@@ -284,7 +284,7 @@ namespace UE::UAF
 #if WITH_EDITOR
 		if (InstanceData->bIsPostEvaluateBeingCalled)
 		{
-			UE_LOG(LogPoseSearch, Error, TEXT("FHistoryCollectorTrait::PostEvaluate, PostEvaluate called without calling PreUpdate on this frame! Some trait in the TraitStack doesn't propagate correctly the PreUpdate!"));
+			UE_LOGF(LogPoseSearch, Error, "FHistoryCollectorTrait::PostEvaluate, PostEvaluate called without calling PreUpdate on this frame! Some trait in the TraitStack doesn't propagate correctly the PreUpdate!");
 		}
 		else
 		{
@@ -298,7 +298,13 @@ namespace UE::UAF
 			FPoseHistoryReference Reference;
 			Reference.PoseHistory = InstanceData->PoseHistoryPtr;
 			PRAGMA_DISABLE_DEPRECATION_WARNINGS
-			SetVariableValue(&Binding.GetTraitPtr().GetNodeInstance()->GetOwner(), FAnimNextVariableReference(PoseHistoryReferenceVariableName), Reference);
+			if (SetVariableValue(&Binding.GetTraitPtr().GetNodeInstance()->GetOwner(), FAnimNextVariableReference(PoseHistoryReferenceVariableName), Reference) != EPropertyBagResult::Success)
+			{
+				if (SetVariableValue(&Binding.GetTraitPtr().GetNodeInstance()->GetOwner(), FAnimNextVariableReference(PoseHistoryReferenceVariableName), InstanceData->PoseHistoryPtr->GetPoseSearchHistory()) != EPropertyBagResult::Success)
+				{
+					UE_LOGF(LogPoseSearch, Error, "FHistoryCollectorTrait::PostEvaluate, failed to set pose history to variable '%ls' (variable must be of type FPoseHistoryReference or FPoseSearchHistory)", *PoseHistoryReferenceVariableName.ToString());
+				}
+			}
 			PRAGMA_ENABLE_DEPRECATION_WARNINGS
 		}
 		
diff --git a/Engine/Plugins/Experimental/UAF/UAFPoseSearch/Source/UAFPoseSearch/Private/Module.cpp b/Engine/Plugins/Experimental/UAF/UAFPoseSearch/Source/UAFPoseSearch/Private/Module.cpp
--- a/Engine/Plugins/Experimental/UAF/UAFPoseSearch/Source/UAFPoseSearch/Private/Module.cpp
+++ b/Engine/Plugins/Experimental/UAF/UAFPoseSearch/Source/UAFPoseSearch/Private/Module.cpp
@@ -41,6 +41,7 @@ public:
 			FPoseSearchInteractionAssetItem::StaticStruct(),
 			FPoseSearchInteractionAvailability::StaticStruct(),
 			FPoseHistoryReference::StaticStruct(),
+			FPoseSearchHistory::StaticStruct(),
 			FPoseSearchEvent::StaticStruct(),
 		};
 
diff --git a/Engine/Plugins/Experimental/UAF/UAFPoseSearch/Source/UAFPoseSearch/Private/PoseHistoryChooserParameter.cpp b/Engine/Plugins/Experimental/UAF/UAFPoseSearch/Source/UAFPoseSearch/Private/PoseHistoryChooserParameter.cpp
--- a/Engine/Plugins/Experimental/UAF/UAFPoseSearch/Source/UAFPoseSearch/Private/PoseHistoryChooserParameter.cpp
+++ b/Engine/Plugins/Experimental/UAF/UAFPoseSearch/Source/UAFPoseSearch/Private/PoseHistoryChooserParameter.cpp
@@ -51,6 +51,15 @@ bool FPoseHistoryAnimProperty::GetValue(FChooserEvaluationContext& Context, FPos
 	return false;
 }
 
+bool FPoseHistoryAnimProperty::GetValue(FChooserEvaluationContext& Context, FPoseSearchHistory& OutResult) const
+{
+	if(FUAFAssetInstance* Instance = GetFirstUAFAssetInstance(Context))
+	{
+		return GetVariableValue(Instance, Variable, OutResult) == EPropertyBagResult::Success;
+	}
+	return false;
+}
+
 #if WITH_EDITORONLY_DATA
 bool FPoseHistoryAnimProperty::Serialize(FArchive& Ar)
 {
diff --git a/Engine/Plugins/Experimental/UAF/UAFPoseSearch/Source/UAFPoseSearch/Private/PoseHistoryChooserParameter.h b/Engine/Plugins/Experimental/UAF/UAFPoseSearch/Source/UAFPoseSearch/Private/PoseHistoryChooserParameter.h
--- a/Engine/Plugins/Experimental/UAF/UAFPoseSearch/Source/UAFPoseSearch/Private/PoseHistoryChooserParameter.h
+++ b/Engine/Plugins/Experimental/UAF/UAFPoseSearch/Source/UAFPoseSearch/Private/PoseHistoryChooserParameter.h
@@ -24,6 +24,7 @@ public:
 	FAnimNextVariableReference Variable;
 
 	bool GetValue(FChooserEvaluationContext& Context, FPoseHistoryReference& OutResult) const override;
+	bool GetValue(FChooserEvaluationContext& Context, FPoseSearchHistory& OutResult) const override;
 	virtual bool IsBound() const { return !Variable.IsNone(); }
 
 #if WITH_EDITORONLY_DATA
