PoseSearch - UAnimNotifyState_PoseSearchSamplingRange

#Synced-CL 52486799

diff --git a/Engine/Plugins/Animation/PoseSearch/Source/Runtime/Private/PoseSearchDerivedData.cpp b/Engine/Plugins/Animation/PoseSearch/Source/Runtime/Private/PoseSearchDerivedData.cpp
--- a/Engine/Plugins/Animation/PoseSearch/Source/Runtime/Private/PoseSearchDerivedData.cpp
+++ b/Engine/Plugins/Animation/PoseSearch/Source/Runtime/Private/PoseSearchDerivedData.cpp
@@ -406,7 +406,7 @@ public:
 
 typedef TArray<FFloatRange, TInlineAllocator<32>> FValidRanges;
 
-static void FindValidRanges(const FPoseSearchDatabaseAnimationAssetBase* DatabaseAsset, const FVector& BlendParameters, const FFloatInterval& ExcludeFromDatabaseParameters, FValidRanges& ValidRanges)
+static bool FindValidRanges(const FPoseSearchDatabaseAnimationAssetBase* DatabaseAsset, const FVector& BlendParameters, const FFloatInterval& ExcludeFromDatabaseParameters, FValidRanges& ValidRanges)
 {
 	check(DatabaseAsset);
 
@@ -425,6 +425,9 @@ static void FindValidRanges(const FPoseSearchDatabaseAnimationAssetBase* Databas
 	ValidRanges.Reset();
 	ValidRanges.Add(EffectiveSamplingRange);
 
+	bool bAnyExclusionNotifyStateFound = false;
+	bool bAnySamplingRangeNotifyStateFound = false;
+
 	for (int32 RoleIndex = 0; RoleIndex < DatabaseAsset->GetNumRoles(); ++RoleIndex)
 	{
 		const FRole Role = DatabaseAsset->GetRole(RoleIndex);
@@ -432,14 +435,22 @@ static void FindValidRanges(const FPoseSearchDatabaseAnimationAssetBase* Databas
 		{
 			FAnimNotifyContext NotifyContext;
 			SequenceBase->GetAnimNotifies(0.0f, PlayLength, NotifyContext);
-
+			
 			for (const FAnimNotifyEventReference& EventReference : NotifyContext.ActiveNotifies)
 			{
 				if (const FAnimNotifyEvent* NotifyEvent = EventReference.GetNotify())
 				{
 					if (const UAnimNotifyState_PoseSearchExcludeFromDatabase* ExclusionNotifyState = Cast<const UAnimNotifyState_PoseSearchExcludeFromDatabase>(NotifyEvent->NotifyStateClass))
 					{
-						FFloatRange ExclusionRange = FFloatRange::Inclusive(NotifyEvent->GetTime(), NotifyEvent->GetTime() + NotifyEvent->GetDuration());
+						if (bAnySamplingRangeNotifyStateFound)
+						{
+							UE_LOGF(LogPoseSearch, Error, "FindValidRanges conflicting UAnimNotifyState_PoseSearchExcludeFromDatabase and UAnimNotifyState_PoseSearchSamplingRange found in '%ls'.", *DatabaseAsset->GetAnimationAsset()->GetName());
+							ValidRanges.Reset();
+							return false;
+						}
+
+						bAnyExclusionNotifyStateFound = true;
+						const FFloatRange ExclusionRange = FFloatRange::Inclusive(NotifyEvent->GetTime(), NotifyEvent->GetTime() + NotifyEvent->GetDuration());
 
 						// Split every valid range based on the exclusion range just found. Because this might increase the 
 						// number of ranges in ValidRanges, the algorithm iterates from end to start.
@@ -452,10 +463,47 @@ static void FindValidRanges(const FPoseSearchDatabaseAnimationAssetBase* Databas
 							ValidRanges.Append(Diff);
 						}
 					}
+					else if (const UAnimNotifyState_PoseSearchSamplingRange* SamplingRangeNotifyState = Cast<const UAnimNotifyState_PoseSearchSamplingRange>(NotifyEvent->NotifyStateClass))
+					{
+						if (bAnyExclusionNotifyStateFound)
+						{
+							UE_LOGF(LogPoseSearch, Error, "FindValidRanges conflicting UAnimNotifyState_PoseSearchExcludeFromDatabase and UAnimNotifyState_PoseSearchSamplingRange found in '%ls'.", *DatabaseAsset->GetAnimationAsset()->GetName());
+							ValidRanges.Reset();
+							return false;
+						}
+
+						FFloatRange InclusionRange = FFloatRange::Inclusive(NotifyEvent->GetTime(), NotifyEvent->GetTime() + NotifyEvent->GetDuration());
+						if (!bAnySamplingRangeNotifyStateFound)
+						{
+							ValidRanges.Reset();
+							ValidRanges.Add(InclusionRange);
+							bAnySamplingRangeNotifyStateFound = true;
+						}
+						else
+						{
+							for (int32 RangeIdx = ValidRanges.Num() - 1; RangeIdx >= 0; --RangeIdx)
+							{
+								const FFloatRange EvaluatedRange = ValidRanges[RangeIdx];
+								if (InclusionRange.Contiguous(EvaluatedRange))
+								{
+									ValidRanges.RemoveAt(RangeIdx);
+									const TArray<FFloatRange> Union = FFloatRange::Union(InclusionRange, EvaluatedRange);
+									if (ensure(Union.Num() == 1))
+									{
+										InclusionRange = Union[0];
+									}
+								}
+							}
+
+							ValidRanges.Add(InclusionRange);
+						}
+					}
 				}
 			}
 		}
 	}
+	
+	return true;
 }
 
 // returns false in case of errors
@@ -546,50 +594,56 @@ static bool InitSearchIndexAssets(FSearchIndexBase& SearchIndex, const UPoseSear
 			}
 
 			DatabaseAsset->IterateOverSamplingParameter(
-				[DatabaseAsset, ExcludeFromDatabaseParameters, &ValidRanges, Schema, bAddUnmirrored, bAddMirrored, bIsLooping, bIsBlendSpace, bDisableReselection, AnimationAssetIndex, &TotalPoses, &SearchIndex]
+				[DatabaseAsset, ExcludeFromDatabaseParameters, &ValidRanges, Schema, bAddUnmirrored, bAddMirrored, 
+				bIsLooping, bIsBlendSpace, bDisableReselection, AnimationAssetIndex, &TotalPoses, &SearchIndex, &bAnyErrors]
 				(const FVector& BlendParameters)
 				{
-					FindValidRanges(DatabaseAsset, BlendParameters, ExcludeFromDatabaseParameters, ValidRanges);
-
-					for (const FFloatRange& Range : ValidRanges)
+					if (FindValidRanges(DatabaseAsset, BlendParameters, ExcludeFromDatabaseParameters, ValidRanges))
 					{
-						for (int32 PermutationIdx = 0; PermutationIdx < Schema->NumberOfPermutations; ++PermutationIdx)
+						for (const FFloatRange& Range : ValidRanges)
 						{
-							const FFloatInterval RangeInterval(Range.GetLowerBoundValue(), Range.GetUpperBoundValue());
-
-							float ToRealTimeFactor = 1.f;
-							if (bIsBlendSpace)
+							for (int32 PermutationIdx = 0; PermutationIdx < Schema->NumberOfPermutations; ++PermutationIdx)
 							{
-								const float PlayLength = DatabaseAsset->GetPlayLength(BlendParameters);
-								if (PlayLength > UE_KINDA_SMALL_NUMBER)
+								const FFloatInterval RangeInterval(Range.GetLowerBoundValue(), Range.GetUpperBoundValue());
+
+								float ToRealTimeFactor = 1.f;
+								if (bIsBlendSpace)
 								{
-									ToRealTimeFactor = PlayLength;
+									const float PlayLength = DatabaseAsset->GetPlayLength(BlendParameters);
+									if (PlayLength > UE_KINDA_SMALL_NUMBER)
+									{
+										ToRealTimeFactor = PlayLength;
+									}
 								}
-							}
 
-							if (bAddUnmirrored)
-							{
-								const FSearchIndexAsset PoseSearchIndexAsset(AnimationAssetIndex, TotalPoses, false, bIsLooping,
-									bDisableReselection, RangeInterval, Schema->SampleRate, PermutationIdx, BlendParameters, ToRealTimeFactor);
-								if (PoseSearchIndexAsset.GetNumPoses() > 0)
+								if (bAddUnmirrored)
 								{
-									SearchIndex.Assets.Add(PoseSearchIndexAsset);
-									TotalPoses += PoseSearchIndexAsset.GetNumPoses();
+									const FSearchIndexAsset PoseSearchIndexAsset(AnimationAssetIndex, TotalPoses, false, bIsLooping,
+										bDisableReselection, RangeInterval, Schema->SampleRate, PermutationIdx, BlendParameters, ToRealTimeFactor);
+									if (PoseSearchIndexAsset.GetNumPoses() > 0)
+									{
+										SearchIndex.Assets.Add(PoseSearchIndexAsset);
+										TotalPoses += PoseSearchIndexAsset.GetNumPoses();
+									}
 								}
-							}
 
-							if (bAddMirrored)
-							{
-								const FSearchIndexAsset PoseSearchIndexAsset(AnimationAssetIndex, TotalPoses, true, bIsLooping,
-									bDisableReselection, RangeInterval, Schema->SampleRate, PermutationIdx, BlendParameters, ToRealTimeFactor);
-								if (PoseSearchIndexAsset.GetNumPoses() > 0)
+								if (bAddMirrored)
 								{
-									SearchIndex.Assets.Add(PoseSearchIndexAsset);
-									TotalPoses += PoseSearchIndexAsset.GetNumPoses();
+									const FSearchIndexAsset PoseSearchIndexAsset(AnimationAssetIndex, TotalPoses, true, bIsLooping,
+										bDisableReselection, RangeInterval, Schema->SampleRate, PermutationIdx, BlendParameters, ToRealTimeFactor);
+									if (PoseSearchIndexAsset.GetNumPoses() > 0)
+									{
+										SearchIndex.Assets.Add(PoseSearchIndexAsset);
+										TotalPoses += PoseSearchIndexAsset.GetNumPoses();
+									}
 								}
 							}
 						}
 					}
+					else
+					{
+						bAnyErrors = true;
+					}
 				});
 		}
 	}
diff --git a/Engine/Plugins/Animation/PoseSearch/Source/Runtime/Public/PoseSearch/PoseSearchAnimNotifies.h b/Engine/Plugins/Animation/PoseSearch/Source/Runtime/Public/PoseSearch/PoseSearchAnimNotifies.h
--- a/Engine/Plugins/Animation/PoseSearch/Source/Runtime/Public/PoseSearch/PoseSearchAnimNotifies.h
+++ b/Engine/Plugins/Animation/PoseSearch/Source/Runtime/Public/PoseSearch/PoseSearchAnimNotifies.h
@@ -20,13 +20,21 @@ class UAnimNotifyState_PoseSearchBase : public UAnimNotifyState
 };
 
 // Use this notify state to remove animation segments from the database completely, they will never play or return from
-// a search result
+// a search result. PSD indexing will report errors in case the sampled animation asset contains conflicting UAnimNotifyState_PoseSearchSamplingRange(s)
 UCLASS(MinimalAPI, Blueprintable, meta = (DisplayName = "Pose Search: Exclude From Database"))
 class UAnimNotifyState_PoseSearchExcludeFromDatabase : public UAnimNotifyState_PoseSearchBase
 {
 	GENERATED_BODY()
 };
 
+// Use this notify state to include animation segments into the database, whatever excluded will never play or return from
+// a search result. PSD indexing will report errors in case the sampled animation asset contains conflicting UAnimNotifyState_PoseSearchExcludeFromDatabase(s)
+UCLASS(MinimalAPI, Experimental, Blueprintable, meta = (DisplayName = "Pose Search: Sampling Range"))
+class UAnimNotifyState_PoseSearchSamplingRange : public UAnimNotifyState_PoseSearchBase
+{
+	GENERATED_BODY()
+};
+
 // A pose search search will not return results that overlap this notify, but the animation segment can still play
 // if a previous search result advances into it.
 UCLASS(MinimalAPI, Blueprintable, meta = (DisplayName = "Pose Search: Block Transition In"))
