PoseSearch - fix for comment typo

#Synced-CL 52486799
#Synced-CL 52486799

diff --git a/Engine/Plugins/Animation/PoseSearch/Source/Editor/Private/PoseSearchDatabaseAssetTree.cpp b/Engine/Plugins/Animation/PoseSearch/Source/Editor/Private/PoseSearchDatabaseAssetTree.cpp
--- a/Engine/Plugins/Animation/PoseSearch/Source/Editor/Private/PoseSearchDatabaseAssetTree.cpp
+++ b/Engine/Plugins/Animation/PoseSearch/Source/Editor/Private/PoseSearchDatabaseAssetTree.cpp
@@ -559,7 +559,7 @@ namespace UE::PoseSearch
 					const bool IsEnabled = ViewModel->IsEnabled(SelectedNodes[0]->SourceAssetIdx);
 					const FName IconToUse = !IsEnabled ? "Icons.Visible" : "Icons.Hidden";
 					const FText LabelToUse = !IsEnabled ? LOCTEXT("EnableAsset", "Enable asset") : LOCTEXT("DisableAsset", "Disable asset");
-					const FText TooltipToUse = !IsEnabled ? LOCTEXT("EnableAssetTooltip", "Inlcude asset in query when pose matching / motion matching.") : LOCTEXT("DisableAssetTooltip", "Exclude asset in query when pose matching / motion matching.");
+					const FText TooltipToUse = !IsEnabled ? LOCTEXT("EnableAssetTooltip", "Include asset in query when pose matching / motion matching.") : LOCTEXT("DisableAssetTooltip", "Exclude asset in query when pose matching / motion matching.");
 					
 					MenuBuilder.AddMenuEntry(
 						LabelToUse,
